Move Optimus plugin out of experimental and rename to DeformerGraph.

diff --git a/Engine/Plugins/Animation/DeformerGraph/Config/BaseDeformerGraph.ini b/Engine/Plugins/Animation/DeformerGraph/Config/BaseDeformerGraph.ini
new file mode 100644
--- /dev/null
+++ b/Engine/Plugins/Animation/DeformerGraph/Config/BaseDeformerGraph.ini
@@ -0,0 +1,47 @@
+[CoreRedirects]
+; Uncomment these if you have old Deformer Graph assets that won't load properly.
+;+StructRedirects=(OldName="OptimusType_ShaderText",NewName="/Script/OptimusCore.OptimusShaderText")
+
+;+ClassRedirects=(OldName="/Script/OptimusDeveloper.",NewName="/Script/OptimusCore.", MatchSubstring=true)
+;+StructRedirects=(OldName="/Script/OptimusDeveloper.",NewName="/Script/OptimusCore.", MatchSubstring=true)
+;+EnumRedirects=(OldName="/Script/OptimusDeveloper.",NewName="/Script/OptimusCore.", MatchSubstring=true)
+
++ClassRedirects=(OldName="/Script/OptimusCore.ClothDataInterface",NewName="/Script/OptimusCore.OptimusClothDataInterface")
++ClassRedirects=(OldName="/Script/OptimusCore.ClothDataProvider",NewName="/Script/OptimusCore.OptimusClothDataProvider")
+
++ClassRedirects=(OldName="/Script/OptimusCore.DataInterface",NewName="/Script/OptimusCore.OptimusGraphDataInterface")
++ClassRedirects=(OldName="/Script/OptimusCore.GraphDataProvider",NewName="/Script/OptimusCore.OptimusGraphDataProvider")
++StructRedirects=(OldName="/Script/OptimusCore.GraphVariableDescription",NewName="/Script/OptimusCore.OptimusGraphVariableDescription")
+
++ClassRedirects=(OldName="/Script/OptimusCore.MorphTargetDataInterface",NewName="/Script/OptimusCore.OptimusMorphTargetDataInterface")
++ClassRedirects=(OldName="/Script/OptimusCore.MorphTargetDataProvider",NewName="/Script/OptimusCore.OptimusMorphTargetDataProvider")
+
++ClassRedirects=(OldName="/Script/OptimusCore.TransientBufferDataInterface",NewName="/Script/OptimusCore.OptimusTransientBufferDataInterface")
++ClassRedirects=(OldName="/Script/OptimusCore.PersistentBufferDataInterface",NewName="/Script/OptimusCore.OptimusPersistentBufferDataInterface")
++ClassRedirects=(OldName="/Script/OptimusCore.TransientBufferDataProvider",NewName="/Script/OptimusCore.OptimusTransientBufferDataProvider")
++ClassRedirects=(OldName="/Script/OptimusCore.PersistentBufferDataProvider",NewName="/Script/OptimusCore.OptimusPersistentBufferDataProvider")
+
++ClassRedirects=(OldName="/Script/OptimusCore.SceneDataInterface",NewName="/Script/OptimusCore.OptimusSceneDataInterface")
++ClassRedirects=(OldName="/Script/OptimusCore.SceneDataProvider",NewName="/Script/OptimusCore.OptimusSceneDataProvider")
+
++ClassRedirects=(OldName="/Script/OptimusCore.SkeletalMeshReadDataInterface",NewName="/Script/OptimusCore.OptimusSkinnedMeshDataInterface")
++ClassRedirects=(OldName="/Script/OptimusCore.SkeletalMeshReadDataProvider",NewName="/Script/OptimusCore.OptimusSkinnedMeshDataProvider")
+
++ClassRedirects=(OldName="/Script/OptimusCore.SkeletonDataInterface",NewName="/Script/OptimusCore.OptimusSkeletonDataInterface")
++ClassRedirects=(OldName="/Script/OptimusCore.SkeletonDataProvider",NewName="/Script/OptimusCore.OptimusSkeletonDataProvider")
+
++ClassRedirects=(OldName="/Script/OptimusCore.SkinnedMeshDataInterface",NewName="/Script/OptimusCore.OptimusSkinnedMeshDataInterface")
++ClassRedirects=(OldName="/Script/OptimusCore.SkinnedMeshDataProvider",NewName="/Script/OptimusCore.OptimusSkinnedMeshDataProvider")
+
++ClassRedirects=(OldName="/Script/OptimusCore.SkinnedMeshExecDataInterface",NewName="/Script/OptimusCore.OptimusSkinnedMeshExecDataInterface")
++ClassRedirects=(OldName="/Script/OptimusCore.SkinnedMeshExecDataProvider",NewName="/Script/OptimusCore.OptimusSkinnedMeshExecDataProvider")
++EnumRedirects=(OldName="/Script/OptimusCore.ESkinnedMeshExecDomain",NewName="/Script/OptimusCore.EOptimusSkinnedMeshExecDomain")
+
++ClassRedirects=(OldName="/Script/OptimusCore.SkinnedMeshWriteDataInterface",NewName="/Script/OptimusCore.OptimusSkinnedMeshWriteDataInterface")
++ClassRedirects=(OldName="/Script/OptimusCore.SkinnedMeshWriteDataProvider",NewName="/Script/OptimusCore.OptimusSkinnedMeshWriteDataProvider")
+
++StructRedirects=(OldName="/Script/OptimusCore.OptimusMultiLevelDataDomain", NewName="/Script/OptimusCore.OptimusDataDomain")
+
+[/Script/OptimusSettings.OptimusSettings]
++DefaultDeformer=/DeformerGraph/Deformers/DG_LinearBlendSkin_Morph_Cloth.DG_LinearBlendSkin_Morph_Cloth
++DefaultRecomputeTangentDeformer=/DeformerGraph/Deformers/DG_LinearBlendSkin_Morph_Cloth_RecomputeNormals.DG_LinearBlendSkin_Morph_Cloth_RecomputeNormals
diff --git a/Engine/Plugins/Animation/DeformerGraph/Source/OptimusCore/Private/Nodes/OptimusNode_ComputeKernelBase.cpp b/Engine/Plugins/Animation/DeformerGraph/Source/OptimusCore/Private/Nodes/OptimusNode_ComputeKernelBase.cpp
new file mode 100644
--- /dev/null
+++ b/Engine/Plugins/Animation/DeformerGraph/Source/OptimusCore/Private/Nodes/OptimusNode_ComputeKernelBase.cpp
@@ -0,0 +1,674 @@
+﻿// Copyright Epic Games, Inc. All Rights Reserved.
+
+#include "OptimusNode_ComputeKernelBase.h"
+
+#include "IOptimusDataInterfaceProvider.h"
+#include "OptimusCoreModule.h"
+#include "OptimusNodeGraph.h"
+#include "OptimusNodePin.h"
+
+#include "DataInterfaces/OptimusDataInterfaceGraph.h"
+#include "DataInterfaces/OptimusDataInterfaceRawBuffer.h"
+#include "IOptimusValueProvider.h"
+#include "OptimusDataTypeRegistry.h"
+#include "OptimusHelpers.h"
+#include "OptimusKernelSource.h"
+
+
+#define LOCTEXT_NAMESPACE "OptimusKernelBase"
+
+
+static FString GetShaderParamPinValueString(
+	const UOptimusNodePin *InPin
+	)
+{
+	return InPin->GetDataType()->ShaderValueType->GetZeroValueAsString();
+	
+	// FIXME: Need property storage.
+	const FShaderValueType& ValueType = *InPin->GetDataType()->ShaderValueType;
+	TArrayView<UOptimusNodePin* const> SubPins = InPin->GetSubPins();
+	TArray<FString> ValueArray;
+	if (SubPins.IsEmpty())
+	{
+		ValueArray.Add(InPin->GetValueAsString());
+	}
+	
+	// FIXME: Support all types properly. Should probably be moved to a better place.
+	return FString::Printf(TEXT("%s(%s)"), *ValueType.ToString(), *InPin->GetValueAsString());
+}
+
+
+static void CopyValueType(FShaderValueTypeHandle InValueType,  FShaderParamTypeDefinition& OutParamDef)
+{
+	OutParamDef.ValueType = InValueType;
+	OutParamDef.ArrayElementCount = 0;
+	OutParamDef.ResetTypeDeclaration();
+}
+
+FOptimus_ComputeKernelResult UOptimusNode_ComputeKernelBase::CreateComputeKernel(
+	UObject *InKernelSourceOuter,
+	const FOptimusPinTraversalContext& InTraversalContext,
+	const FOptimus_NodeToDataInterfaceMap& InNodeDataInterfaceMap,
+	const FOptimus_PinToDataInterfaceMap& InLinkDataInterfaceMap,
+	const TArray<const UOptimusNode *>& InValueNodes,
+	const UComputeDataInterface* InGraphDataInterface,
+	const UOptimusComponentSourceBinding* InGraphDataComponentBinding,
+	FOptimus_InterfaceBindingMap& OutInputDataBindings,
+	FOptimus_InterfaceBindingMap& OutOutputDataBindings
+) const
+{
+	// Maps friendly name to unique name for each struct type
+	TMap<FName, FName> StructTypeDefs;
+
+	auto ReturnError = [](FText InErrorText) -> FOptimus_ComputeKernelResult
+	{
+		return FOptimus_ComputeKernelResult(TInPlaceType<FText>(), InErrorText);
+	};
+
+	auto GetStructTypeDefFromPin = [&StructTypeDefs](const UOptimusNodePin* InPin) -> TOptional<FText>
+	{
+		const FOptimusDataTypeHandle TypeHandle = InPin->GetDataType();
+		const FShaderValueTypeHandle ShaderValueType = TypeHandle->ShaderValueType;
+
+		// Only process shader struct types
+		if (ShaderValueType->Type != EShaderFundamentalType::Struct)
+		{
+			return {};
+		}
+			
+		TArray<FShaderValueTypeHandle> StructTypes = ShaderValueType->GetMemberStructTypes();
+		StructTypes.Add(ShaderValueType);
+
+		for (const FShaderValueTypeHandle StructType : StructTypes)
+		{
+			const FOptimusDataTypeHandle StructOptimusType = FOptimusDataTypeRegistry::Get().FindType(StructType->Name);
+			if (ensure(StructOptimusType))
+			{
+				if (UScriptStruct* ScriptStruct = Cast<UScriptStruct>(StructOptimusType->TypeObject))
+				{
+					const FName UniqueName = StructType->Name;
+					const FName FriendlyNameForKernel = Optimus::GetTypeName(ScriptStruct, false);
+
+					if (FName* FoundUniqueName = StructTypeDefs.Find(FriendlyNameForKernel))
+					{
+						// The same friendly name cannot be claimed by two unique types;
+						if (*FoundUniqueName != UniqueName)
+						{
+							return FText::Format(LOCTEXT("InvalidPinFriendlyName", "Invalid unique friendly name on pin '{}'"), FText::FromName(InPin->GetUniqueName()));
+						}
+
+						// Type is already in the map, no more actions needed
+						return {};
+					}
+
+					StructTypeDefs.Add(FriendlyNameForKernel) = UniqueName;
+				}
+			}
+		}
+		return {};
+	};
+
+	TSet<UOptimusComponentSourceBinding*> PrimaryBinding;
+	for (const UOptimusNodePin* Pin: GetPins())
+	{
+		if (Pin->IsGroupingPin())
+		{
+			TSet<UOptimusComponentSourceBinding*> SecondaryBindings;
+			
+			for (const UOptimusNodePin* SubPin: Pin->GetSubPins())
+			{
+				if (TOptional<FText> Result = GetStructTypeDefFromPin(SubPin); Result.IsSet())
+				{
+					return ReturnError(Result.GetValue());
+				}
+			}
+		}
+		else 
+		{
+			if (TOptional<FText> Result = GetStructTypeDefFromPin(Pin); Result.IsSet())
+			{
+				return ReturnError(Result.GetValue());
+			}
+		}
+	}
+
+	UOptimusKernelSource* KernelSource = NewObject<UOptimusKernelSource>(InKernelSourceOuter);
+
+	// Wrap functions for unconnected resource pins (or value pins) that return default values
+	// (for reads) or do nothing (for writes).
+	TArray<FString> GeneratedFunctions;
+
+	for (const UOptimusNodePin* Pin: GetPins())
+	{
+		if (Pin->IsGroupingPin())
+		{
+			for (const UOptimusNodePin* SubPin: Pin->GetSubPins())
+			{
+				ProcessInputPinForComputeKernel(
+					InTraversalContext, SubPin, Pin->GetName(),
+					InNodeDataInterfaceMap, InLinkDataInterfaceMap, InValueNodes,
+					InGraphDataInterface, InGraphDataComponentBinding,
+					KernelSource, GeneratedFunctions, OutInputDataBindings
+					);
+			}
+		}
+		else if (Pin->GetDirection() == EOptimusNodePinDirection::Input)
+		{
+			ProcessInputPinForComputeKernel(
+				InTraversalContext, Pin, FString(),
+				InNodeDataInterfaceMap, InLinkDataInterfaceMap, InValueNodes,
+				InGraphDataInterface, InGraphDataComponentBinding,
+				KernelSource, GeneratedFunctions, OutInputDataBindings
+				);
+		}
+		else if (Pin->GetDirection() == EOptimusNodePinDirection::Output)
+		{
+			ProcessOutputPinForComputeKernel(
+				InTraversalContext, Pin,
+				InNodeDataInterfaceMap, InLinkDataInterfaceMap,
+				KernelSource, GeneratedFunctions, OutOutputDataBindings);
+		}
+	}
+
+	FString CookedSource =
+		"#include \"/Engine/Private/Common.ush\"\n"
+		"#include \"/Plugin/ComputeFramework/Private/ComputeKernelCommon.ush\"\n\n";
+	CookedSource += FString::Join(GeneratedFunctions, TEXT("\n"));
+	CookedSource += "\n\n";
+	CookedSource += GetKernelSourceText();
+	
+	KernelSource->SetSource(CookedSource);
+	KernelSource->EntryPoint = GetKernelName();
+	KernelSource->GroupSize = GetGroupSize();
+	KernelSource->AdditionalSources = GetAdditionalSources();
+
+	for (const TTuple<FName, FName>& NamePair : StructTypeDefs)
+	{
+		if (NamePair.Key != NamePair.Value)
+		{
+			KernelSource->DefinitionsSet.Defines.Add(FComputeKernelDefinition(NamePair.Key.ToString(), NamePair.Value.ToString()));
+		}
+	}
+
+#if 0
+	UE_LOG(LogOptimusCore, Log, TEXT("Kernel: %s [%s]"), *GetNodePath(), *GetNodeName().ToString());
+	UE_LOG(LogOptimusCore, Log, TEXT("Cooked Source:\n%s\n"), *CookedSource);
+	if (!OutInputDataBindings.IsEmpty())
+	{
+		UE_LOG(LogOptimusCore, Log, TEXT("Input Bindings:"));
+		for (const TPair<int32, FOptimus_InterfaceBinding>& Binding: OutInputDataBindings)
+		{
+			TArray<FShaderFunctionDefinition> Defs;
+			Binding.Value.DataInterface->GetSupportedInputs(Defs);
+			UE_LOG(LogOptimusCore, Log, TEXT("  K[%d] %s -> %s@%d [%s]"), Binding.Key, *Binding.Value.BindingFunctionName,
+				*Binding.Value.DataInterface->GetName(), Binding.Value.DataInterfaceBindingIndex,
+				Defs.IsValidIndex(Binding.Value.DataInterfaceBindingIndex) ? *Defs[Binding.Value.DataInterfaceBindingIndex].Name : TEXT("<undefined>"));
+		}
+	}
+	if (!OutOutputDataBindings.IsEmpty())
+	{
+		UE_LOG(LogOptimusCore, Log, TEXT("Output Bindings:"));
+		for (const TPair<int32, FOptimus_InterfaceBinding>& Binding: OutOutputDataBindings)
+		{
+			TArray<FShaderFunctionDefinition> Defs;
+			Binding.Value.DataInterface->GetSupportedOutputs(Defs);
+			UE_LOG(LogOptimusCore, Log, TEXT("  K[%d] %s -> %s@%d [%s]"), Binding.Key, *Binding.Value.BindingFunctionName,
+				*Binding.Value.DataInterface->GetName(), Binding.Value.DataInterfaceBindingIndex,
+				Defs.IsValidIndex(Binding.Value.DataInterfaceBindingIndex) ? *Defs[Binding.Value.DataInterfaceBindingIndex].Name : TEXT("<undefined>"));
+		}
+	}
+#endif
+	
+	return FOptimus_ComputeKernelResult(TInPlaceType<UOptimusKernelSource*>(), KernelSource);
+}
+
+TOptional<FText> UOptimusNode_ComputeKernelBase::ValidateForCompile() const
+{
+	auto GetStructTypeDefFromPin = [](const UOptimusNodePin* InPin) -> TOptional<FText>
+	{
+		const FOptimusDataTypeHandle TypeHandle = InPin->GetDataType();
+		if (!TypeHandle.IsValid())
+		{
+			return FText::Format(LOCTEXT("InvalidPinType", "Invalid data type on pin '{0}'"), FText::FromName(InPin->GetUniqueName()));
+		}
+
+		const FShaderValueTypeHandle ShaderValueType = TypeHandle->ShaderValueType;
+		if (!ShaderValueType.IsValid())
+		{
+			return FText::Format(LOCTEXT("InvalidPinValueType", "Invalid shader value type on pin '{0}'"), FText::FromName(InPin->GetUniqueName()));
+		}
+
+		return {};
+	};
+
+	// Collect struct types and ensure all groups have the same component source (or none) within each group
+	auto VerifyPinBindings = [](TSet<UOptimusComponentSourceBinding*> &InCollectedBindings, const UOptimusNodePin* InPin)-> TOptional<FText>
+	{
+		// Component bindings only matter for resource pins.
+		if (!InPin->GetDataDomain().IsSingleton())
+		{
+			const TSet<UOptimusComponentSourceBinding*> PinBindings = InPin->GetComponentSourceBindings();
+			if (PinBindings.Num() > 1)
+			{
+				return FText::Format(LOCTEXT("MultipleBindingsOnPin", "Multiple component bindings arriving into pin '{0}'"), FText::FromName(InPin->GetUniqueName()));
+			}
+			if (InCollectedBindings.IsEmpty())
+			{
+				InCollectedBindings = PinBindings;
+			}
+			else if (!PinBindings.IsEmpty())
+			{
+				if (InCollectedBindings.Num() != PinBindings.Num() || !InCollectedBindings.Includes(PinBindings))
+				{
+					return FText::Format(LOCTEXT("IncompatibleBindingsPinGroup", "Component binding for pin '{0}' is different from the component bindings of the other pins in its group"), FText::FromName(InPin->GetUniqueName()));			}
+			}
+		}
+		return {};
+	};
+	
+	TSet<UOptimusComponentSourceBinding*> PrimaryBinding;
+	for (const UOptimusNodePin* Pin: GetPins())
+	{
+		if (Pin->IsGroupingPin())
+		{
+			TSet<UOptimusComponentSourceBinding*> SecondaryBindings;
+			
+			for (const UOptimusNodePin* SubPin: Pin->GetSubPins())
+			{
+				if(TOptional<FText> Result = VerifyPinBindings(SecondaryBindings, Pin); Result.IsSet())
+				{
+					return Result;
+				}
+				
+				if (TOptional<FText> Result = GetStructTypeDefFromPin(SubPin); Result.IsSet())
+				{
+					return Result;
+				}
+			}
+		}
+		else 
+		{
+			if (Pin->GetDirection() == EOptimusNodePinDirection::Input)
+			{
+				if(TOptional<FText> Result = VerifyPinBindings(PrimaryBinding, Pin); Result.IsSet())
+				{
+					return Result;
+				}
+			}
+			if (TOptional<FText> Result = GetStructTypeDefFromPin(Pin); Result.IsSet())
+			{
+				return Result;
+			}
+		}
+	}
+	return {};
+}
+
+
+FString UOptimusNode_ComputeKernelBase::GetCookedKernelSource(
+	const FString& InObjectPathName,
+	const FString& InShaderSource,
+	const FString& InKernelName,
+	FIntVector InGroupSize
+	)
+{
+	// FIXME: Create source range mappings so that we can go from error location to
+	// our source.
+	FString Source = InShaderSource;
+
+#if PLATFORM_WINDOWS
+	// Remove old-school stuff.
+	Source.ReplaceInline(TEXT("\r"), TEXT(""));
+#endif
+
+	const bool bHasKernelKeyword = Source.Contains(TEXT("KERNEL"));
+	
+	const FString KernelFunc = FString::Printf(
+		TEXT("[numthreads(%d,%d,%d)]\nvoid %s(uint3 DTid : SV_DispatchThreadID)"), 
+		InGroupSize.X, InGroupSize.Y, InGroupSize.Z, *InKernelName);
+	
+	if (bHasKernelKeyword)
+	{
+		Source.ReplaceInline(TEXT("KERNEL"), TEXT("void __kernel_func(uint Index)"));
+
+		return FString::Printf(
+			TEXT(
+				"#line 1 \"%s\"\n"
+				"%s\n\n"
+				"%s { __kernel_func(DTid.x); }\n"
+				), *InObjectPathName, *Source, *KernelFunc);
+	}
+	else
+	{
+		return FString::Printf(
+		TEXT(
+			"%s\n"
+			"{\n"
+			"uint Index = DTid.x;\n"
+			"#line 1 \"%s\"\n"
+			"%s\n"
+			"}\n"
+			), *KernelFunc, *InObjectPathName, *Source);
+	}
+}
+
+
+void UOptimusNode_ComputeKernelBase::ProcessInputPinForComputeKernel(
+	const FOptimusPinTraversalContext& InTraversalContext,
+	const UOptimusNodePin* InInputPin,
+	const FString& InGroupName,
+    const FOptimus_NodeToDataInterfaceMap& InNodeDataInterfaceMap,
+	const FOptimus_PinToDataInterfaceMap& InLinkDataInterfaceMap,
+	const TArray<const UOptimusNode*>& InValueNodes,
+	const UComputeDataInterface* InGraphDataInterface,
+	const UOptimusComponentSourceBinding* InGraphDataComponentBinding,
+	UOptimusKernelSource* InKernelSource,
+	TArray<FString>& OutGeneratedFunctions,
+	FOptimus_InterfaceBindingMap& OutInputDataBindings
+	) const
+{
+	const UOptimusNodePin* OutputPin = nullptr;
+	TArray<FOptimusRoutedNodePin> ConnectedPins = InInputPin->GetConnectedPinsWithRouting(InTraversalContext);
+	if (ConnectedPins.Num() == 1)
+	{
+		OutputPin = ConnectedPins[0].NodePin;
+	}
+	
+	FShaderParamTypeDefinition IndexParamDef;
+	CopyValueType(FShaderValueType::Get(EShaderFundamentalType::Uint), IndexParamDef);
+
+	const FShaderValueTypeHandle ValueType = InInputPin->GetDataType()->ShaderValueType;
+	const UOptimusNode *OutputNode = OutputPin ? OutputPin->GetOwningNode() : nullptr;
+
+	// For inputs, we only have to deal with a single read, because only one
+	// link can connect into it. 
+	if (OutputPin)
+	{
+		const UComputeDataInterface* DataInterface = nullptr;
+		const UOptimusComponentSourceBinding* ComponentBinding = nullptr;
+		int32 DataInterfaceFuncIndex = INDEX_NONE;
+		FString DataFunctionName;
+		
+		// Are we being connected from a scene data interface or a transient buffer?
+		if (InLinkDataInterfaceMap.Contains(OutputPin))
+		{
+			// For transient buffers we need the function index as given by the
+			// ReadValue function. 
+			DataInterface = InLinkDataInterfaceMap[OutputPin];
+			DataInterfaceFuncIndex = UOptimusTransientBufferDataInterface::ReadValueInputIndex;
+
+			TArray<FShaderFunctionDefinition> ReadFunctions;
+			DataInterface->GetSupportedInputs(ReadFunctions);
+
+			// Get the component binding from the upstream connection.
+			TArray<UOptimusComponentSourceBinding*> ComponentBindings = GetOwningGraph()->GetComponentSourceBindingsForPin(OutputPin).Array();
+			if (ensure(ComponentBindings.Num() == 1))
+			{
+				ComponentBinding = ComponentBindings[0];
+			}
+
+			DataFunctionName = ReadFunctions[DataInterfaceFuncIndex].Name;
+		}
+		else if(InNodeDataInterfaceMap.Contains(OutputNode))
+		{
+			const IOptimusDataInterfaceProvider* InterfaceProvider = Cast<const IOptimusDataInterfaceProvider>(OutputNode);
+			
+			// FIXME: Sub-pin read support.
+			UOptimusComputeDataInterface const* OptimusDataInterface = InNodeDataInterfaceMap[OutputNode];
+			DataInterface = OptimusDataInterface;
+			ComponentBinding = InterfaceProvider->GetComponentBinding();
+
+			DataInterfaceFuncIndex = InterfaceProvider->GetDataFunctionIndexFromPin(OutputPin);
+			
+			TArray<FShaderFunctionDefinition> ReadFunctions;
+			DataInterface->GetSupportedInputs(ReadFunctions);
+
+			DataFunctionName = ReadFunctions[DataInterfaceFuncIndex].Name;
+		}
+		else if (IOptimusValueProvider const* ValueProvider = Cast<const IOptimusValueProvider>(OutputNode))
+		{
+			// Value nodes bind the single graph data interface.
+			DataInterface = InGraphDataInterface;
+			ComponentBinding = InGraphDataComponentBinding;
+			DataInterfaceFuncIndex = InValueNodes.Find(OutputNode);
+			check(DataInterfaceFuncIndex != INDEX_NONE);
+			DataFunctionName = ValueProvider->GetValueName();
+		}
+
+		// If we are connected from a data interface, set the input binding up now.
+		if (DataInterface && ensure(ComponentBinding))
+		{
+			// The shader function definition that exposes the function that we use to
+			// read values to input into the kernel.
+			FShaderFunctionDefinition FuncDef;
+			FuncDef.Name = DataFunctionName;
+			FuncDef.bHasReturnType = true;
+			
+			FShaderParamTypeDefinition ParamDef;
+			CopyValueType(ValueType, ParamDef);
+			FuncDef.ParamTypes.Emplace(ParamDef);
+
+			// For resources we need the index parameter.
+			if (!InInputPin->GetDataDomain().IsSingleton())
+			{
+				TArray<FName> LevelNames = InInputPin->GetDataDomain().DimensionNames;
+
+				for (int32 Count = 0; Count < LevelNames.Num(); Count++)
+				{
+					FuncDef.ParamTypes.Add(IndexParamDef);
+				}
+			}
+
+			FOptimus_InterfaceBinding InterfaceBinding;
+			InterfaceBinding.DataInterface = DataInterface;
+			InterfaceBinding.ComponentBinding = ComponentBinding;
+			InterfaceBinding.DataInterfaceBindingIndex = DataInterfaceFuncIndex;
+			InterfaceBinding.BindingFunctionName = FString::Printf(TEXT("Read%s"), *InInputPin->GetName());
+			InterfaceBinding.BindingFunctionNamespace = InGroupName;
+			
+			OutInputDataBindings.Add(InKernelSource->ExternalInputs.Num(), InterfaceBinding);  
+			
+			InKernelSource->ExternalInputs.Emplace(FuncDef);
+		}
+	}
+	else
+	{
+		// Nothing connected. Get the default value (for now).
+		FString ValueStr;
+		FString OptionalParamStr;
+		if (InInputPin->GetDataDomain().IsSingleton())
+		{
+			ValueStr = GetShaderParamPinValueString(InInputPin);
+		}
+		else
+		{
+			TArray<FName> LevelNames = InInputPin->GetDataDomain().DimensionNames;
+			
+			ValueStr = InInputPin->GetDataType()->ShaderValueType->GetZeroValueAsString();
+
+			// No output connections, leave a stub function. The compiler will be in charge
+			// of optimizing out anything that causes us to ends up here.
+			TArray<FString> StubIndexes;
+			
+			for (const FString &IndexName: GetIndexNamesFromDataDomainLevels(LevelNames))
+			{
+				StubIndexes.Add(*FString::Printf(TEXT("uint %s"), *IndexName));
+			}
+		
+			OptionalParamStr = *FString::Join(StubIndexes, TEXT(", "));
+		}
+
+		FString NamespacePrefix, NamespaceSuffix;
+		if (!InGroupName.IsEmpty())
+		{
+			NamespacePrefix = FString::Printf(TEXT("namespace %s { "), *InGroupName);
+			NamespaceSuffix = TEXT(" }");
+		}
+
+		OutGeneratedFunctions.Add(
+			FString::Printf(TEXT("%s%s Read%s(%s) { return %s; }%s"),
+				*NamespacePrefix, *ValueType->ToString(), *InInputPin->GetName(), *OptionalParamStr, *ValueStr, *NamespaceSuffix));
+	}
+}
+
+
+void UOptimusNode_ComputeKernelBase::ProcessOutputPinForComputeKernel(
+	const FOptimusPinTraversalContext& InTraversalContext,
+	const UOptimusNodePin* InOutputPin,
+	const FOptimus_NodeToDataInterfaceMap& InNodeDataInterfaceMap,
+	const FOptimus_PinToDataInterfaceMap& InLinkDataInterfaceMap,
+	UOptimusKernelSource* InKernelSource,
+	TArray<FString>& OutGeneratedFunctions,
+	FOptimus_InterfaceBindingMap& OutOutputDataBindings
+	) const
+{
+	FShaderParamTypeDefinition IndexParamDef;
+	CopyValueType(FShaderValueType::Get(EShaderFundamentalType::Uint), IndexParamDef);
+
+	TArray<FName> LevelNames = InOutputPin->GetDataDomain().DimensionNames;
+	TArray<FString> IndexNames = GetIndexNamesFromDataDomainLevels(LevelNames);
+	const FShaderValueTypeHandle ValueType = InOutputPin->GetDataType()->ShaderValueType;
+
+	TArray<FOptimusRoutedNodePin> ConnectedPins = InOutputPin->GetConnectedPinsWithRouting(InTraversalContext);
+	
+	if (!ConnectedPins.IsEmpty())
+	{
+		// If we have an output connection going to multiple data interfaces, then we
+		// have to wrap them all up in a single proxy function to make it still transparent
+		// to the kernel writer.
+		struct FWriteConnectionDef
+		{
+			UOptimusComputeDataInterface* DataInterface = nullptr;
+			const UOptimusComponentSourceBinding* ComponentBinding = nullptr;
+			FString DataFunctionName;
+			FString WriteToName;
+		};
+		TArray<FWriteConnectionDef> WriteConnectionDefs;
+
+		// If we're scheduled to write to a transient data interface, do that now.
+		// There is only ever a single transient data interface per output pin.
+		if (InLinkDataInterfaceMap.Contains(InOutputPin))
+		{
+			UOptimusComputeDataInterface* DataInterface = InLinkDataInterfaceMap[InOutputPin]; 
+
+			TArray<FShaderFunctionDefinition> WriteFunctions;
+			DataInterface->GetSupportedOutputs(WriteFunctions);
+
+			// This is a horrible hack for detecting interlocked writes
+			// TODO: Either express this via the kernel metadata or add full support for buffer data interface in graph editor.
+			int32 WriteValueOutputIndex = UOptimusTransientBufferDataInterface::WriteValueOutputIndex;
+			if (InOutputPin->GetName().Contains(TEXT("Interlocked")) && WriteFunctions.Num() > WriteValueOutputIndex + 1)
+			{
+				++WriteValueOutputIndex;
+			}
+
+			// Get the component binding from the upstream connection.
+			UOptimusComponentSourceBinding* ComponentBinding = nullptr;
+			TArray<UOptimusComponentSourceBinding*> ComponentBindings = GetOwningGraph()->GetComponentSourceBindingsForPin(InOutputPin).Array();
+			if (ensure(ComponentBindings.Num() == 1))
+			{
+				ComponentBinding = ComponentBindings[0];
+			}
+			
+			WriteConnectionDefs.Add({DataInterface, ComponentBinding, WriteFunctions[WriteValueOutputIndex].Name, TEXT("Transient")});
+		}
+		
+		for (const FOptimusRoutedNodePin& ConnectedPin: ConnectedPins)
+		{
+			const UOptimusNode *ConnectedNode = ConnectedPin.NodePin->GetOwningNode();
+
+			// Connected to a data interface node?
+			if(!InNodeDataInterfaceMap.Contains(ConnectedNode))
+			{
+				continue;
+			}
+
+			const IOptimusDataInterfaceProvider* InterfaceProvider = Cast<const IOptimusDataInterfaceProvider>(ConnectedNode);
+			
+			// FIXME: Sub-pin write support.
+			UOptimusComputeDataInterface* DataInterface = InNodeDataInterfaceMap[ConnectedNode];
+			int32 DataFunctionIndex = InterfaceProvider->GetDataFunctionIndexFromPin(ConnectedPin.NodePin);
+			const UOptimusComponentSourceBinding* ComponentBinding = InterfaceProvider->GetComponentBinding();
+			
+			TArray<FShaderFunctionDefinition> FunctionDefinitions;
+			DataInterface->GetSupportedOutputs(FunctionDefinitions);
+			FString DataFunctionName = FunctionDefinitions[DataFunctionIndex].Name;
+			
+			WriteConnectionDefs.Add({DataInterface, ComponentBinding, DataFunctionName, ConnectedPin.NodePin->GetName()});
+		}
+
+		TArray<FString> WrapFunctionNameCalls;
+
+		for (int32 WriteConnectionIndex = 0; WriteConnectionIndex < WriteConnectionDefs.Num(); ++WriteConnectionIndex)
+		{
+			const FWriteConnectionDef& WriteConnectionDef = WriteConnectionDefs[WriteConnectionIndex];
+			const FString DataFunctionName = WriteConnectionDef.DataFunctionName;
+			FShaderFunctionDefinition FuncDef;
+			FuncDef.Name = DataFunctionName;
+			FuncDef.bHasReturnType = false;
+		
+			FShaderParamTypeDefinition ParamDef;
+			CopyValueType(ValueType, ParamDef);
+			for (int32 Count = 0; Count < LevelNames.Num(); Count++)
+			{
+				FuncDef.ParamTypes.Add(IndexParamDef);
+			}
+			FuncDef.ParamTypes.Emplace(ParamDef);
+
+			TArray<FShaderFunctionDefinition> WriteFunctions;
+			WriteConnectionDef.DataInterface->GetSupportedOutputs(WriteFunctions);
+			int32 DataInterfaceFuncIndex = WriteFunctions.IndexOfByPredicate([DataFunctionName](const FShaderFunctionDefinition &InDef) { return DataFunctionName == InDef.Name; });
+		
+			FString WrapFunctionName;
+			if (WriteConnectionDefs.Num() > 1)
+			{
+				WrapFunctionName = FString::Printf(TEXT("Write%sTo%s%d"), *InOutputPin->GetName(), *WriteConnectionDef.WriteToName, WriteConnectionIndex);
+				WrapFunctionNameCalls.Add(FString::Printf(TEXT("    %s(%s, Value)"), *WrapFunctionName, *FString::Join(IndexNames, TEXT(", "))));
+			}
+			else
+			{
+				WrapFunctionName = FString::Printf(TEXT("Write%s"), *InOutputPin->GetName());
+			}
+			
+			FOptimus_InterfaceBinding InterfaceBinding;
+			InterfaceBinding.DataInterface = WriteConnectionDef.DataInterface;
+			InterfaceBinding.ComponentBinding = WriteConnectionDef.ComponentBinding;
+			InterfaceBinding.DataInterfaceBindingIndex = DataInterfaceFuncIndex;
+			InterfaceBinding.BindingFunctionName = WrapFunctionName;
+			
+			OutOutputDataBindings.Add(InKernelSource->ExternalOutputs.Num(), InterfaceBinding);
+			InKernelSource->ExternalOutputs.Emplace(FuncDef);
+		}
+
+		if (!WrapFunctionNameCalls.IsEmpty())
+		{
+			TArray<FString> IndexParamNames;
+			for (const FString& IndexName: IndexNames)
+			{
+				IndexParamNames.Add(FString::Printf(TEXT("uint %s"), *IndexName));
+			}
+			
+			// Add a wrapper function that calls all the write functions in one shot.
+			OutGeneratedFunctions.Add(
+				FString::Printf(TEXT("void Write%s(%s, %s Value)\n{\n%s;\n}"),
+					*InOutputPin->GetName(), *FString::Join(IndexParamNames, TEXT(", ")), *ValueType->ToString(), *FString::Join(WrapFunctionNameCalls, TEXT(";\n"))));
+		}
+	}
+	else
+	{
+		// No output connections, leave a stub function. The compiler will be in charge
+		// of optimizing out anything that causes us to ends up here.
+		TArray<FString> StubIndexes;
+
+		for (const FString &IndexName: IndexNames)
+		{
+			StubIndexes.Add(*FString::Printf(TEXT("uint %s"), *IndexName));
+		}
+		
+		OutGeneratedFunctions.Add(
+			FString::Printf(TEXT("void Write%s(%s, %s Value) { }"),
+				*InOutputPin->GetName(), *FString::Join(StubIndexes, TEXT(", ")), *ValueType->ToString()));
+	}
+}
+
+#undef LOCTEXT_NAMESPACE
diff --git a/Engine/Plugins/Animation/DeformerGraph/Source/OptimusCore/Private/Nodes/OptimusNode_ComputeKernelBase.h b/Engine/Plugins/Animation/DeformerGraph/Source/OptimusCore/Private/Nodes/OptimusNode_ComputeKernelBase.h
new file mode 100644
--- /dev/null
+++ b/Engine/Plugins/Animation/DeformerGraph/Source/OptimusCore/Private/Nodes/OptimusNode_ComputeKernelBase.h
@@ -0,0 +1,102 @@
+﻿// Copyright Epic Games, Inc. All Rights Reserved.
+
+#pragma once
+
+#include "IOptimusComputeKernelProvider.h"
+#include "OptimusDataDomain.h"
+#include "OptimusDataType.h"
+#include "OptimusNode.h"
+
+#include "OptimusNode_ComputeKernelBase.generated.h"
+
+class UComputeSource;
+
+
+UCLASS(Abstract)
+class UOptimusNode_ComputeKernelBase :
+	public UOptimusNode, 
+	public IOptimusComputeKernelProvider
+{
+	GENERATED_BODY()
+	
+public:
+	/** Implement this to return the HLSL kernel's function name */
+	virtual FString GetKernelName() const PURE_VIRTUAL(UOptimusNode_ComputeKernelBase::GetKernelName, return FString();)
+
+	/** Implement this to return the HLSL kernel's function name */
+	virtual FIntVector GetGroupSize() const PURE_VIRTUAL(UOptimusNode_ComputeKernelBase::GetGroupSize, return FIntVector();)
+
+	/** Implement this to return the complete HLSL code for this kernel */
+	virtual FString GetKernelSourceText() const PURE_VIRTUAL(UOptimusNode_ComputeKernelBase::GetKernelSourceText, return FString();)
+
+	/** Implement this to return additional source code that the kernel requires. */
+	virtual TArray<TObjectPtr<UComputeSource>> GetAdditionalSources() const { return {}; }
+	
+	// IOptimusComputeKernelProvider
+	FOptimus_ComputeKernelResult CreateComputeKernel(
+		UObject* InKernelSourceOuter,
+		const FOptimusPinTraversalContext& InTraversalContext,
+		const FOptimus_NodeToDataInterfaceMap& InNodeDataInterfaceMap,
+		const FOptimus_PinToDataInterfaceMap& InLinkDataInterfaceMap,
+		const TArray<const UOptimusNode*>& InValueNodes,
+		const UComputeDataInterface* InGraphDataInterface,
+		const UOptimusComponentSourceBinding* InGraphDataComponentBinding,
+		FOptimus_InterfaceBindingMap& OutInputDataBindings,
+		FOptimus_InterfaceBindingMap& OutOutputDataBindings
+	) const override;
+
+	FName GetExecutionDomain() const override PURE_VIRTUAL(UOptimusNode_ComputeKernelBase::GetExecutionDomain, return NAME_None; );
+	TArray<const UOptimusNodePin*> GetPrimaryGroupInputPins() const override PURE_VIRTUAL(UOptimusNode_ComputeKernelBase::GetPrimaryGroupInputPins, return {}; );  
+
+	// -- UOptimusNode overrides
+	TOptional<FText> ValidateForCompile() const override;
+	
+protected:
+	static TArray<FString> GetIndexNamesFromDataDomainLevels(
+		const TArray<FName> &InLevelNames
+		)
+	{
+		TArray<FString> IndexNames;
+	
+		for (FName DomainName: InLevelNames)
+		{
+			IndexNames.Add(FString::Printf(TEXT("%sIndex"), *DomainName.ToString()));
+		}
+		return IndexNames;
+	}
+
+	static FString GetCookedKernelSource(
+		const FString& InObjectPathName,
+		const FString& InShaderSource,
+		const FString& InKernelName,
+		FIntVector InGroupSize
+		);
+	
+	
+private:
+	void ProcessInputPinForComputeKernel(
+		const FOptimusPinTraversalContext& InTraversalContext,
+		const UOptimusNodePin* InInputPin,
+		const FString& InGroupName,
+		const FOptimus_NodeToDataInterfaceMap& InNodeDataInterfaceMap,
+		const FOptimus_PinToDataInterfaceMap& InLinkDataInterfaceMap,
+		const TArray<const UOptimusNode*>& InValueNodes,
+		const UComputeDataInterface* InGraphDataInterface,
+		const UOptimusComponentSourceBinding* InGraphDataComponentBinding,
+		UOptimusKernelSource* InKernelSource,
+		TArray<FString>& OutGeneratedFunctions,
+		FOptimus_InterfaceBindingMap& OutInputDataBindings
+		) const;
+
+	void ProcessOutputPinForComputeKernel(
+		const FOptimusPinTraversalContext& InTraversalContext,
+		const UOptimusNodePin* InOutputPin,
+		const FOptimus_NodeToDataInterfaceMap& InNodeDataInterfaceMap,
+		const FOptimus_PinToDataInterfaceMap& InLinkDataInterfaceMap,
+		UOptimusKernelSource* InKernelSource,
+		TArray<FString>& OutGeneratedFunctions,
+		FOptimus_InterfaceBindingMap& OutOutputDataBindings
+		) const;
+
+};
+
diff --git a/Engine/Plugins/Animation/DeformerGraph/Source/OptimusCore/Private/Nodes/OptimusNode_ComputeKernelFunction.h b/Engine/Plugins/Animation/DeformerGraph/Source/OptimusCore/Private/Nodes/OptimusNode_ComputeKernelFunction.h
new file mode 100644
--- /dev/null
+++ b/Engine/Plugins/Animation/DeformerGraph/Source/OptimusCore/Private/Nodes/OptimusNode_ComputeKernelFunction.h
@@ -0,0 +1,83 @@
+﻿// Copyright Epic Games, Inc. All Rights Reserved.
+
+#pragma once
+
+#include "IOptimusExecutionDomainProvider.h"
+#include "OptimusNode_ComputeKernelBase.h"
+#include "OptimusBindingTypes.h"
+
+#include "OptimusNode_ComputeKernelFunction.generated.h"
+
+
+UCLASS()
+class UOptimusNode_ComputeKernelFunctionGeneratorClass :
+	public UClass
+{
+	GENERATED_BODY()
+public:
+	static UClass *CreateNodeClass(
+		UObject* InPackage,
+		FName InCategory,
+		FName InKernelName,
+		FIntVector InGroupSize,
+		const TArray<FOptimusParameterBinding>& InInputBindings,
+		const TArray<FOptimusParameterBinding>& InOutputBindings,
+		const FString& InShaderSource
+		);
+
+	// UClass overrides
+	void InitPropertiesFromCustomList(uint8* InObjectPtr, const uint8* InCDOPtr) override;
+	void Link(FArchive& Ar, bool bRelinkExistingProperties) override;
+	
+	UPROPERTY()
+	FName Category;
+
+	UPROPERTY()
+	FName KernelName;
+
+	UPROPERTY()
+	FOptimusExecutionDomain ExecutionDomain;
+
+	UPROPERTY()
+	FIntVector GroupSize;
+
+	UPROPERTY()
+	TArray<FOptimusParameterBinding> InputBindings;
+
+	UPROPERTY()
+	TArray<FOptimusParameterBinding> OutputBindings;
+
+	UPROPERTY()
+	FString ShaderSource;
+};
+
+
+/**
+ * 
+ */
+UCLASS(Hidden)
+class OPTIMUSCORE_API UOptimusNode_ComputeKernelFunction :
+	public UOptimusNode_ComputeKernelBase
+{
+	GENERATED_BODY()
+
+public:
+	UOptimusNode_ComputeKernelFunction();
+
+	// UOptimusNode overrides
+	FText GetDisplayName() const override;
+	FName GetNodeCategory() const override; 
+
+	// UOptimusNode_ComputeKernelBase overrides
+	FString GetKernelName() const override;
+	FIntVector GetGroupSize() const override;
+	FString GetKernelSourceText() const override;
+
+	void ConstructNode() override;
+	
+	// IOptimusComputeKernelProvider
+	FName GetExecutionDomain() const override; 
+	TArray<const UOptimusNodePin*> GetPrimaryGroupInputPins() const override { return {}; }
+private:
+	UOptimusNode_ComputeKernelFunctionGeneratorClass *GetGeneratorClass() const;
+};
diff --git a/Engine/Plugins/Animation/DeformerGraph/Source/OptimusCore/Private/OptimusDeformerInstance.cpp b/Engine/Plugins/Animation/DeformerGraph/Source/OptimusCore/Private/OptimusDeformerInstance.cpp
new file mode 100644
--- /dev/null
+++ b/Engine/Plugins/Animation/DeformerGraph/Source/OptimusCore/Private/OptimusDeformerInstance.cpp
@@ -0,0 +1,520 @@
+﻿// Copyright Epic Games, Inc. All Rights Reserved.
+
+#include "OptimusDeformerInstance.h"
+
+#include "Components/MeshComponent.h"
+#include "DataInterfaces/OptimusDataInterfaceGraph.h"
+#include "OptimusComputeGraph.h"
+#include "OptimusDataTypeRegistry.h"
+#include "OptimusDeformer.h"
+#include "OptimusVariableDescription.h"
+#include "RenderGraphBuilder.h"
+#include "SkeletalRenderPublic.h"
+#include "DataInterfaces/OptimusDataInterfaceRawBuffer.h"
+
+
+/** Container for a pooled buffer. */
+struct FOptimusPersistentStructuredBuffer
+{
+	TRefCountPtr<FRDGPooledBuffer> PooledBuffer;
+	int32 ElementStride = 0;
+	int32 ElementCount = 0;
+};
+
+
+void FOptimusPersistentBufferPool::GetResourceBuffers(
+	FRDGBuilder& GraphBuilder,
+	FName InResourceName,
+	int32 InLODIndex,
+	int32 InElementStride,
+	TArray<int32> const& InElementCounts,
+	TArray<FRDGBufferRef>& OutBuffers)
+{
+	OutBuffers.Reset();
+
+	TMap<int32, TArray<FOptimusPersistentStructuredBuffer>>& LODResources = ResourceBuffersMap.FindOrAdd(InResourceName);  
+	TArray<FOptimusPersistentStructuredBuffer>* ResourceBuffersPtr = LODResources.Find(InLODIndex);
+	if (ResourceBuffersPtr == nullptr)
+	{
+		// Create pooled buffers and store.
+		TArray<FOptimusPersistentStructuredBuffer> ResourceBuffers;
+		ResourceBuffers.Reserve(InElementCounts.Num());
+
+		for (int32 Index = 0; Index < InElementCounts.Num(); Index++)
+		{
+			FRDGBufferDesc BufferDesc = FRDGBufferDesc::CreateStructuredDesc(InElementStride, InElementCounts[Index]);
+			FRDGBufferRef Buffer = GraphBuilder.CreateBuffer(BufferDesc, TEXT("FOptimusPersistentBuffer"), ERDGBufferFlags::None);
+			OutBuffers.Add(Buffer);
+
+			FOptimusPersistentStructuredBuffer& PersistentBuffer = ResourceBuffers.AddDefaulted_GetRef();
+			PersistentBuffer.ElementStride = InElementStride;
+			PersistentBuffer.ElementCount = InElementCounts[Index];
+			PersistentBuffer.PooledBuffer = GraphBuilder.ConvertToExternalBuffer(Buffer);
+		}
+
+		LODResources.Add(InLODIndex, MoveTemp(ResourceBuffers));
+	}
+	else
+	{
+		// Verify that the buffers are correct based on the incoming information. 
+		// If there's a mismatch, then something has gone wrong upstream.
+		// Maybe either duplicated names, missing resource clearing on recompile, or something else.
+		if (!ensure(ResourceBuffersPtr->Num() == InElementCounts.Num()))
+		{
+			return;
+		}
+
+		for (int32 Index = 0; Index < ResourceBuffersPtr->Num(); Index++)
+		{
+			FOptimusPersistentStructuredBuffer& PersistentBuffer = (*ResourceBuffersPtr)[Index];
+			if (!ensure(PersistentBuffer.PooledBuffer.IsValid()) ||
+				!ensure(PersistentBuffer.ElementStride == InElementStride) ||
+				!ensure(PersistentBuffer.ElementCount == InElementCounts[Index]))
+			{
+				OutBuffers.Reset();
+				return;
+			}	
+
+			// Register buffer back into the graph and return it.
+			FRDGBufferRef Buffer = GraphBuilder.RegisterExternalBuffer(PersistentBuffer.PooledBuffer);
+			OutBuffers.Add(Buffer);
+		}
+	}
+}
+
+void FOptimusPersistentBufferPool::ReleaseResources()
+{
+	check(IsInRenderingThread());
+	ResourceBuffersMap.Reset();
+}
+
+FOptimusDeformerInstanceExecInfo::FOptimusDeformerInstanceExecInfo()
+{
+	GraphType = EOptimusNodeGraphType::Update;
+}
+
+
+void UOptimusDeformerInstanceSettings::RefreshComponentBindings(
+	UOptimusDeformer* InDeformer,
+	UMeshComponent* InMeshComponent
+	)
+{
+	
+	// Try to retain existing component binding as much as possible if called after a compilation or component rebuild
+	TMap<FName, UActorComponent*> ExistingBindings;
+
+	for (FOptimusDeformerInstanceComponentBinding& Binding: Bindings)
+	{
+		// It's likely that a component we referred to has been nixed, via AActor::RerunConstructionScripts, so
+		// we need to re-fetch the object from memory.
+		Binding.ActorComponent.ResetWeakPtr();
+		ExistingBindings.Add(Binding.ProviderName, Binding.ActorComponent.Get());
+	}
+	
+	AActor* Actor = InMeshComponent->GetOwner();
+	TSet<UActorComponent*> ComponentsUsed;
+
+	const TArray<UOptimusComponentSourceBinding*>& ComponentBindings = InDeformer->GetComponentBindings();
+	Bindings.Reset(ComponentBindings.Num());
+	for (const UOptimusComponentSourceBinding* Binding: ComponentBindings)
+	{
+		FName BindingName = Binding->BindingName;
+		UActorComponent* BoundComponent = nullptr;
+
+		// Primary binding always binds to the mesh component we're applied to.
+		if (Binding->IsPrimaryBinding())
+		{
+			BoundComponent = InMeshComponent;
+		}
+		else
+		{
+			// Try an existing binding first and see if they still match by class. We ignore tags for this match
+			// because we want to respect the will of the user, unless absolutely not possible (i.e. class mismatch).
+			if (ExistingBindings.Contains(BindingName))
+			{
+				if (UActorComponent* Component = ExistingBindings[BindingName])
+				{
+					if (Component->IsA(Binding->GetComponentSource()->GetComponentClass()))
+					{
+						BoundComponent = Component;
+					}
+				}
+			}
+			
+			// If not, try to find a component owned by this actor that matches the tag and class.
+			if (!BoundComponent && Actor != nullptr && !Binding->ComponentTags.IsEmpty())
+			{
+				TSet<UActorComponent*> TaggedComponents;
+				for (FName Tag: Binding->ComponentTags)
+				{
+					TArray<UActorComponent*> Components = Actor->GetComponentsByTag(Binding->GetComponentSource()->GetComponentClass(), Tag);
+
+					for (UActorComponent* Component: Components)
+					{
+						if (!ComponentsUsed.Contains(Component))
+						{
+							TaggedComponents.Add(Component);
+						}
+					}
+				}
+				TArray<UActorComponent*> RankedTaggedComponents = TaggedComponents.Array();
+
+				// Rank the components by the number of tags they match.
+				RankedTaggedComponents.Sort([Tags=TSet<FName>(Binding->ComponentTags)](const UActorComponent& InCompA, const UActorComponent& InCompB)
+				{
+					TSet<FName> TagsA(InCompA.ComponentTags);
+					TSet<FName> TagsB(InCompB.ComponentTags);
+					
+					return Tags.Intersect(TagsA).Num() < Tags.Intersect(TagsB).Num();
+				});
+
+				if (!RankedTaggedComponents.IsEmpty())
+				{
+					BoundComponent = RankedTaggedComponents[0];
+				}
+			}
+
+			// Otherwise just use class matching on components owned by the actor.
+			if (!BoundComponent && Actor != nullptr)
+			{
+				TArray<UActorComponent*> Components;
+				Actor->GetComponents(Binding->GetComponentSource()->GetComponentClass(), Components);
+
+				for (UActorComponent* Component: Components)
+				{
+					if (!ComponentsUsed.Contains(Component))
+					{
+						BoundComponent = Component;
+						break;
+					}
+				}
+			}
+		}
+
+		Bindings.Add({BindingName, BoundComponent});
+		ComponentsUsed.Add(BoundComponent);
+	}
+}
+
+TArray<UActorComponent*> UOptimusDeformerInstanceSettings::GetBoundComponents() const
+{
+	TArray<UActorComponent*> Result;
+	for (const FOptimusDeformerInstanceComponentBinding& Binding: Bindings)
+	{
+		UActorComponent* ActorComponent = Binding.ActorComponent.Get();
+		if (ActorComponent)
+		{
+			Result.Add(ActorComponent);
+		}
+	}
+	return Result;
+}
+
+
+AActor* UOptimusDeformerInstanceSettings::GetActor() const
+{
+	// We should be owned by an actor at some point.
+	return GetTypedOuter<AActor>();
+}
+
+UOptimusComponentSourceBinding* UOptimusDeformerInstanceSettings::GetComponentBindingByName(FName InBindingName) const
+{
+	if (const UOptimusDeformer* DeformerResolved = Deformer.Get())
+	{
+		for (UOptimusComponentSourceBinding* Binding: DeformerResolved->GetComponentBindings())
+		{
+			if (Binding->BindingName == InBindingName)
+			{
+				return Binding;
+			}
+		}
+	}
+	return nullptr;
+}
+
+#if WITH_EDITOR
+void UOptimusDeformerInstanceSettings::PostEditChangeProperty(FPropertyChangedEvent& PropertyChangedEvent)
+{
+	if (PropertyChangedEvent.GetPropertyName() == GET_MEMBER_NAME_CHECKED(FOptimusDeformerInstanceComponentBinding, ActorComponent))
+	{
+		// TODO: Update dependencies.
+	}
+}
+#endif
+
+void UOptimusDeformerInstanceSettings::InitializeSettings(
+	UOptimusDeformer* InDeformer,
+	UMeshComponent* InMeshComponent
+	)
+{
+	Deformer = InDeformer;
+	RefreshComponentBindings(InDeformer, InMeshComponent);
+}
+
+
+void UOptimusDeformerInstance::SetMeshComponent(UMeshComponent* InMeshComponent)
+{ 
+	MeshComponent = InMeshComponent;
+}
+
+void UOptimusDeformerInstance::SetInstanceSettings(UOptimusDeformerInstanceSettings* InInstanceSettings)
+{
+	InstanceSettings = InInstanceSettings; 
+}
+
+
+void UOptimusDeformerInstance::SetupFromDeformer(
+	UOptimusDeformer* InDeformer,
+	const bool bInRefreshBindings
+	)
+{
+	// If we're doing a recompile, ditch all stored render resources.
+	ReleaseResources();
+
+	// Update the component bindings before creating data providers. The bindings are in the same order
+	// as the component bindings in the deformer.
+	UOptimusDeformerInstanceSettings* InstanceSettingsPtr = InstanceSettings.Get(); 
+	if (bInRefreshBindings && InstanceSettingsPtr)
+	{
+		InstanceSettingsPtr->RefreshComponentBindings(InDeformer, MeshComponent.Get());
+	}
+
+	// Create the persistent buffer pool
+	BufferPool = MakeShared<FOptimusPersistentBufferPool>();
+	
+	// (Re)Create and bind data providers.
+	ComputeGraphExecInfos.Reset();
+	GraphsToRunOnNextTick.Reset();
+
+	for (const FOptimusComputeGraphInfo& ComputeGraphInfo : InDeformer->ComputeGraphs)
+	{
+		FOptimusDeformerInstanceExecInfo& Info = ComputeGraphExecInfos.AddDefaulted_GetRef();
+		Info.GraphName = ComputeGraphInfo.GraphName;
+		Info.GraphType = ComputeGraphInfo.GraphType;
+		Info.ComputeGraph = ComputeGraphInfo.ComputeGraph;
+
+		if (InstanceSettingsPtr)
+		{
+			// Fall back on everything being the given component.
+			for (int32 Index = 0; Index < InstanceSettingsPtr->Bindings.Num(); Index++)
+			{
+				Info.ComputeGraphInstance.CreateDataProviders(Info.ComputeGraph, Index, InstanceSettingsPtr->Bindings[Index].ActorComponent.Get());
+			}
+		}
+		else
+		{
+			// Fall back on everything being the given component.
+			for (int32 Index = 0; Index < InDeformer->GetComponentBindings().Num(); Index++)
+			{
+				Info.ComputeGraphInstance.CreateDataProviders(Info.ComputeGraph, Index, MeshComponent.Get());
+			}
+		}
+
+		int32 LODIndex = 0;
+		if (const USkinnedMeshComponent* SkinnedMeshComponent = Cast<USkinnedMeshComponent>(MeshComponent))
+		{
+			// This guff should be a utility function on USkinnedMeshComponent. 
+			LODIndex = SkinnedMeshComponent->GetPredictedLODLevel();
+			
+			if (SkinnedMeshComponent->GetSkinnedAsset() && SkinnedMeshComponent->GetSkinnedAsset()->IsStreamable() && SkinnedMeshComponent->MeshObject)
+			{
+				LODIndex = FMath::Max<int32>(LODIndex, SkinnedMeshComponent->MeshObject->GetSkeletalMeshRenderData().PendingFirstLODIdx);
+			}
+		}
+
+		for(TObjectPtr<UComputeDataProvider> DataProvider: Info.ComputeGraphInstance.GetDataProviders())
+		{
+			// Make the persistent buffer data provider aware of the buffer pool and current LOD index.
+			// TBD: Interface-based.
+			if (UOptimusPersistentBufferDataProvider* PersistentBufferProvider = Cast<UOptimusPersistentBufferDataProvider>(DataProvider))
+			{
+				PersistentBufferProvider->BufferPool = BufferPool;
+				PersistentBufferProvider->LODIndex = LODIndex;
+			}
+		}
+
+		// Schedule the setup graph to run.
+		if (Info.GraphType == EOptimusNodeGraphType::Setup)
+		{
+			GraphsToRunOnNextTick.Add(Info.GraphName);
+		}
+	}
+	
+	// Create local storage for deformer graph variables.
+	Variables = NewObject<UOptimusVariableContainer>(this);
+	Variables->Descriptions.Reserve(InDeformer->GetVariables().Num());
+	TSet<const UOptimusVariableDescription*> Visited;
+	for (const UOptimusVariableDescription* VariableDescription : InDeformer->GetVariables())
+	{
+		if (!VariableDescription)
+			continue;
+		if (Visited.Contains(VariableDescription))
+			continue;
+		Visited.Add(VariableDescription);
+		UOptimusVariableDescription* VariableDescriptionCopy = NewObject<UOptimusVariableDescription>();
+		VariableDescriptionCopy->Guid = VariableDescription->Guid;
+		VariableDescriptionCopy->VariableName = VariableDescription->VariableName;
+		VariableDescriptionCopy->DataType = VariableDescription->DataType;
+		VariableDescriptionCopy->ValueData = VariableDescription->ValueData;
+		VariableDescriptionCopy->ResetValueDataSize();
+		Variables->Descriptions.Add(VariableDescriptionCopy);
+	}
+
+	if (UMeshComponent* Ptr = MeshComponent.Get())
+	{
+		Ptr->MarkRenderDynamicDataDirty();
+	}
+}
+
+
+void UOptimusDeformerInstance::SetCanBeActive(bool bInCanBeActive)
+{
+	bCanBeActive = bInCanBeActive;
+}
+
+void UOptimusDeformerInstance::AllocateResources()
+{
+	
+}
+
+
+void UOptimusDeformerInstance::ReleaseResources()
+{
+	if (BufferPool)
+	{
+		ENQUEUE_RENDER_COMMAND(FOptimusReleasePoolMemory)(
+			[BufferPool=MoveTemp(BufferPool)](FRHICommandListImmediate& InCmdList)
+			{
+				BufferPool->ReleaseResources();
+			});
+	}
+}
+
+
+bool UOptimusDeformerInstance::IsActive() const
+{
+	if (!bCanBeActive)
+	{
+		return false;
+	}
+		
+	for (const FOptimusDeformerInstanceExecInfo& Info: ComputeGraphExecInfos)
+	{
+		if (!Info.ComputeGraphInstance.ValidateDataProviders(Info.ComputeGraph))
+		{
+			return false;
+		}
+	}
+	return !ComputeGraphExecInfos.IsEmpty();
+}
+
+void UOptimusDeformerInstance::EnqueueWork(FSceneInterface* InScene, EWorkLoad InWorkLoadType, FName InOwnerName)
+{
+	TSet<FName> GraphsToRun;
+	{
+		UE::TScopeLock<FCriticalSection> Lock(GraphsToRunOnNextTickLock);
+		Swap(GraphsToRunOnNextTick, GraphsToRun);
+	}
+	
+	for (FOptimusDeformerInstanceExecInfo& Info: ComputeGraphExecInfos)
+	{
+		if (Info.GraphType == EOptimusNodeGraphType::Update || GraphsToRun.Contains(Info.GraphName))
+		{
+			Info.ComputeGraphInstance.EnqueueWork(Info.ComputeGraph, InScene, InOwnerName);
+		}
+	}
+}
+
+namespace
+{
+	template <typename T>
+	bool SetVariableValue(UOptimusVariableContainer const* InVariables, FName InVariableName, FName InTypeName, T const& InValue)
+	{
+		const uint8* ValueBytes = reinterpret_cast<const uint8*>(&InValue);
+
+		FOptimusDataTypeHandle WantedType = FOptimusDataTypeRegistry::Get().FindType(InTypeName);
+		for (UOptimusVariableDescription* VariableDesc : InVariables->Descriptions)
+		{
+			if (VariableDesc->VariableName == InVariableName && VariableDesc->DataType == WantedType)
+			{
+				TUniquePtr<FProperty> Property(WantedType->CreateProperty(nullptr, NAME_None));
+				if (ensure(Property->GetSize() == sizeof(T)) &&
+					ensure(Property->GetSize() == VariableDesc->ValueData.Num()))
+				{
+					FPlatformMemory::Memcpy(VariableDesc->ValueData.GetData(), ValueBytes, sizeof(T));
+				}
+
+				return true;
+			}
+		}
+
+		return false;
+	}
+}
+
+
+bool UOptimusDeformerInstance::SetBoolVariable(FName InVariableName, bool InValue)
+{
+	return SetVariableValue(Variables, InVariableName, FBoolProperty::StaticClass()->GetFName(), InValue);
+}
+
+bool UOptimusDeformerInstance::SetIntVariable(FName InVariableName, int32 InValue)
+{
+	return SetVariableValue<int32>(Variables, InVariableName, FIntProperty::StaticClass()->GetFName(), InValue);
+}
+
+bool UOptimusDeformerInstance::SetFloatVariable(FName InVariableName, float InValue)
+{
+	return SetVariableValue<float>(Variables, InVariableName, FFloatProperty::StaticClass()->GetFName(), InValue);
+}
+
+bool UOptimusDeformerInstance::SetVectorVariable(FName InVariableName, const FVector& InValue)
+{
+	return SetVariableValue<FVector>(Variables, InVariableName, "FVector", InValue);
+}
+
+bool UOptimusDeformerInstance::SetVector4Variable(FName InVariableName, const FVector4& InValue)
+{
+	return SetVariableValue<FVector4>(Variables, InVariableName, "FVector4", InValue);
+}
+
+const TArray<UOptimusVariableDescription*>& UOptimusDeformerInstance::GetVariables() const
+{
+	return Variables->Descriptions;
+}
+
+
+bool UOptimusDeformerInstance::EnqueueTriggerGraph(FName InTriggerGraphName)
+{
+	for(FOptimusDeformerInstanceExecInfo& ExecInfo: ComputeGraphExecInfos)
+	{
+		if (ExecInfo.GraphType == EOptimusNodeGraphType::ExternalTrigger && ExecInfo.GraphName == InTriggerGraphName)
+		{
+			UE::TScopeLock<FCriticalSection> Lock(GraphsToRunOnNextTickLock);
+			GraphsToRunOnNextTick.Add(ExecInfo.GraphName);
+			return true;
+		}
+	}
+	
+	return false;
+}
+
+
+void UOptimusDeformerInstance::SetConstantValueDirect(FString const& InVariableName, TArray<uint8> const& InValue)
+{
+	// Poke constants into the UGraphDataProvider objects.
+	// This is an editor only operation when constant nodes are edited in the graph and we want to see the result without a full compile step.
+	// Not sure that this is the best approach, or whether to store constants on the UOptimusDeformerInstance like variables?
+	for (FOptimusDeformerInstanceExecInfo& ExecInfo : ComputeGraphExecInfos)
+	{
+		TArray< TObjectPtr<UComputeDataProvider> >& DataProviders = ExecInfo.ComputeGraphInstance.GetDataProviders();
+		for (UComputeDataProvider* DataProvider : DataProviders)
+		{
+			if (UOptimusGraphDataProvider* GraphDataProvider = Cast<UOptimusGraphDataProvider>(DataProvider))
+			{
+				GraphDataProvider->SetConstant(InVariableName, InValue);
+				break;
+			}
+		}
+	}
+}
diff --git a/Engine/Plugins/Animation/DeformerGraph/Source/OptimusCore/Public/IOptimusComputeKernelProvider.h b/Engine/Plugins/Animation/DeformerGraph/Source/OptimusCore/Public/IOptimusComputeKernelProvider.h
new file mode 100644
--- /dev/null
+++ b/Engine/Plugins/Animation/DeformerGraph/Source/OptimusCore/Public/IOptimusComputeKernelProvider.h
@@ -0,0 +1,104 @@
+﻿// Copyright Epic Games, Inc. All Rights Reserved.
+
+#pragma once
+
+#include "ComputeFramework/ShaderParamTypeDefinition.h"
+#include "UObject/Interface.h"
+#include "Misc/TVariant.h"
+
+#include "IOptimusComputeKernelProvider.generated.h"
+
+class UComputeDataInterface;
+class UOptimusComputeDataInterface;
+class UOptimusComponentSourceBinding;
+class UOptimusKernelSource;
+class UOptimusNode;
+class UOptimusNodePin;
+struct FOptimusPinTraversalContext;
+
+
+// Maps the data interface's data binding index to the function we would like to have present
+// during kernel compilation to read/write values from/to that data interface's resource.
+struct FOptimus_InterfaceBinding
+{
+	const UComputeDataInterface* DataInterface;
+	const UOptimusComponentSourceBinding* ComponentBinding;
+	int32 DataInterfaceBindingIndex;
+	FString BindingFunctionName;
+	FString BindingFunctionNamespace;
+};
+using FOptimus_InterfaceBindingMap = TMap<int32 /* Kernel Index */, FOptimus_InterfaceBinding>;
+
+// A map that goes from a value/variable node to a compute shader input parameter.
+struct FOptimus_KernelParameterBinding
+{
+	/** The node to retrieve the value from */
+	const UOptimusNode* ValueNode;
+	
+	/** The name of the shader parameter */ 
+	FString ParameterName;
+	
+	/** The value type of the parameter */
+	FShaderValueTypeHandle ValueType;
+};
+using FOptimus_KernelParameterBindingList = TArray<FOptimus_KernelParameterBinding>;
+
+// Maps from a data interface node to the data interface that it represents.
+using FOptimus_NodeToDataInterfaceMap =  TMap<const UOptimusNode*, UOptimusComputeDataInterface*>;
+
+// Maps from an output pin to the transient data interface, used to store intermediate results,
+// that it represents.
+using FOptimus_PinToDataInterfaceMap = TMap<const UOptimusNodePin*, UOptimusComputeDataInterface*>;
+
+
+using FOptimus_ComputeKernelResult = TVariant<UOptimusKernelSource* /* Kernel */, FText /* Error */>;
+
+UINTERFACE()
+class OPTIMUSCORE_API UOptimusComputeKernelProvider :
+	public UInterface
+{
+	GENERATED_BODY()
+};
+
+/**
+ * Interface that provides a mechanism to identify and work with node graph owners.
+ */
+class OPTIMUSCORE_API IOptimusComputeKernelProvider
+{
+	GENERATED_BODY()
+
+public:
+	/**
+	 * Return an UOptimusKernelSource object, from a compute kernel node state that implements
+	 * this interface.
+	 * @param InKernelSourceOuter The outer object that will own the new kernel source.
+	 * @param InTraversalContext The current context being used to traverse from the graph. Used mainly to keep track of graph nesting.
+	 * @param InNodeDataInterfaceMap Map from UOptimusNode objects to data interfaces.
+	 * @param InLinkDataInterfaceMap
+	 * @param InValueNodes
+	 * @param InGraphDataInterface
+	 * @param InGraphDataComponentBinding
+	 * @param OutInputDataBindings
+	 * @param OutOutputDataBindings
+	 * @param OutExecutionDataInterface
+	 */
+	virtual FOptimus_ComputeKernelResult CreateComputeKernel(
+		UObject* InKernelSourceOuter,
+		const FOptimusPinTraversalContext& InTraversalContext,
+		const FOptimus_NodeToDataInterfaceMap& InNodeDataInterfaceMap,
+		const FOptimus_PinToDataInterfaceMap& InLinkDataInterfaceMap,
+		const TArray<const UOptimusNode*>& InValueNodes,
+		const UComputeDataInterface* InGraphDataInterface,
+		const UOptimusComponentSourceBinding* InGraphDataComponentBinding,
+		FOptimus_InterfaceBindingMap& OutInputDataBindings,
+		FOptimus_InterfaceBindingMap& OutOutputDataBindings
+		) const = 0;
+
+	/** Returns the execution domain that this kernel should iterate over */
+	virtual FName GetExecutionDomain() const = 0;
+
+	/** Returns the input pins of the primary group only. Used for traversing up the graph to identify the component
+	 *  source.
+	 */
+	virtual TArray<const UOptimusNodePin*> GetPrimaryGroupInputPins() const = 0; 
+};
diff --git a/Engine/Plugins/Animation/DeformerGraph/Source/OptimusCore/Public/IOptimusExecutionDomainProvider.h b/Engine/Plugins/Animation/DeformerGraph/Source/OptimusCore/Public/IOptimusExecutionDomainProvider.h
new file mode 100644
--- /dev/null
+++ b/Engine/Plugins/Animation/DeformerGraph/Source/OptimusCore/Public/IOptimusExecutionDomainProvider.h
@@ -0,0 +1,46 @@
+﻿// Copyright Epic Games, Inc. All Rights Reserved.
+
+#pragma once
+
+#include "UObject/Interface.h"
+
+#include "IOptimusExecutionDomainProvider.generated.h"
+
+/** A struct to hold onto a single-level domain for controlling a kernel's execution domain. 
+  * The reason it's in a struct is so that we can apply a property panel customization for it
+  * to make it easier to select from a pre-defined list of execution domains.
+*/
+USTRUCT()
+struct OPTIMUSCORE_API FOptimusExecutionDomain
+{
+	GENERATED_BODY()
+
+	FOptimusExecutionDomain() = default;
+	FOptimusExecutionDomain(FName InExecutionDomainName) : Name(InExecutionDomainName) {}
+
+	// The name of the execution domain that this kernel operates on.
+	UPROPERTY(EditAnywhere, Category = Domain)
+	FName Name;
+};
+
+UINTERFACE()
+class OPTIMUSCORE_API UOptimusExecutionDomainProvider :
+	public UInterface
+{
+	GENERATED_BODY()
+};
+
+
+/** An interface that should be implemented by any class that has an FOptimusExecutionDomain member.
+ *  This is used by the property detail specialization code to provide a list of possible domains.
+ */
+class IOptimusExecutionDomainProvider
+{
+	GENERATED_BODY()
+
+public:
+	/**
+	 * Returns a list of execution domains that the implementer of this interface provides.
+	 */
+	virtual TArray<FName> GetExecutionDomains() const = 0;
+};
diff --git a/Engine/Plugins/Animation/DeformerGraph/Source/OptimusCore/Public/IOptimusNodeFunctionLibraryOwner.h b/Engine/Plugins/Animation/DeformerGraph/Source/OptimusCore/Public/IOptimusNodeFunctionLibraryOwner.h
new file mode 100644
--- /dev/null
+++ b/Engine/Plugins/Animation/DeformerGraph/Source/OptimusCore/Public/IOptimusNodeFunctionLibraryOwner.h
@@ -0,0 +1,26 @@
+﻿// Copyright Epic Games, Inc. All Rights Reserved.
+
+#pragma once
+
+#include "UObject/Interface.h"
+
+#include "IOptimusNodeFunctionLibraryOwner.generated.h"
+
+
+UINTERFACE()
+class OPTIMUSCORE_API UOptimusNodeFunctionLibraryOwner :
+	public UInterface
+{
+	GENERATED_BODY()
+};
+
+/**
+  * Interface that provides a mechanism to identify and work with node function library owners
+  */
+class OPTIMUSCORE_API IOptimusNodeFunctionLibraryOwner
+{
+	GENERATED_BODY()
+
+public:
+	
+};
diff --git a/Engine/Plugins/Animation/DeformerGraph/Source/OptimusCore/Public/IOptimusValueProvider.h b/Engine/Plugins/Animation/DeformerGraph/Source/OptimusCore/Public/IOptimusValueProvider.h
new file mode 100644
--- /dev/null
+++ b/Engine/Plugins/Animation/DeformerGraph/Source/OptimusCore/Public/IOptimusValueProvider.h
@@ -0,0 +1,32 @@
+﻿// Copyright Epic Games, Inc. All Rights Reserved.
+
+#pragma once
+
+#include "OptimusDataType.h"
+#include "UObject/Interface.h"
+
+#include "IOptimusValueProvider.generated.h"
+
+UINTERFACE()
+class UOptimusValueProvider :
+	public UInterface
+{
+	GENERATED_BODY()
+};
+
+
+/** FIXME: A stop-gap shader value provider until we have a proper pin evaluation that handles
+  * paths that have a constant, computed, varying and a mix thereof, results.
+  */
+class IOptimusValueProvider
+{
+	GENERATED_BODY()
+
+public:
+	// Returns the value name.
+	virtual FString GetValueName() const = 0;
+	// Returns the value type.
+	virtual FOptimusDataTypeRef GetValueType() const = 0;
+	// Returns the stored value as a shader-compatible value.
+	virtual FShaderValueType::FValue GetShaderValue() const = 0;
+};
diff --git a/Engine/Plugins/Animation/DeformerGraph/Source/OptimusCore/Public/OptimusDeformerInstance.h b/Engine/Plugins/Animation/DeformerGraph/Source/OptimusCore/Public/OptimusDeformerInstance.h
new file mode 100644
--- /dev/null
+++ b/Engine/Plugins/Animation/DeformerGraph/Source/OptimusCore/Public/OptimusDeformerInstance.h
@@ -0,0 +1,234 @@
+﻿// Copyright Epic Games, Inc. All Rights Reserved.
+
+#pragma once
+
+#include "Animation/MeshDeformerInstance.h"
+#include "ComputeFramework/ComputeGraphInstance.h"
+
+#include "OptimusDeformerInstance.generated.h"
+
+enum class EOptimusNodeGraphType;
+struct FOptimusPersistentStructuredBuffer;
+class FRDGBuffer;
+class FRDGBuilder;
+class UActorComponent;
+class UMeshComponent;
+class UOptimusDeformer;
+class UOptimusVariableContainer;
+class UOptimusVariableDescription;
+class UOptimusComponentSourceBinding;
+
+class FOptimusPersistentBufferPool
+{
+public:
+	/** 
+	 * Get or allocate buffers for the given resource
+	 * If the buffer already exists but has different sizing characteristics the allocation fails. 
+	 * The number of buffers will equal the size of the InElementCount array.
+	 * But if the allocation fails, the returned array will be empty.
+	 */
+	void GetResourceBuffers(
+		FRDGBuilder& GraphBuilder,
+		FName InResourceName,
+		int32 InLODIndex,
+		int32 InElementStride,
+		TArray<int32> const& InElementCounts,
+		TArray<FRDGBuffer*>& OutBuffers );
+
+	/** Release _all_ resources allocated by this pool */
+	void ReleaseResources();
+	
+private:
+	TMap<FName, TMap<int32, TArray<FOptimusPersistentStructuredBuffer>>> ResourceBuffersMap;
+};
+using FOptimusPersistentBufferPoolPtr = TSharedPtr<FOptimusPersistentBufferPool>;
+
+
+/** Structure with cached state for a single compute graph. */
+USTRUCT()
+struct FOptimusDeformerInstanceExecInfo
+{
+	GENERATED_BODY()
+
+	FOptimusDeformerInstanceExecInfo();
+
+	/** The name of the graph */
+	UPROPERTY()
+	FName GraphName;
+
+	/** The graph type. */
+	UPROPERTY()
+	EOptimusNodeGraphType GraphType;
+	
+	/** The ComputeGraph asset. */
+	UPROPERTY()
+	TObjectPtr<UComputeGraph> ComputeGraph = nullptr;
+
+	/** The cached state for the ComputeGraph. */
+	UPROPERTY()
+	FComputeGraphInstance ComputeGraphInstance;
+};
+
+
+/** Defines a binding between a component provider in the graph and an actor component in the component hierarchy on
+ *  the actor whose deformable component we're bound to.
+ */
+USTRUCT(BlueprintType)
+struct FOptimusDeformerInstanceComponentBinding
+{
+	GENERATED_BODY()
+	
+	UPROPERTY(VisibleAnywhere, Category="Binding")
+	FName ProviderName;
+	
+	UPROPERTY(EditAnywhere, Category="Binding")
+	TSoftObjectPtr<UActorComponent> ActorComponent;
+};
+
+
+UCLASS(Blueprintable, BlueprintType)
+class OPTIMUSCORE_API UOptimusDeformerInstanceSettings :
+	public UMeshDeformerInstanceSettings
+{
+	GENERATED_BODY()
+
+public:
+	/** */
+	UPROPERTY(EditAnywhere, Category="Deformer|Settings", EditFixedSize, meta=(NoResetToDefault, EditFixedOrder))
+	TArray<FOptimusDeformerInstanceComponentBinding> Bindings;
+
+	void RefreshComponentBindings(
+		UOptimusDeformer* InDeformer,
+		UMeshComponent* InMeshComponent
+		);
+
+	TArray<UActorComponent*> GetBoundComponents() const;
+	AActor* GetActor() const;
+	UOptimusComponentSourceBinding* GetComponentBindingByName(FName InBindingName) const;
+
+#if WITH_EDITOR
+	// -- UObject
+	void PostEditChangeProperty(FPropertyChangedEvent& PropertyChangedEvent) override;
+#endif
+	
+protected:
+	friend class UOptimusDeformer;
+
+	void InitializeSettings(
+		UOptimusDeformer* InDeformer,
+		UMeshComponent* InMeshComponent);
+
+private:
+	UPROPERTY()
+	TWeakObjectPtr<UOptimusDeformer> Deformer; 
+};
+
+
+/** 
+ * Class representing an instance of an Optimus Mesh Deformer.
+ * This implements the UMeshDeformerInstance interface to enqueue the graph execution.
+ * It also contains the per instance deformer variable state and local state for each of the graphs in the deformer.
+ */
+UCLASS(Blueprintable, BlueprintType)
+class UOptimusDeformerInstance :
+	public UMeshDeformerInstance
+{
+	GENERATED_BODY()
+
+public:
+	/** 
+	 * Set the Mesh Component that owns this instance.
+	 * Call once before first call to SetupFromDeformer().
+	 */
+	void SetMeshComponent(UMeshComponent* InMeshComponent);
+
+	/** 
+	 * Set the instance settings that control this deformer instance. The deformer instance is transient whereas
+	 * the settings are persistent.
+	 */
+	void SetInstanceSettings(UOptimusDeformerInstanceSettings* InInstanceSettings);
+	
+	/** 
+	 * Setup the instance. 
+	 * Needs to be called after the UOptimusDeformer creates this instance, and whenever the instance is invalidated.
+	 * Invalidation happens whenever any bound Data Providers become invalid.
+	 */
+	void SetupFromDeformer(
+		UOptimusDeformer* InDeformer,
+		const bool bInRefreshBindings = false
+		);
+
+	void SetupFromDeformerAndRefreshBindings(UOptimusDeformer* InDeformer)
+	{
+		SetupFromDeformer(InDeformer, true);
+	}
+
+	/** Set the value of a boolean variable. */
+	UFUNCTION(BlueprintPure, Category="Deformer", meta=(DisplayName="Set Variable (bool)"))
+	bool SetBoolVariable(FName InVariableName, bool InValue);
+
+	/** Set the value of a boolean variable. */
+	UFUNCTION(BlueprintCallable, Category="Deformer", meta=(DisplayName="Set Variable (int)"))
+	bool SetIntVariable(FName InVariableName, int32 InValue);
+
+	/** Set the value of a boolean variable. */
+	UFUNCTION(BlueprintCallable, Category="Deformer", meta=(DisplayName="Set Variable (float)"))
+	bool SetFloatVariable(FName InVariableName, float InValue);
+
+	/** Set the value of a boolean variable. */
+	UFUNCTION(BlueprintCallable, Category="Deformer", meta=(DisplayName="Set Variable (Vector)"))
+	bool SetVectorVariable(FName InVariableName, const FVector& InValue);
+
+	/** Set the value of a boolean variable. */
+	UFUNCTION(BlueprintCallable, Category="Deformer", meta=(DisplayName="Set Variable (Vector4)"))
+	bool SetVector4Variable(FName InVariableName, const FVector4& InValue);
+
+	/** Get an array containing all the variables. */
+	UFUNCTION(BlueprintGetter)
+	const TArray<UOptimusVariableDescription*>& GetVariables() const;
+
+	/** Trigger a named trigger graph to run on the next tick */
+	UFUNCTION(BlueprintCallable, Category="Deformer")
+	bool EnqueueTriggerGraph(FName InTriggerGraphName);
+	
+	
+	
+	/** Directly set a graph constant value. */
+	void SetConstantValueDirect(FString const& InVariableName, TArray<uint8> const& InValue);
+
+	FOptimusPersistentBufferPoolPtr GetBufferPool() const { return BufferPool; }
+
+	void SetCanBeActive(bool bInCanBeActive);
+
+protected:
+	/** Implementation of UMeshDeformerInstance. */
+	void AllocateResources() override;
+	void ReleaseResources() override;
+	bool IsActive() const override;
+	void EnqueueWork(FSceneInterface* InScene, EWorkLoad InWorkLoadType, FName InOwnerName) override;
+
+private:
+	/** The Mesh Component that owns this Mesh Deformer Instance. */
+	UPROPERTY()
+	TWeakObjectPtr<UMeshComponent> MeshComponent;
+
+	/** The Mesh Component that owns this Mesh Deformer Instance. */
+	UPROPERTY()
+	TWeakObjectPtr<UOptimusDeformerInstanceSettings> InstanceSettings;
+	
+	/** An array of state. One for each graph owned by the deformer. */
+	UPROPERTY()
+	TArray<FOptimusDeformerInstanceExecInfo> ComputeGraphExecInfos;
+
+	/** Storage for variable data. */
+	UPROPERTY()
+	TObjectPtr<UOptimusVariableContainer> Variables;
+	
+	// List of graphs that should be run on the next tick. 
+	TSet<FName> GraphsToRunOnNextTick;
+	FCriticalSection GraphsToRunOnNextTickLock;
+
+	FOptimusPersistentBufferPoolPtr BufferPool;
+
+	bool bCanBeActive = true;
+};
diff --git a/Engine/Plugins/Animation/DeformerGraph/Source/OptimusCore/Public/OptimusResourceDescription.h b/Engine/Plugins/Animation/DeformerGraph/Source/OptimusCore/Public/OptimusResourceDescription.h
new file mode 100644
--- /dev/null
+++ b/Engine/Plugins/Animation/DeformerGraph/Source/OptimusCore/Public/OptimusResourceDescription.h
@@ -0,0 +1,66 @@
+// Copyright Epic Games, Inc. All Rights Reserved.
+
+#pragma once
+
+#include "OptimusDataType.h"
+
+#include "CoreMinimal.h"
+#include "OptimusComponentSource.h"
+#include "OptimusDataDomain.h"
+#include "UObject/Object.h"
+
+#include "OptimusResourceDescription.generated.h"
+
+
+class UOptimusComponentSourceBinding;
+class UOptimusDeformer;
+class UOptimusPersistentBufferDataInterface;
+
+
+UCLASS(BlueprintType)
+class OPTIMUSCORE_API UOptimusResourceDescription :
+	public UObject
+{
+	GENERATED_BODY()
+public:
+
+	UOptimusResourceDescription() = default;
+
+	/** Returns the owning deformer to operate on this resource */
+	// FIXME: Move to interface-based system.
+	UOptimusDeformer* GetOwningDeformer() const;
+
+	UPROPERTY(EditAnywhere, BlueprintReadOnly, Category = ResourceDescription)
+	FName ResourceName;
+
+	/** The the data type of each element of the resource */
+	UPROPERTY(EditAnywhere, Category=ResourceDescription, meta=(UseInResource))
+	FOptimusDataTypeRef DataType;
+
+	/** The component binding that this resource description is bound to */
+	UPROPERTY(VisibleAnywhere, Category=ResourceDescription)
+	TWeakObjectPtr<UOptimusComponentSourceBinding> ComponentBinding;
+	
+	/** The data domain for this resource. */
+	UPROPERTY(EditAnywhere, Category=ResourceDescription, meta=(EditCondition="ComponentBinding != nullptr", HideEditConditionToggle))
+	FOptimusDataDomain DataDomain;
+
+	UPROPERTY()
+	TObjectPtr<UOptimusPersistentBufferDataInterface> DataInterface;
+
+	// UObject overrides
+	void PostLoad() override;
+#if WITH_EDITOR
+	void PostEditChangeProperty(FPropertyChangedEvent& PropertyChangedEvent) override;
+	void PreEditUndo() override;
+	void PostEditUndo() override;
+#endif
+
+	bool IsValidComponentBinding() const;
+	FOptimusDataDomain GetDataDomainFromComponentBinding() const;
+
+private:
+#if WITH_EDITORONLY_DATA
+	FName ResourceNameForUndo;
+#endif
+};
diff --git a/Engine/Plugins/Animation/DeformerGraph/Source/OptimusTools/Private/OptimusToolsCommands.h b/Engine/Plugins/Animation/DeformerGraph/Source/OptimusTools/Private/OptimusToolsCommands.h
new file mode 100644
--- /dev/null
+++ b/Engine/Plugins/Animation/DeformerGraph/Source/OptimusTools/Private/OptimusToolsCommands.h
@@ -0,0 +1,25 @@
+// Copyright Epic Games, Inc. All Rights Reserved.
+
+#pragma once
+
+#include "CoreMinimal.h"
+#include "Styling/AppStyle.h"
+#include "Framework/Commands/Commands.h"
+
+
+class FOptimusToolsCommands : public TCommands<FOptimusToolsCommands>
+{
+public:
+	FOptimusToolsCommands()
+	    : TCommands<FOptimusToolsCommands>(TEXT("OptimusTools"), NSLOCTEXT("Contexts", "OptimusTools", "Optimus Tools"), NAME_None, FAppStyle::GetAppStyleSetName())
+	{
+	}
+
+	virtual void RegisterCommands() override;
+	static const FOptimusToolsCommands& Get();
+
+	// Modeling tools commands
+	TSharedPtr<FUICommandInfo> ToggleModelingToolsMode;
+
+
+};
diff --git a/Engine/Plugins/Experimental/MLDeformer/MLDeformerFramework/MLDeformerFramework.uplugin b/Engine/Plugins/Experimental/MLDeformer/MLDeformerFramework/MLDeformerFramework.uplugin
--- a/Engine/Plugins/Experimental/MLDeformer/MLDeformerFramework/MLDeformerFramework.uplugin
+++ b/Engine/Plugins/Experimental/MLDeformer/MLDeformerFramework/MLDeformerFramework.uplugin
@@ -28,7 +28,7 @@
 	],
 	"Plugins": [
 		{
-			"Name": "Optimus",
+			"Name": "DeformerGraph",
 			"Enabled": true
 		},
 		{
diff --git a/Engine/Plugins/Experimental/MLDeformer/NeuralMorphModel/NeuralMorphModel.uplugin b/Engine/Plugins/Experimental/MLDeformer/NeuralMorphModel/NeuralMorphModel.uplugin
--- a/Engine/Plugins/Experimental/MLDeformer/NeuralMorphModel/NeuralMorphModel.uplugin
+++ b/Engine/Plugins/Experimental/MLDeformer/NeuralMorphModel/NeuralMorphModel.uplugin
@@ -36,7 +36,7 @@
 			"Enabled": true
 		},
 		{
-			"Name": "Optimus",
+			"Name": "DeformerGraph",
 			"Enabled": true
 		},
 		{
