Optimus: Fix the content path for icons after moving the deformer graph plugin to its current location.

#rnx

diff --git a/Engine/Plugins/Animation/DeformerGraph/Source/OptimusEditor/Private/OptimusEditorStyle.cpp b/Engine/Plugins/Animation/DeformerGraph/Source/OptimusEditor/Private/OptimusEditorStyle.cpp
--- a/Engine/Plugins/Animation/DeformerGraph/Source/OptimusEditor/Private/OptimusEditorStyle.cpp
+++ b/Engine/Plugins/Animation/DeformerGraph/Source/OptimusEditor/Private/OptimusEditorStyle.cpp
@@ -16,7 +16,7 @@ FOptimusEditorStyle::FOptimusEditorStyle() :
 
 	static const FSlateColor DefaultForeground(FLinearColor(0.72f, 0.72f, 0.72f, 1.f));
 
-	SetContentRoot(FPaths::EnginePluginsDir() / TEXT("Experimental/Optimus/Resources"));
+	SetContentRoot(FPaths::EnginePluginsDir() / TEXT("Animation/DeformerGraph/Resources"));
 
 	// Text editor styles
 	{
@@ -35,7 +35,6 @@ FOptimusEditorStyle::FOptimusEditorStyle() :
 		const FTextBlockStyle ErrorText = FTextBlockStyle(NormalText)
 			.SetUnderlineBrush(IMAGE_BRUSH("Images/White", Icon8x8, FLinearColor::Red, ESlateBrushTileType::Both))
 			.SetColorAndOpacity(FLinearColor::Red);
-			// ;
 		
 		Set("TextEditor.NormalText", NormalText);
 
