Optimus: Refresh the data type ref's TypeObject on each load.

#rnx

diff --git a/Engine/Plugins/Animation/DeformerGraph/Source/OptimusCore/Private/OptimusDataType.cpp b/Engine/Plugins/Animation/DeformerGraph/Source/OptimusCore/Private/OptimusDataType.cpp
--- a/Engine/Plugins/Animation/DeformerGraph/Source/OptimusCore/Private/OptimusDataType.cpp
+++ b/Engine/Plugins/Animation/DeformerGraph/Source/OptimusCore/Private/OptimusDataType.cpp
@@ -55,6 +55,21 @@ FOptimusDataTypeHandle FOptimusDataTypeRef::Resolve() const
 	return TypeHandle;
 }
 
+
+void FOptimusDataTypeRef::PostSerialize(const FArchive& Ar)
+{
+	// Fix up data types so that the type points to the type object.
+	if (Ar.IsLoading())
+	{
+		const FOptimusDataTypeHandle TypeHandle = FOptimusDataTypeRegistry::Get().FindType(TypeName);
+		if (TypeHandle.IsValid())
+		{
+			TypeObject = TypeHandle->TypeObject;
+		}
+	}
+}
+
+
 FProperty* FOptimusDataType::CreateProperty(
 	UStruct* InScope, 
 	FName InName
diff --git a/Engine/Plugins/Animation/DeformerGraph/Source/OptimusCore/Private/OptimusNodePin.cpp b/Engine/Plugins/Animation/DeformerGraph/Source/OptimusCore/Private/OptimusNodePin.cpp
--- a/Engine/Plugins/Animation/DeformerGraph/Source/OptimusCore/Private/OptimusNodePin.cpp
+++ b/Engine/Plugins/Animation/DeformerGraph/Source/OptimusCore/Private/OptimusNodePin.cpp
@@ -518,14 +518,6 @@ PRAGMA_DISABLE_DEPRECATION_WARNINGS
 			DataDomain = FOptimusDataDomain();
 		}
 PRAGMA_ENABLE_DEPRECATION_WARNINGS
-
-		// Fix up data types so that the type points to the type object.
-		const FOptimusDataTypeRegistry& Registry = FOptimusDataTypeRegistry::Get();
-		FOptimusDataTypeHandle TypeHandle = Registry.FindType(DataType.TypeName);
-		if (TypeHandle.IsValid())
-		{
-			DataType.TypeObject = TypeHandle->TypeObject;
-		}
 	}
 }
 
diff --git a/Engine/Plugins/Animation/DeformerGraph/Source/OptimusCore/Public/OptimusDataType.h b/Engine/Plugins/Animation/DeformerGraph/Source/OptimusCore/Public/OptimusDataType.h
--- a/Engine/Plugins/Animation/DeformerGraph/Source/OptimusCore/Public/OptimusDataType.h
+++ b/Engine/Plugins/Animation/DeformerGraph/Source/OptimusCore/Public/OptimusDataType.h
@@ -164,4 +164,14 @@ struct OPTIMUSCORE_API FOptimusDataTypeRef
 	UPROPERTY(EditAnywhere, Category=Type)
 	TWeakObjectPtr<UObject> TypeObject;
 	
-};
\ No newline at end of file
+	void PostSerialize(const FArchive& Ar);
+};
+
+template<>
+struct TStructOpsTypeTraits<FOptimusDataTypeRef> : public TStructOpsTypeTraitsBase2<FOptimusDataTypeRef>
+{
+	enum 
+	{
+		WithPostSerialize = true,
+	};
+};
