Optimus: Fixed long load time issue. In the future we will only register registered animation attribute types, instead of all struct types

diff --git a/Engine/Plugins/Animation/DeformerGraph/Source/OptimusCore/Private/OptimusDataTypeRegistry.cpp b/Engine/Plugins/Animation/DeformerGraph/Source/OptimusCore/Private/OptimusDataTypeRegistry.cpp
--- a/Engine/Plugins/Animation/DeformerGraph/Source/OptimusCore/Private/OptimusDataTypeRegistry.cpp
+++ b/Engine/Plugins/Animation/DeformerGraph/Source/OptimusCore/Private/OptimusDataTypeRegistry.cpp
@@ -1295,7 +1295,6 @@ void FOptimusDataTypeRegistry::RegisterAssetRegistryCallbacks()
 {
 	if (FAssetRegistryModule* AssetRegistryModule = FModuleManager::GetModulePtr<FAssetRegistryModule>(TEXT("AssetRegistry")))
 	{
-		AssetRegistryModule->Get().OnFilesLoaded().AddRaw(&Get(), &FOptimusDataTypeRegistry::OnFilesLoaded);
 		AssetRegistryModule->Get().OnAssetRemoved().AddRaw(&Get(), &FOptimusDataTypeRegistry::OnAssetRemoved);
 		AssetRegistryModule->Get().OnAssetRenamed().AddRaw(&Get(), &FOptimusDataTypeRegistry::OnAssetRenamed);
 	}
@@ -1305,7 +1304,6 @@ void FOptimusDataTypeRegistry::UnregisterAssetRegistryCallbacks()
 {
 	if (FAssetRegistryModule* AssetRegistryModule = FModuleManager::GetModulePtr<FAssetRegistryModule>(TEXT("AssetRegistry")))
 	{
-		AssetRegistryModule->Get().OnFilesLoaded().RemoveAll(&Get());
 		AssetRegistryModule->Get().OnAssetRemoved().RemoveAll(&Get());
 		AssetRegistryModule->Get().OnAssetRenamed().RemoveAll(&Get());
 	}
@@ -1330,56 +1328,6 @@ FOptimusDataTypeRegistry::FOptimusDataTypeRegistry()
 {
 }
 
-void FOptimusDataTypeRegistry::OnFilesLoaded()
-{
-	TArray<FAssetData> StructAssets;
-	StructAssets.Reset();
-		
-	// Get the loaded ones
-	for (TObjectIterator<UScriptStruct> It; It; ++It)
-	{
-		UScriptStruct* Struct = *It;
-		StructAssets.Add(FAssetData(Struct));
-	}
-
-	// Now get unloaded ones
-	const FAssetRegistryModule& AssetRegistryModule = FModuleManager::LoadModuleChecked<FAssetRegistryModule>(TEXT("AssetRegistry"));
-	TArray<FAssetData> AssetData;
-	AssetRegistryModule.Get().GetAssetsByClass(UUserDefinedStruct::StaticClass()->GetClassPathName(), AssetData);
-
-	for (int32 AssetIndex = 0; AssetIndex < AssetData.Num(); ++AssetIndex)
-	{
-		const FAssetData& Asset = AssetData[AssetIndex];
-		if (Asset.IsValid() && !Asset.IsAssetLoaded())
-		{
-			StructAssets.Add(Asset);
-		}
-	}
-
-	StructAssets.Sort([](const FAssetData& A, const FAssetData& B) { return A.AssetName.LexicalLess(B.AssetName); });
-
-	for (const FAssetData& StructAsset : StructAssets)
-	{
-		UObject* Struct = StructAsset.GetAsset();
-		if (UObjectRedirector* RedirectedStruct = Cast<UObjectRedirector>(Struct))
-		{
-			Struct = RedirectedStruct->DestinationObject;
-		}
-		
-		UScriptStruct* ScriptStruct = Cast<UScriptStruct>(Struct);
-		
-		if (ensure(ScriptStruct))
-		{
-			FOptimusDataTypeHandle DataType = FindType(Optimus::GetTypeName(ScriptStruct));
-
-			if (!DataType.IsValid())
-			{
-				RegisterStructType(ScriptStruct);
-			}
-		}
-	}
-}
-
 void FOptimusDataTypeRegistry::OnAssetRemoved(const FAssetData& InAssetData)
 {
 	if (InAssetData.AssetClassPath == UUserDefinedStruct::StaticClass()->GetClassPathName())
diff --git a/Engine/Plugins/Animation/DeformerGraph/Source/OptimusCore/Public/OptimusDataTypeRegistry.h b/Engine/Plugins/Animation/DeformerGraph/Source/OptimusCore/Public/OptimusDataTypeRegistry.h
--- a/Engine/Plugins/Animation/DeformerGraph/Source/OptimusCore/Public/OptimusDataTypeRegistry.h
+++ b/Engine/Plugins/Animation/DeformerGraph/Source/OptimusCore/Public/OptimusDataTypeRegistry.h
@@ -200,8 +200,7 @@ protected:
 private:
 	FOptimusDataTypeRegistry();
 
-	// Callbacks to update registry when user defined structs are loaded/removed/renamed
-	void OnFilesLoaded();
+	// Callbacks to update registry when user defined structs are removed/renamed
 	void OnAssetRemoved(const FAssetData& InAssetData);
 	void OnAssetRenamed(const FAssetData& InAssetData, const FString& InOldName);
 	
