Optimus: Ensure that the Deformer Graph shows up in the asset filters.

#rnx

diff --git a/Engine/Plugins/Animation/DeformerGraph/Source/OptimusEditor/Private/OptimusDeformerAssetActions.cpp b/Engine/Plugins/Animation/DeformerGraph/Source/OptimusEditor/Private/OptimusDeformerAssetActions.cpp
--- a/Engine/Plugins/Animation/DeformerGraph/Source/OptimusEditor/Private/OptimusDeformerAssetActions.cpp
+++ b/Engine/Plugins/Animation/DeformerGraph/Source/OptimusEditor/Private/OptimusDeformerAssetActions.cpp
@@ -11,6 +11,30 @@
 #include "Widgets/Images/SImage.h"
 
 
+FOptimusDeformerAssetActions::FOptimusDeformerAssetActions(EAssetTypeCategories::Type InAssetCategoryBit) :
+	AssetCategoryBit(InAssetCategoryBit)
+{
+}
+
+
+FText FOptimusDeformerAssetActions::GetName() const
+{ 
+	return NSLOCTEXT("AssetTypeActions", "OptimusDeformerActions", "Deformer Graph"); 
+}
+
+
+FColor FOptimusDeformerAssetActions::GetTypeColor() const
+{ 
+	return FColor::Blue;
+}
+
+
+UClass* FOptimusDeformerAssetActions::GetSupportedClass() const
+{ 
+	return UOptimusDeformer::StaticClass(); 
+}
+
+
 void FOptimusDeformerAssetActions::OpenAssetEditor(const TArray<UObject*>& InObjects, TSharedPtr<class IToolkitHost> EditWithinLevelEditor /*= TSharedPtr<IToolkitHost>() */)
 {
 	EToolkitMode::Type Mode = EditWithinLevelEditor.IsValid() ? EToolkitMode::WorldCentric : EToolkitMode::Standalone;
@@ -34,6 +58,22 @@ void FOptimusDeformerAssetActions::OpenAssetEditor(const TArray<UObject*>& InObj
 }
 
 
+uint32 FOptimusDeformerAssetActions::GetCategories()
+{ 
+	return AssetCategoryBit;
+}
+
+
+const TArray<FText>& FOptimusDeformerAssetActions::GetSubMenus() const
+{
+	static const TArray<FText> SubMenus
+	{
+		NSLOCTEXT("FOptimusDeformerAssetActions", "AnimDeformersSubMenu", "Deformers")
+	};
+	return SubMenus;
+}
+
+
 TSharedPtr<SWidget> FOptimusDeformerAssetActions::GetThumbnailOverlay(const FAssetData& AssetData) const
 {
 	const FSlateBrush* Icon = FSlateIconFinder::FindIconBrushForClass(UOptimusDeformer::StaticClass());
diff --git a/Engine/Plugins/Animation/DeformerGraph/Source/OptimusEditor/Private/OptimusDeformerAssetActions.h b/Engine/Plugins/Animation/DeformerGraph/Source/OptimusEditor/Private/OptimusDeformerAssetActions.h
--- a/Engine/Plugins/Animation/DeformerGraph/Source/OptimusEditor/Private/OptimusDeformerAssetActions.h
+++ b/Engine/Plugins/Animation/DeformerGraph/Source/OptimusEditor/Private/OptimusDeformerAssetActions.h
@@ -11,43 +11,23 @@ class FOptimusDeformerAssetActions
 	: public FAssetTypeActions_Base
 {
 public:
+	FOptimusDeformerAssetActions(EAssetTypeCategories::Type InAssetCategoryBit = EAssetTypeCategories::Animation);
+	
+protected:
 	// IAssetTypeActions overrides
-	FText GetName() const override 
-	{ 
-		return NSLOCTEXT("AssetTypeActions", "OptimusDeformerActions", "Deformer Graph"); 
-	}
-
-	FColor GetTypeColor() const override 
-	{ 
-		return FColor::Blue;
-	}
-
-	UClass* GetSupportedClass() const override 
-	{ 
-		return UOptimusDeformer::StaticClass(); 
-	}
-
+	FText GetName() const override;
+	FColor GetTypeColor() const override;
+	UClass* GetSupportedClass() const override;
 	void OpenAssetEditor(
 		const TArray<UObject*>& InObjects, 
 		TSharedPtr<class IToolkitHost> EditWithinLevelEditor = TSharedPtr<IToolkitHost>()
 	) override;
-
-	uint32 GetCategories() override 
-	{ 
-		return EAssetTypeCategories::Misc;
-	}
-
-	const TArray<FText>& GetSubMenus() const override
-	{
-		static const TArray<FText> SubMenus
-		{
-			NSLOCTEXT("FOptimusDeformerAssetActions", "AnimDeformersSubMenu", "Deformers")
-		};
-		return SubMenus;
-	}
-
+	uint32 GetCategories() override;
+	const TArray<FText>& GetSubMenus() const override;
 	TSharedPtr<SWidget> GetThumbnailOverlay(
 		const FAssetData& AssetData
 	) const override;
 
+private:
+	EAssetTypeCategories::Type AssetCategoryBit;
 };
