Optimus: Set the IndexBuffer to be Triangle x 3 sizing on the skinned mesh DI. Fix back-comp issues with loading data domains.

#rnx

diff --git a/Engine/Plugins/Animation/DeformerGraph/Source/OptimusCore/Private/DataInterfaces/OptimusDataInterfaceDebugDraw.cpp b/Engine/Plugins/Animation/DeformerGraph/Source/OptimusCore/Private/DataInterfaces/OptimusDataInterfaceDebugDraw.cpp
--- a/Engine/Plugins/Animation/DeformerGraph/Source/OptimusCore/Private/DataInterfaces/OptimusDataInterfaceDebugDraw.cpp
+++ b/Engine/Plugins/Animation/DeformerGraph/Source/OptimusCore/Private/DataInterfaces/OptimusDataInterfaceDebugDraw.cpp
@@ -2,7 +2,8 @@
 
 #include "OptimusDataInterfaceDebugDraw.h"
 
-#include "Components/PrimitiveComponent.h"
+// #include "Components/PrimitiveComponent.h"
+#include "Components/SkinnedMeshComponent.h"
 #include "ComputeFramework/ShaderParamTypeDefinition.h"
 #include "OptimusDataDomain.h"
 #include "OptimusDataTypeRegistry.h"
@@ -30,7 +31,8 @@ TArray<FOptimusCDIPinDefinition> UOptimusDebugDrawDataInterface::GetPinDefinitio
 
 TSubclassOf<UActorComponent> UOptimusDebugDrawDataInterface::GetRequiredComponentClass() const
 {
-	return UPrimitiveComponent::StaticClass();
+	// return UPrimitiveComponent::StaticClass();
+	return USkinnedMeshComponent::StaticClass();
 }
 
 void UOptimusDebugDrawDataInterface::RegisterTypes() 
diff --git a/Engine/Plugins/Animation/DeformerGraph/Source/OptimusCore/Private/DataInterfaces/OptimusDataInterfaceSkinnedMesh.cpp b/Engine/Plugins/Animation/DeformerGraph/Source/OptimusCore/Private/DataInterfaces/OptimusDataInterfaceSkinnedMesh.cpp
--- a/Engine/Plugins/Animation/DeformerGraph/Source/OptimusCore/Private/DataInterfaces/OptimusDataInterfaceSkinnedMesh.cpp
+++ b/Engine/Plugins/Animation/DeformerGraph/Source/OptimusCore/Private/DataInterfaces/OptimusDataInterfaceSkinnedMesh.cpp
@@ -2,8 +2,10 @@
 
 #include "OptimusDataInterfaceSkinnedMesh.h"
 
-#include "ComputeFramework/ShaderParamTypeDefinition.h"
 #include "OptimusDataDomain.h"
+#include "ComponentSources/OptimusSkinnedMeshComponentSource.h"
+
+#include "ComputeFramework/ShaderParamTypeDefinition.h"
 #include "Rendering/SkeletalMeshLODRenderData.h"
 #include "Rendering/SkeletalMeshRenderData.h"
 #include "ShaderParameterMetadataBuilder.h"
@@ -17,16 +19,19 @@ FString UOptimusSkinnedMeshDataInterface::GetDisplayName() const
 
 TArray<FOptimusCDIPinDefinition> UOptimusSkinnedMeshDataInterface::GetPinDefinitions() const
 {
+	FName Vertex(UOptimusSkinnedMeshComponentSource::Domains::Vertex);
+	FName Triangle(UOptimusSkinnedMeshComponentSource::Domains::Triangle);
+	
 	TArray<FOptimusCDIPinDefinition> Defs;
 	Defs.Add({"NumVertices", "ReadNumVertices"});
-	Defs.Add({"Position", "ReadPosition", Optimus::DomainName::Vertex, "ReadNumVertices"});
-	Defs.Add({"TangentX", "ReadTangentX", Optimus::DomainName::Vertex, "ReadNumVertices"});
-	Defs.Add({"TangentZ", "ReadTangentZ", Optimus::DomainName::Vertex, "ReadNumVertices"});
+	Defs.Add({"Position", "ReadPosition", Vertex, "ReadNumVertices"});
+	Defs.Add({"TangentX", "ReadTangentX", Vertex, "ReadNumVertices"});
+	Defs.Add({"TangentZ", "ReadTangentZ", Vertex, "ReadNumVertices"});
 	Defs.Add({"NumUVChannels", "ReadNumUVChannels"});
-	Defs.Add({"UV", "ReadUV", {{Optimus::DomainName::Vertex, "ReadNumVertices"}, {Optimus::DomainName::UVChannel, "ReadNumUVChannels"}}});
-	Defs.Add({"Color", "ReadColor", Optimus::DomainName::Vertex, "ReadColor" });
+	Defs.Add({"UV", "ReadUV", {{Vertex, "ReadNumVertices"}, {Optimus::DomainName::UVChannel, "ReadNumUVChannels"}}});
+	Defs.Add({"Color", "ReadColor", Vertex, "ReadColor" });
 	Defs.Add({"NumTriangles", "ReadNumTriangles" });
-	Defs.Add({"IndexBuffer", "ReadIndexBuffer", Optimus::DomainName::Triangle, "ReadNumTriangles"});
+	Defs.Add({"IndexBuffer", "ReadIndexBuffer", Triangle, 3, "ReadNumTriangles"});
 	return Defs;
 }
 
diff --git a/Engine/Plugins/Animation/DeformerGraph/Source/OptimusCore/Private/Nodes/OptimusNode_DataInterface.cpp b/Engine/Plugins/Animation/DeformerGraph/Source/OptimusCore/Private/Nodes/OptimusNode_DataInterface.cpp
--- a/Engine/Plugins/Animation/DeformerGraph/Source/OptimusCore/Private/Nodes/OptimusNode_DataInterface.cpp
+++ b/Engine/Plugins/Animation/DeformerGraph/Source/OptimusCore/Private/Nodes/OptimusNode_DataInterface.cpp
@@ -269,7 +269,7 @@ void UOptimusNode_DataInterface::CreatePinFromDefinition(
 	// If there's no count function, then we have a value pin. The data function should
 	// have a return parameter but no input parameters. The value function only exists in 
 	// the read function map and so can only be an output pin.
-	if (InDefinition.Contexts.IsEmpty())
+	if (InDefinition.DataDimensions.IsEmpty())
 	{
 		if (!InReadFunctionMap.Contains(InDefinition.DataFunctionName))
 		{
@@ -301,7 +301,7 @@ void UOptimusNode_DataInterface::CreatePinFromDefinition(
 	else if (!InDefinition.DataFunctionName.IsEmpty())
 	{
 		// The count function is always in the read function list.
-		for (const FOptimusCDIPinDefinition::FContextInfo& ContextInfo: InDefinition.Contexts)
+		for (const FOptimusCDIPinDefinition::FDimensionInfo& ContextInfo: InDefinition.DataDimensions)
 		{
 			if (!InReadFunctionMap.Contains(ContextInfo.CountFunctionName))
 			{
@@ -320,7 +320,7 @@ void UOptimusNode_DataInterface::CreatePinFromDefinition(
 			const FShaderFunctionDefinition* FuncDef = InReadFunctionMap[InDefinition.DataFunctionName];
 
 			// FIXME: Ensure it takes a scalar uint/int as input index.
-			if (!FuncDef->bHasReturnType || FuncDef->ParamTypes.Num() != (1 + InDefinition.Contexts.Num()))
+			if (!FuncDef->bHasReturnType || FuncDef->ParamTypes.Num() != (1 + InDefinition.DataDimensions.Num()))
 			{
 				UE_LOG(LogOptimusCore, Error, TEXT("Data read function %s given for pin %s in %s is not properly declared."),
 					*InDefinition.DataFunctionName, *InDefinition.PinName.ToString(), *DataInterfaceClass->GetName());
@@ -337,7 +337,7 @@ void UOptimusNode_DataInterface::CreatePinFromDefinition(
 			const FShaderFunctionDefinition* FuncDef = InWriteFunctionMap[InDefinition.DataFunctionName];
 
 			// FIXME: Ensure it takes a scalar uint/int as input index.
-			if (FuncDef->bHasReturnType || FuncDef->ParamTypes.Num() != (1 + InDefinition.Contexts.Num()))
+			if (FuncDef->bHasReturnType || FuncDef->ParamTypes.Num() != (1 + InDefinition.DataDimensions.Num()))
 			{
 				UE_LOG(LogOptimusCore, Error, TEXT("Data write function %s given for pin %s in %s is not properly declared."),
 					*InDefinition.DataFunctionName, *InDefinition.PinName.ToString(), *DataInterfaceClass->GetName());
@@ -364,12 +364,12 @@ void UOptimusNode_DataInterface::CreatePinFromDefinition(
 		}
 
 		TArray<FName> ContextNames;
-		for (const FOptimusCDIPinDefinition::FContextInfo& ContextInfo: InDefinition.Contexts)
+		for (const FOptimusCDIPinDefinition::FDimensionInfo& ContextInfo: InDefinition.DataDimensions)
 		{
 			ContextNames.Add(ContextInfo.ContextName);
 		}
 
-		const FOptimusDataDomain DataDomain{ContextNames};
+		const FOptimusDataDomain DataDomain{ContextNames, InDefinition.DomainMultiplier};
 		AddPinDirect(InDefinition.PinName, PinDirection, DataDomain, PinDataType);
 	}
 	else
diff --git a/Engine/Plugins/Animation/DeformerGraph/Source/OptimusCore/Private/OptimusComputeDataInterface.cpp b/Engine/Plugins/Animation/DeformerGraph/Source/OptimusCore/Private/OptimusComputeDataInterface.cpp
--- a/Engine/Plugins/Animation/DeformerGraph/Source/OptimusCore/Private/OptimusComputeDataInterface.cpp
+++ b/Engine/Plugins/Animation/DeformerGraph/Source/OptimusCore/Private/OptimusComputeDataInterface.cpp
@@ -19,10 +19,10 @@ TSet<TArray<FName>> UOptimusComputeDataInterface::GetUniqueNestedContexts() cons
 	TSet<TArray<FName>> UniqueContextNames;
 	for (const FOptimusCDIPinDefinition& PinDef: GetPinDefinitions())
 	{
-		if (!PinDef.Contexts.IsEmpty())
+		if (!PinDef.DataDimensions.IsEmpty())
 		{
 			TArray<FName> ContextNames;
-			for (const FOptimusCDIPinDefinition::FContextInfo& ContextInfo: PinDef.Contexts)
+			for (const FOptimusCDIPinDefinition::FDimensionInfo& ContextInfo: PinDef.DataDimensions)
 			{
 				ContextNames.Add(ContextInfo.ContextName);
 			}
diff --git a/Engine/Plugins/Animation/DeformerGraph/Source/OptimusCore/Private/OptimusDataDomain.cpp b/Engine/Plugins/Animation/DeformerGraph/Source/OptimusCore/Private/OptimusDataDomain.cpp
--- a/Engine/Plugins/Animation/DeformerGraph/Source/OptimusCore/Private/OptimusDataDomain.cpp
+++ b/Engine/Plugins/Animation/DeformerGraph/Source/OptimusCore/Private/OptimusDataDomain.cpp
@@ -3,6 +3,7 @@
 #include "OptimusDataDomain.h"
 
 #include "OptimusExpressionEvaluator.h"
+#include "OptimusObjectVersion.h"
 
 
 namespace Optimus::DomainName
@@ -73,7 +74,8 @@ bool FOptimusDataDomain::operator==(const FOptimusDataDomain& InOtherDomain) con
 
 void FOptimusDataDomain::PostSerialize(const FArchive& Ar)
 {
-	if (Ar.IsLoading())
+	if (Ar.IsLoading() &&
+		Ar.CustomVer(FOptimusObjectVersion::GUID) < FOptimusObjectVersion::DataDomainExpansion)
 	{
 		BackCompFixupLevels();
 	}
diff --git a/Engine/Plugins/Animation/DeformerGraph/Source/OptimusCore/Private/OptimusNodePin.cpp b/Engine/Plugins/Animation/DeformerGraph/Source/OptimusCore/Private/OptimusNodePin.cpp
--- a/Engine/Plugins/Animation/DeformerGraph/Source/OptimusCore/Private/OptimusNodePin.cpp
+++ b/Engine/Plugins/Animation/DeformerGraph/Source/OptimusCore/Private/OptimusNodePin.cpp
@@ -8,6 +8,7 @@
 #include "OptimusHelpers.h"
 #include "OptimusNode.h"
 #include "OptimusNodeGraph.h"
+#include "OptimusObjectVersion.h"
 
 #include "Containers/Queue.h"
 #include "Misc/DefaultValueHelper.h"
@@ -508,11 +509,11 @@ void UOptimusNodePin::PostLoad()
 {
 	UObject::PostLoad();
 
-	if (!DataType.TypeName.IsNone() && DataType.TypeObject == nullptr)
+	// If the storage was marked as a value, the domain should now be a singleton.
+	if (GetLinkerCustomVersion(FOptimusObjectVersion::GUID) < FOptimusObjectVersion::DataDomainExpansion)
 	{
-		// If the storage was marked as a value, the domain should now be a singleton.
-PRAGMA_DISABLE_DEPRECATION_WARNINGS
 		DataDomain.BackCompFixupLevels();
+PRAGMA_DISABLE_DEPRECATION_WARNINGS
 		if (StorageType_DEPRECATED == EOptimusNodePinStorageType::Value)
 		{
 			DataDomain = FOptimusDataDomain();
diff --git a/Engine/Plugins/Animation/DeformerGraph/Source/OptimusCore/Public/OptimusComputeDataInterface.h b/Engine/Plugins/Animation/DeformerGraph/Source/OptimusCore/Public/OptimusComputeDataInterface.h
--- a/Engine/Plugins/Animation/DeformerGraph/Source/OptimusCore/Public/OptimusComputeDataInterface.h
+++ b/Engine/Plugins/Animation/DeformerGraph/Source/OptimusCore/Public/OptimusComputeDataInterface.h
@@ -14,7 +14,7 @@ class UActorComponent;
 
 struct FOptimusCDIPinDefinition
 {
-	struct FContextInfo;
+	struct FDimensionInfo;
 	
 	// Singleton value read/write. The context name is implied as Optimus::ContextName::Singleton.
 	FOptimusCDIPinDefinition(
@@ -34,20 +34,33 @@ struct FOptimusCDIPinDefinition
 		) :
 		PinName(InPinName),
 		DataFunctionName(InDataFunctionName),
-		Contexts{{InContextName, InCountFunctionName}}
+		DataDimensions{{InContextName, InCountFunctionName}}
 	{ }
 
 	FOptimusCDIPinDefinition(
 		const FName InPinName,
 		const FString InDataFunctionName,
-		const std::initializer_list<FContextInfo> InContexts
+		const FName InContextName,
+		const int32 InMultiplier,
+		const FString InCountFunctionName
+		) :
+		PinName(InPinName),
+		DataFunctionName(InDataFunctionName),
+		DataDimensions{{InContextName, InCountFunctionName}},
+		DomainMultiplier(FMath::Max(1, InMultiplier))
+	{ }
+	
+	FOptimusCDIPinDefinition(
+		const FName InPinName,
+		const FString InDataFunctionName,
+		const std::initializer_list<FDimensionInfo> InContexts
 		) :
 		PinName(InPinName),
 		DataFunctionName(InDataFunctionName)
 	{
-		for (FContextInfo ContextInfo: InContexts)
+		for (FDimensionInfo ContextInfo: InContexts)
 		{
-			Contexts.Add(ContextInfo);
+			DataDimensions.Add(ContextInfo);
 		}
 	}
 
@@ -74,7 +87,7 @@ struct FOptimusCDIPinDefinition
 	//    0 <= BoneIndex < GetVertexBoneCount(VertexIndex);
 	FString DataFunctionName;
 
-	struct FContextInfo
+	struct FDimensionInfo
 	{
 		// The data context for a given context level. For pins to be connectable they need to
 		// have identical set of contexts, in order.
@@ -93,7 +106,10 @@ struct FOptimusCDIPinDefinition
 	};
 
 	// List of nested data contexts.
-	TArray<FContextInfo> Contexts;
+	TArray<FDimensionInfo> DataDimensions;
+
+	// For single-level domains, how many values per element of that dimension's range. 
+	int32 DomainMultiplier = 1;
 };
 
 
