Optimus: Removed a struct asset scan that is very slow.

diff --git a/Engine/Plugins/Animation/DeformerGraph/Source/OptimusCore/Private/Nodes/OptimusNode_AnimAttributeDataInterface.cpp b/Engine/Plugins/Animation/DeformerGraph/Source/OptimusCore/Private/Nodes/OptimusNode_AnimAttributeDataInterface.cpp
--- a/Engine/Plugins/Animation/DeformerGraph/Source/OptimusCore/Private/Nodes/OptimusNode_AnimAttributeDataInterface.cpp
+++ b/Engine/Plugins/Animation/DeformerGraph/Source/OptimusCore/Private/Nodes/OptimusNode_AnimAttributeDataInterface.cpp
@@ -74,14 +74,14 @@ void UOptimusNode_AnimAttributeDataInterface::PostEditChangeChainProperty(FPrope
 	{
 		if (PropertyChangedEvent.PropertyChain.GetTail()->GetValue()->GetFName() == GET_MEMBER_NAME_STRING_CHECKED(FOptimusAnimAttributeArray, InnerArray))
 		{
-			RefreshPins();
+			RefreshOutputPins();
 		}
 	}
 	else if(PropertyChangedEvent.ChangeType == EPropertyChangeType::ArrayClear)
 	{
 		if (PropertyChangedEvent.PropertyChain.GetTail()->GetValue()->GetFName() == GET_MEMBER_NAME_STRING_CHECKED(FOptimusAnimAttributeArray, InnerArray))
 		{
-			ClearPins();
+			ClearOutputPins();
 		}
 	}
 }
@@ -112,13 +112,25 @@ void UOptimusNode_AnimAttributeDataInterface::UpdatePinTypes()
 	// Let's try and figure out which pin got changed.
 	TArrayView<UOptimusNodePin* const> NodePins = GetPins();
 	
-	if (ensure(NumAttributes == NodePins.Num()))
+	const TArray<UOptimusNodePin*> OutputPins = NodePins.FilterByPredicate(
+		[](UOptimusNodePin* const InPin)
+		{
+			if (InPin->GetDirection() == EOptimusNodePinDirection::Output)
+			{
+				return true;
+			}
+			
+			return false;
+		}
+	);	
+	
+	if (ensure(NumAttributes == OutputPins.Num()))
 	{
-		for (int32 Index = 0; Index < NodePins.Num(); Index++)
+		for (int32 Index = 0; Index < OutputPins.Num(); Index++)
 		{
-			if (NodePins[Index]->GetDataType() != Interface->AttributeArray[Index].DataType.Resolve())
+			if (OutputPins[Index]->GetDataType() != Interface->AttributeArray[Index].DataType.Resolve())
 			{
-				SetPinDataType(NodePins[Index], Interface->AttributeArray[Index].DataType);
+				SetPinDataType(OutputPins[Index], Interface->AttributeArray[Index].DataType);
 			}
 		}
 	}
@@ -134,38 +146,56 @@ void UOptimusNode_AnimAttributeDataInterface::UpdatePinNames()
 	// Let's try and figure out which pin got changed.
 	TArrayView<UOptimusNodePin* const> NodePins = GetPins();
 
-	if (ensure(PinDefinitions.Num() == NodePins.Num()))
+	const TArray<UOptimusNodePin*> OutputPins = NodePins.FilterByPredicate(
+		[](UOptimusNodePin* const InPin)
+		{
+			if (InPin->GetDirection() == EOptimusNodePinDirection::Output)
+			{
+				return true;
+			}
+			
+			return false;
+		}
+	);
+	
+	if (ensure(PinDefinitions.Num() == OutputPins.Num()))
 	{
-		for (int32 Index = 0; Index < NodePins.Num(); Index++)
+		for (int32 Index = 0; Index < OutputPins.Num(); Index++)
 		{
-			if (NodePins[Index]->GetFName() != PinDefinitions[Index].PinName)
+			if (OutputPins[Index]->GetFName() != PinDefinitions[Index].PinName)
 			{
-				SetPinName(NodePins[Index], PinDefinitions[Index].PinName);
+				SetPinName(OutputPins[Index], PinDefinitions[Index].PinName);
 			}
 		}
 	}
 }
 
 
-void UOptimusNode_AnimAttributeDataInterface::ClearPins()
+void UOptimusNode_AnimAttributeDataInterface::ClearOutputPins()
 {
 	for (UOptimusNodePin* Pin : GetPins())
 	{
-		RemovePin(Pin);
+		if (Pin->GetDirection() == EOptimusNodePinDirection::Output)
+		{
+			RemovePin(Pin);
+		}
 	}
 }
 
-void UOptimusNode_AnimAttributeDataInterface::RefreshPins()
+void UOptimusNode_AnimAttributeDataInterface::RefreshOutputPins()
 {
 	// Save the links and readd them later when new pins are created
 	TMap<FName, TArray<UOptimusNodePin*>> ConnectedPinsMap;
 
 	for (const UOptimusNodePin* Pin : GetPins())
 	{
-		ConnectedPinsMap.Add(Pin->GetFName()) = Pin->GetConnectedPins();
+		if (Pin->GetDirection() == EOptimusNodePinDirection::Output)
+		{
+			ConnectedPinsMap.Add(Pin->GetFName()) = Pin->GetConnectedPins();
+		}
 	}	
 
-	ClearPins();
+	ClearOutputPins();
 
 	UOptimusAnimAttributeDataInterface* Interface = Cast<UOptimusAnimAttributeDataInterface>(DataInterfaceData);
 	for (const FOptimusAnimAttributeDescription& Attribute : Interface->AttributeArray)
@@ -175,12 +205,15 @@ void UOptimusNode_AnimAttributeDataInterface::RefreshPins()
 
 	for (UOptimusNodePin* AddedPin : GetPins())
 	{
-		if (TArray<UOptimusNodePin*>* ConnectedPins = ConnectedPinsMap.Find(AddedPin->GetFName()))
+		if (AddedPin->GetDirection() == EOptimusNodePinDirection::Output)
 		{
-			for (UOptimusNodePin* ConnectedPin : *ConnectedPins)
+			if (TArray<UOptimusNodePin*>* ConnectedPins = ConnectedPinsMap.Find(AddedPin->GetFName()))
 			{
-				GetOwningGraph()->AddLink(AddedPin, ConnectedPin);
-			}
+				for (UOptimusNodePin* ConnectedPin : *ConnectedPins)
+				{
+					GetOwningGraph()->AddLink(AddedPin, ConnectedPin);
+				}
+			}	
 		}
 	}
 }
diff --git a/Engine/Plugins/Animation/DeformerGraph/Source/OptimusCore/Private/Nodes/OptimusNode_AnimAttributeDataInterface.h b/Engine/Plugins/Animation/DeformerGraph/Source/OptimusCore/Private/Nodes/OptimusNode_AnimAttributeDataInterface.h
--- a/Engine/Plugins/Animation/DeformerGraph/Source/OptimusCore/Private/Nodes/OptimusNode_AnimAttributeDataInterface.h
+++ b/Engine/Plugins/Animation/DeformerGraph/Source/OptimusCore/Private/Nodes/OptimusNode_AnimAttributeDataInterface.h
@@ -30,7 +30,7 @@ private:
 	void UpdatePinTypes();
 	void UpdatePinNames();
 
-	void ClearPins();
+	void ClearOutputPins();
 
-	void RefreshPins();
+	void RefreshOutputPins();
 };
diff --git a/Engine/Plugins/Animation/DeformerGraph/Source/OptimusCore/Private/OptimusDataTypeRegistry.cpp b/Engine/Plugins/Animation/DeformerGraph/Source/OptimusCore/Private/OptimusDataTypeRegistry.cpp
--- a/Engine/Plugins/Animation/DeformerGraph/Source/OptimusCore/Private/OptimusDataTypeRegistry.cpp
+++ b/Engine/Plugins/Animation/DeformerGraph/Source/OptimusCore/Private/OptimusDataTypeRegistry.cpp
@@ -799,9 +799,12 @@ bool FOptimusDataTypeRegistry::RegisterStructType(UScriptStruct* InStructType)
 void FOptimusDataTypeRegistry::RefreshStructType(UUserDefinedStruct* InStructType)
 {
 	FName TypeName = Optimus::GetTypeName(InStructType);
-	UnregisterType(TypeName);
-	RegisterStructType(InStructType);
-	OnDataTypeChanged.Broadcast(TypeName);
+	if (RegisteredTypes.Contains(TypeName))
+	{
+		UnregisterType(TypeName);
+		RegisterStructType(InStructType);
+		OnDataTypeChanged.Broadcast(TypeName);
+	}
 }
 
 bool FOptimusDataTypeRegistry::RegisterType(
@@ -1428,32 +1431,3 @@ void FOptimusDataTypeRegistry::UnregisterType(FName InTypeName)
 	RegisteredTypes.Remove(InTypeName);
 	RegistrationOrder.Remove(InTypeName);
 }
-
-
-void FOptimusDataTypeRegistry::RefreshRegistry()
-{
-	// Checks for new user defined structs that have been added
-	
-	TArray<FAssetData> StructAssets;
-
-	const FAssetRegistryModule& AssetRegistryModule = FModuleManager::LoadModuleChecked<FAssetRegistryModule>(TEXT("AssetRegistry"));
-	AssetRegistryModule.Get().GetAssetsByClass(UUserDefinedStruct::StaticClass()->GetClassPathName(), StructAssets);
-
-	StructAssets.Sort([](const FAssetData& A, const FAssetData& B) { return A.AssetName.LexicalLess(B.AssetName); });
-
-	for (const FAssetData& StructAsset : StructAssets)
-	{
-		UObject* Struct = StructAsset.GetAsset();
-		UScriptStruct* ScriptStruct = Cast<UScriptStruct>(Struct);
-		
-		if (ensure(ScriptStruct))
-		{
-			FOptimusDataTypeHandle DataType = FindType(Optimus::GetTypeName(ScriptStruct));
-
-			if (!DataType.IsValid())
-			{
-				RegisterStructType(ScriptStruct);
-			}
-		}
-	}
-}
diff --git a/Engine/Plugins/Animation/DeformerGraph/Source/OptimusCore/Public/OptimusDataTypeRegistry.h b/Engine/Plugins/Animation/DeformerGraph/Source/OptimusCore/Public/OptimusDataTypeRegistry.h
--- a/Engine/Plugins/Animation/DeformerGraph/Source/OptimusCore/Public/OptimusDataTypeRegistry.h
+++ b/Engine/Plugins/Animation/DeformerGraph/Source/OptimusCore/Public/OptimusDataTypeRegistry.h
@@ -126,9 +126,6 @@ public:
 		UUserDefinedStruct *InStructType
 		);
 
-	// Check for new user defined structs and register them
-	OPTIMUSCORE_API void RefreshRegistry();
-
 	// Unregister a type
 	OPTIMUSCORE_API void UnregisterType(FName InTypeName);
 
diff --git a/Engine/Plugins/Animation/DeformerGraph/Source/OptimusEditor/Private/Widgets/SOptimusDataTypeSelector.cpp b/Engine/Plugins/Animation/DeformerGraph/Source/OptimusEditor/Private/Widgets/SOptimusDataTypeSelector.cpp
--- a/Engine/Plugins/Animation/DeformerGraph/Source/OptimusEditor/Private/Widgets/SOptimusDataTypeSelector.cpp
+++ b/Engine/Plugins/Animation/DeformerGraph/Source/OptimusEditor/Private/Widgets/SOptimusDataTypeSelector.cpp
@@ -173,16 +173,17 @@ FText SOptimusDataTypeSelector::GetTypeTooltip() const
 
 TSharedRef<SWidget> SOptimusDataTypeSelector::GetMenuContent()
 {
-	FOptimusDataTypeRegistry::Get().RefreshRegistry();
-
 	AllDataTypeItems.Reset();
 	FOptimusDataTypeHandle SelectedItem;
 	for (FOptimusDataTypeHandle DataType : FOptimusDataTypeRegistry::Get().GetAllTypes())
 	{
-		AllDataTypeItems.Add(DataType);
-		if (DataType == CurrentDataType.Get())
+		if ( !UsageMask.IsSet() || UsageMask.Get() == EOptimusDataTypeUsageFlags::None || EnumHasAnyFlags(DataType->UsageFlags, UsageMask.Get()))
 		{
-			SelectedItem = AllDataTypeItems.Last();
+			AllDataTypeItems.Add(DataType);
+			if (DataType == CurrentDataType.Get())
+			{
+				SelectedItem = AllDataTypeItems.Last();
+			}
 		}
 	}
 
