Optimus: Use full asset path for AllowClasses. Fixes a warning during cook.

#rnx

diff --git a/Engine/Plugins/Animation/DeformerGraph/Source/OptimusSettings/Public/OptimusSettings.h b/Engine/Plugins/Animation/DeformerGraph/Source/OptimusSettings/Public/OptimusSettings.h
--- a/Engine/Plugins/Animation/DeformerGraph/Source/OptimusSettings/Public/OptimusSettings.h
+++ b/Engine/Plugins/Animation/DeformerGraph/Source/OptimusSettings/Public/OptimusSettings.h
@@ -29,11 +29,11 @@ class OPTIMUSSETTINGS_API UOptimusSettings : public UDeveloperSettings
 	EOptimusDefaultDeformerMode DefaultMode = EOptimusDefaultDeformerMode::Never;
 
 	/** A default deformer that will be used on a skinned mesh if no other deformer has been set. */
-	UPROPERTY(config, EditAnywhere, Category = DeformerGraph, meta = (AllowedClasses = "OptimusDeformer", EditCondition = "DefaultMode != EOptimusDefaultDeformerMode::Never"))
+	UPROPERTY(config, EditAnywhere, Category = DeformerGraph, meta = (AllowedClasses = "/Script/OptimusCore.OptimusDeformer", EditCondition = "DefaultMode != EOptimusDefaultDeformerMode::Never"))
 	TSoftObjectPtr<UMeshDeformer> DefaultDeformer;
 
 	/** A default deformer that will be used on a skinned mesh if no other deformer has been set, and if the mesh has requested to recompute tangets. */
-	UPROPERTY(config, EditAnywhere, Category = DeformerGraph, meta = (AllowedClasses = "OptimusDeformer", EditCondition = "DefaultMode != EOptimusDefaultDeformerMode::Never"))
+	UPROPERTY(config, EditAnywhere, Category = DeformerGraph, meta = (AllowedClasses = "/Script/OptimusCore.OptimusDeformer", EditCondition = "DefaultMode != EOptimusDefaultDeformerMode::Never"))
 	TSoftObjectPtr<UMeshDeformer> DefaultRecomputeTangentDeformer;
 
 #if WITH_EDITOR
