Optimus: Add official icons for the plugin, deformer graph asset type and deformer graphs.

#rnx

diff --git a/Engine/Plugins/Animation/DeformerGraph/Source/OptimusEditor/Private/OptimusDeformerAssetActions.cpp b/Engine/Plugins/Animation/DeformerGraph/Source/OptimusEditor/Private/OptimusDeformerAssetActions.cpp
--- a/Engine/Plugins/Animation/DeformerGraph/Source/OptimusEditor/Private/OptimusDeformerAssetActions.cpp
+++ b/Engine/Plugins/Animation/DeformerGraph/Source/OptimusEditor/Private/OptimusDeformerAssetActions.cpp
@@ -76,7 +76,7 @@ const TArray<FText>& FOptimusDeformerAssetActions::GetSubMenus() const
 
 TSharedPtr<SWidget> FOptimusDeformerAssetActions::GetThumbnailOverlay(const FAssetData& AssetData) const
 {
-	const FSlateBrush* Icon = FSlateIconFinder::FindIconBrushForClass(UOptimusDeformer::StaticClass());
+	const FSlateBrush* Icon = FSlateIconFinder::FindIconBrushForClass(UOptimusDeformer::StaticClass());;
 
 	return SNew(SBorder)
 		.BorderImage(FAppStyle::GetNoBrush())
diff --git a/Engine/Plugins/Animation/DeformerGraph/Source/OptimusEditor/Private/OptimusEditorGraph.cpp b/Engine/Plugins/Animation/DeformerGraph/Source/OptimusEditor/Private/OptimusEditorGraph.cpp
--- a/Engine/Plugins/Animation/DeformerGraph/Source/OptimusEditor/Private/OptimusEditorGraph.cpp
+++ b/Engine/Plugins/Animation/DeformerGraph/Source/OptimusEditor/Private/OptimusEditorGraph.cpp
@@ -12,6 +12,8 @@
 #include "EdGraph/EdGraphPin.h"
 #include "Styling/AppStyle.h"
 #include "GraphEditAction.h"
+#include "OptimusEditorStyle.h"
+
 
 UOptimusEditorGraph::UOptimusEditorGraph()
 {
@@ -104,8 +106,19 @@ void UOptimusEditorGraph::RefreshVisualNode(UOptimusEditorGraphNode* InGraphNode
 
 const FSlateBrush* UOptimusEditorGraph::GetGraphTypeIcon(UOptimusNodeGraph* InModelGraph)
 {
-	// FIXME: Need icon types.
-	return FAppStyle::GetBrush(TEXT("GraphEditor.Animation_24x"));
+	switch(InModelGraph->GetGraphType())
+	{
+	case EOptimusNodeGraphType::Setup:
+			return FOptimusEditorStyle::Get().GetBrush(TEXT("GraphType.Setup"));
+	case EOptimusNodeGraphType::Update:
+			return FOptimusEditorStyle::Get().GetBrush(TEXT("GraphType.Trigger"));
+	case EOptimusNodeGraphType::ExternalTrigger:
+			return FOptimusEditorStyle::Get().GetBrush(TEXT("GraphType.Update"));
+	default:
+		checkNoEntry();
+		break;
+	}
+	return nullptr;
 }
 
 
diff --git a/Engine/Plugins/Animation/DeformerGraph/Source/OptimusEditor/Private/OptimusEditorStyle.cpp b/Engine/Plugins/Animation/DeformerGraph/Source/OptimusEditor/Private/OptimusEditorStyle.cpp
--- a/Engine/Plugins/Animation/DeformerGraph/Source/OptimusEditor/Private/OptimusEditorStyle.cpp
+++ b/Engine/Plugins/Animation/DeformerGraph/Source/OptimusEditor/Private/OptimusEditorStyle.cpp
@@ -13,11 +13,18 @@ FOptimusEditorStyle::FOptimusEditorStyle() :
 {
 	static const FVector2D IconSize10x10(10.0f, 10.0f);
 	static const FVector2D IconSize16x12(16.0f, 12.0f);
+	static const FVector2D IconSize16x16(16.0f, 16.0f);
+	static const FVector2D IconSize20x20(20.0f, 20.0f);
 
 	static const FSlateColor DefaultForeground(FLinearColor(0.72f, 0.72f, 0.72f, 1.f));
 
 	SetContentRoot(FPaths::EnginePluginsDir() / TEXT("Animation/DeformerGraph/Resources"));
 
+	// Asset icons
+	{
+		Set("ClassIcon.OptimusDeformer", new IMAGE_BRUSH_SVG("Icons/DeformerGraph_16", IconSize16x16, DefaultForeground));
+	}
+
 	// Text editor styles
 	{
 		const FSlateFontInfo Consolas10  = FCoreStyle::GetDefaultFontStyle("Mono", 9);
@@ -65,6 +72,13 @@ FOptimusEditorStyle::FOptimusEditorStyle() :
 
 	}
 
+	// Graph type icons
+	{
+		Set("GraphType.Setup", new IMAGE_BRUSH_SVG("Icons/Graph_Setup", IconSize20x20, DefaultForeground));
+		Set("GraphType.Trigger", new IMAGE_BRUSH_SVG("Icons/Graph_Trigger", IconSize20x20, DefaultForeground));
+		Set("GraphType.Update", new IMAGE_BRUSH_SVG("Icons/Graph_Update", IconSize20x20, DefaultForeground));
+	}
+
 	// Graph styles
 	{
 		FTextBlockStyle NormalText = FCoreStyle::Get().GetWidgetStyle<FTextBlockStyle>("NormalText");
