Optimus: Include icons for the component bindings based on the component icons.

#rnx

diff --git a/Engine/Plugins/Animation/DeformerGraph/Source/OptimusEditor/Private/OptimusEditorGraph.cpp b/Engine/Plugins/Animation/DeformerGraph/Source/OptimusEditor/Private/OptimusEditorGraph.cpp
--- a/Engine/Plugins/Animation/DeformerGraph/Source/OptimusEditor/Private/OptimusEditorGraph.cpp
+++ b/Engine/Plugins/Animation/DeformerGraph/Source/OptimusEditor/Private/OptimusEditorGraph.cpp
@@ -104,7 +104,9 @@ void UOptimusEditorGraph::RefreshVisualNode(UOptimusEditorGraphNode* InGraphNode
 }
 
 
-const FSlateBrush* UOptimusEditorGraph::GetGraphTypeIcon(UOptimusNodeGraph* InModelGraph)
+const FSlateBrush* UOptimusEditorGraph::GetGraphTypeIcon(
+	const UOptimusNodeGraph* InModelGraph
+	)
 {
 	switch(InModelGraph->GetGraphType())
 	{
diff --git a/Engine/Plugins/Animation/DeformerGraph/Source/OptimusEditor/Private/OptimusEditorGraph.h b/Engine/Plugins/Animation/DeformerGraph/Source/OptimusEditor/Private/OptimusEditorGraph.h
--- a/Engine/Plugins/Animation/DeformerGraph/Source/OptimusEditor/Private/OptimusEditorGraph.h
+++ b/Engine/Plugins/Animation/DeformerGraph/Source/OptimusEditor/Private/OptimusEditorGraph.h
@@ -36,7 +36,7 @@ public:
 	void RefreshVisualNode(UOptimusEditorGraphNode *InGraphNode);
 
 	///
-	static const FSlateBrush* GetGraphTypeIcon(UOptimusNodeGraph* InModelGraph);
+	static const FSlateBrush* GetGraphTypeIcon(const UOptimusNodeGraph* InModelGraph);
 
 protected:
 	friend class FOptimusEditor;
@@ -59,7 +59,7 @@ private:
 
 	UOptimusEditorGraphNode* AddGraphNodeFromModelNode(UOptimusNode* InModelNode);
 
-	UOptimusNodeGraph* NodeGraph;
+	TObjectPtr<UOptimusNodeGraph> NodeGraph;
 
 	TSet<UOptimusEditorGraphNode*> SelectedNodes;
 };
diff --git a/Engine/Plugins/Animation/DeformerGraph/Source/OptimusEditor/Private/Widgets/SOptimusEditorGraphExplorerItem.cpp b/Engine/Plugins/Animation/DeformerGraph/Source/OptimusEditor/Private/Widgets/SOptimusEditorGraphExplorerItem.cpp
--- a/Engine/Plugins/Animation/DeformerGraph/Source/OptimusEditor/Private/Widgets/SOptimusEditorGraphExplorerItem.cpp
+++ b/Engine/Plugins/Animation/DeformerGraph/Source/OptimusEditor/Private/Widgets/SOptimusEditorGraphExplorerItem.cpp
@@ -16,6 +16,7 @@
 #include "OptimusVariableDescription.h"
 #include "PacketHandler.h"
 #include "SOptimusDataTypeSelector.h"
+#include "Styling/SlateIconFinder.h"
 
 #include "Widgets/Images/SImage.h"
 #include "Widgets/DeclarativeSyntaxSupport.h"
@@ -222,6 +223,12 @@ void SOptimusEditorGraphExplorerItem::Construct(
 		[
 			SNew(SHorizontalBox)
 			+ SHorizontalBox::Slot()
+			.AutoWidth()
+			.Padding(0.0f, 2.0f)
+			[
+				CreateIconWidget(InCreateData, bIsReadOnlyCreate)
+			]
+			+ SHorizontalBox::Slot()
 			.FillWidth(0.6f)
 			.VAlign(VAlign_Center)
 			.Padding(/* horizontal */ 3.0f, /* vertical */ 0.0f)
@@ -252,8 +259,8 @@ TSharedRef<SWidget> SOptimusEditorGraphExplorerItem::CreateIconWidget(
 		IOptimusPathResolver* PathResolver = Editor->GetDeformerInterface<IOptimusPathResolver>();
 		if (Action->GetTypeId() == FOptimusSchemaAction_Graph::StaticGetTypeId())
 		{
-			FOptimusSchemaAction_Graph* GraphAction = static_cast<FOptimusSchemaAction_Graph*>(Action.Get());
-			UOptimusNodeGraph* NodeGraph = PathResolver->ResolveGraphPath(GraphAction->GraphPath);
+			const FOptimusSchemaAction_Graph* GraphAction = static_cast<FOptimusSchemaAction_Graph*>(Action.Get());
+			const UOptimusNodeGraph* NodeGraph = PathResolver->ResolveGraphPath(GraphAction->GraphPath);
 			if (ensure(NodeGraph))
 			{
 				IconWidget = SNew(SImage)
@@ -262,7 +269,7 @@ TSharedRef<SWidget> SOptimusEditorGraphExplorerItem::CreateIconWidget(
 		}
 		else if (Action->GetTypeId() == FOptimusSchemaAction_Resource::StaticGetTypeId())
 		{
-			FOptimusSchemaAction_Resource* ResourceAction = static_cast<FOptimusSchemaAction_Resource*>(Action.Get());
+			const FOptimusSchemaAction_Resource* ResourceAction = static_cast<FOptimusSchemaAction_Resource*>(Action.Get());
 			UOptimusResourceDescription* Resource = PathResolver->ResolveResource(ResourceAction->ResourceName);
 			if (ensure(Resource))
 			{
@@ -272,13 +279,24 @@ TSharedRef<SWidget> SOptimusEditorGraphExplorerItem::CreateIconWidget(
 		}
 		else if (Action->GetTypeId() == FOptimusSchemaAction_Variable::StaticGetTypeId())
 		{
-			FOptimusSchemaAction_Variable* VariableAction = static_cast<FOptimusSchemaAction_Variable*>(Action.Get());
+			const FOptimusSchemaAction_Variable* VariableAction = static_cast<FOptimusSchemaAction_Variable*>(Action.Get());
 			UOptimusVariableDescription* Variable = PathResolver->ResolveVariable(VariableAction->VariableName);
 			if (ensure(Variable))
 			{
 				IconWidget = SNew(SVariableDataTypeSelectorHelper, Variable, bInIsReadOnly);
 			}
 		}
+		else if (Action->GetTypeId() == FOptimusSchemaAction_Binding::StaticGetTypeId())
+		{
+			const FOptimusSchemaAction_Binding* BindingAction = static_cast<FOptimusSchemaAction_Binding*>(Action.Get());
+			const UOptimusComponentSourceBinding* Binding = PathResolver->ResolveComponentBinding(BindingAction->BindingName);
+			if (ensure(Binding))
+			{
+				IconWidget = SNew(SImage)
+					.DesiredSizeOverride(FVector2D(16.f, 16.f))
+					.Image(FSlateIconFinder::FindIconBrushForClass(Binding->GetComponentSource()->GetComponentClass(), TEXT("SCS.Component")));
+			}
+		}
 	}
 
 	if (IconWidget.IsValid())
