OptimusEditor: Fix a crash that can happen if a user defined struct is in a details view when it is structurally changed

the order should be enforced: scriptstruct change -> struct instance change -> editor refresh

previously it was:
scriptstruct change -> editor refresh
scriptscript change -> struct instance change

diff --git a/Engine/Plugins/Animation/DeformerGraph/Source/OptimusCore/Private/OptimusDeformer.cpp b/Engine/Plugins/Animation/DeformerGraph/Source/OptimusCore/Private/OptimusDeformer.cpp
--- a/Engine/Plugins/Animation/DeformerGraph/Source/OptimusCore/Private/OptimusDeformer.cpp
+++ b/Engine/Plugins/Animation/DeformerGraph/Source/OptimusCore/Private/OptimusDeformer.cpp
@@ -1806,6 +1806,9 @@ void UOptimusDeformer::OnDataTypeChanged(FName InTypeName)
 	}
 
 	//TODO: Recreate variables/Resources that uses this type
+
+	// Once we updated the deformer instance, we need to make sure the editor is aware as well
+	Notify(EOptimusGlobalNotifyType::DataTypeChanged, nullptr);
 }
 
 template<typename Allocator>
diff --git a/Engine/Plugins/Animation/DeformerGraph/Source/OptimusCore/Public/OptimusCoreNotify.h b/Engine/Plugins/Animation/DeformerGraph/Source/OptimusCore/Public/OptimusCoreNotify.h
--- a/Engine/Plugins/Animation/DeformerGraph/Source/OptimusCore/Public/OptimusCoreNotify.h
+++ b/Engine/Plugins/Animation/DeformerGraph/Source/OptimusCore/Public/OptimusCoreNotify.h
@@ -39,6 +39,8 @@ enum class EOptimusGlobalNotifyType
 									
 	NodeTypeAdded,					/// A new node type has been added (Subject == UClass)
 	NodeTypeRemoved,				/// A node type is about to be removed (Subject == UClass)
+
+	DataTypeChanged,				/// A user defined struct type is changed (Subject == nullptr)
 };
 
 // A delegate for subscribing / reacting to Optimus global notifications.
diff --git a/Engine/Plugins/Animation/DeformerGraph/Source/OptimusEditor/Private/OptimusEditor.cpp b/Engine/Plugins/Animation/DeformerGraph/Source/OptimusEditor/Private/OptimusEditor.cpp
--- a/Engine/Plugins/Animation/DeformerGraph/Source/OptimusEditor/Private/OptimusEditor.cpp
+++ b/Engine/Plugins/Animation/DeformerGraph/Source/OptimusEditor/Private/OptimusEditor.cpp
@@ -57,8 +57,6 @@ FOptimusEditor::FOptimusEditor()
 
 FOptimusEditor::~FOptimusEditor()
 {
-	FOptimusDataTypeRegistry::Get().GetOnDataTypeChanged().RemoveAll(this);
-	
 	if (DeformerObject)
 	{
 		DeformerObject->GetCompileBeginDelegate().RemoveAll(this);
@@ -177,8 +175,6 @@ void FOptimusEditor::Construct(
 	{
 		HandlePreviewMeshChanged(nullptr, PersonaToolkit->GetPreviewMesh());
 	}
-
-	FOptimusDataTypeRegistry::Get().GetOnDataTypeChanged().AddSP(this, &FOptimusEditor::OnDataTypeChanged);
 }
 
 
@@ -1088,7 +1084,10 @@ void FOptimusEditor::OnDeformerModified(
 	case EOptimusGlobalNotifyType::NodeTypeRemoved: 
 		RefreshEvent.Broadcast();
 		break;
-		
+
+	case EOptimusGlobalNotifyType::DataTypeChanged:
+		OnDataTypeChanged();
+		break;
 
 	case EOptimusGlobalNotifyType::GraphRemoved: 
 	{
@@ -1144,7 +1143,7 @@ void FOptimusEditor::OnFinishedChangingProperties(const FPropertyChangedEvent& P
 	}
 }
 
-void FOptimusEditor::OnDataTypeChanged(FName InTypeName)
+void FOptimusEditor::OnDataTypeChanged()
 {
 	if (PropertyDetailsWidget)
 	{
diff --git a/Engine/Plugins/Animation/DeformerGraph/Source/OptimusEditor/Private/OptimusEditor.h b/Engine/Plugins/Animation/DeformerGraph/Source/OptimusEditor/Private/OptimusEditor.h
--- a/Engine/Plugins/Animation/DeformerGraph/Source/OptimusEditor/Private/OptimusEditor.h
+++ b/Engine/Plugins/Animation/DeformerGraph/Source/OptimusEditor/Private/OptimusEditor.h
@@ -212,7 +212,7 @@ private:
 	// Called when the inspector has changed a value.
 	void OnFinishedChangingProperties(const FPropertyChangedEvent& PropertyChangedEvent);
 
-	void OnDataTypeChanged(FName InTypeName);
+	void OnDataTypeChanged();
 	
 private:
 	// Persona toolkit for the skelmesh preview
