Optimus: fix a few issues with compatibility checking regarding Animation Attribute Node support for user defined structs

diff --git a/Engine/Plugins/Animation/DeformerGraph/Source/OptimusCore/Private/OptimusCoreModule.cpp b/Engine/Plugins/Animation/DeformerGraph/Source/OptimusCore/Private/OptimusCoreModule.cpp
--- a/Engine/Plugins/Animation/DeformerGraph/Source/OptimusCore/Private/OptimusCoreModule.cpp
+++ b/Engine/Plugins/Animation/DeformerGraph/Source/OptimusCore/Private/OptimusCoreModule.cpp
@@ -23,13 +23,13 @@ void FOptimusCoreModule::StartupModule()
 
 	// Make sure all our types are known at startup.
 	FOptimusDataTypeRegistry::RegisterBuiltinTypes();
-	FOptimusDataTypeRegistry::RegisterAssetRegistryCallbacks();
+	FOptimusDataTypeRegistry::RegisterEngineCallbacks();
 	UOptimusComputeDataInterface::RegisterAllTypes();
 }
 
 void FOptimusCoreModule::ShutdownModule()
 {
-	FOptimusDataTypeRegistry::UnregisterAssetRegistryCallbacks();
+	FOptimusDataTypeRegistry::UnregisterEngineCallbacks();
 	FOptimusDataTypeRegistry::UnregisterAllTypes();
 }
 
diff --git a/Engine/Plugins/Animation/DeformerGraph/Source/OptimusCore/Private/OptimusDataTypeRegistry.cpp b/Engine/Plugins/Animation/DeformerGraph/Source/OptimusCore/Private/OptimusDataTypeRegistry.cpp
--- a/Engine/Plugins/Animation/DeformerGraph/Source/OptimusCore/Private/OptimusDataTypeRegistry.cpp
+++ b/Engine/Plugins/Animation/DeformerGraph/Source/OptimusCore/Private/OptimusDataTypeRegistry.cpp
@@ -7,6 +7,7 @@
 #include "OptimusComponentSource.h"
 
 #include "ShaderParameterMetadata.h"
+#include "Animation/AttributeTypes.h"
 #include "Animation/BuiltInAttributeTypes.h"
 
 #include "AssetRegistry/AssetRegistryModule.h"
@@ -14,8 +15,20 @@
 #include "Engine/Classes/Engine/UserDefinedStruct.h"
 
 #include "UObject/UnrealType.h"
-#include "UObject/UObjectIterator.h"
 
+static const TMap<FName, UScriptStruct*>& GetBuiltInAttributeTypes()
+{
+	static const TMap<FName, UScriptStruct*> BuiltInAttributeTypes =
+	{
+		{FFloatProperty::StaticClass()->GetFName(), FFloatAnimationAttribute::StaticStruct()},
+		{FIntProperty::StaticClass()->GetFName(), FIntegerAnimationAttribute::StaticStruct()},
+		{*TBaseStructure<FTransform>::Get()->GetStructCPPName(), FTransformAnimationAttribute::StaticStruct()},
+		{*TBaseStructure<FVector>::Get()->GetStructCPPName(), FVectorAnimationAttribute::StaticStruct()},
+		{*TBaseStructure<FQuat>::Get()->GetStructCPPName(), FQuaternionAnimationAttribute::StaticStruct()},
+	};
+
+	return BuiltInAttributeTypes;
+}
 
 static bool IsStructHashable(const UScriptStruct* InStructType)
 {
@@ -54,7 +67,7 @@ static bool ConvertPropertyValuePOD(
 	return false;
 }
 
-EOptimusDataTypeUsageFlags FOptimusDataTypeRegistry::CanRegisterStructType(UScriptStruct* InStruct)
+EOptimusDataTypeUsageFlags FOptimusDataTypeRegistry::GetStructTypeUsageFlag(UScriptStruct* InStruct)
 {
 	struct FTypeValidator 
     {
@@ -496,7 +509,7 @@ bool FOptimusDataTypeRegistry::RegisterStructType(UScriptStruct* InStructType)
 		return false;
 	}
 
-	EOptimusDataTypeUsageFlags UsageFlags = CanRegisterStructType(InStructType);
+	EOptimusDataTypeUsageFlags UsageFlags = GetStructTypeUsageFlag(InStructType);
 
 	// Disable the use of structs in variables and resources for now
 	UsageFlags &= ~EOptimusDataTypeUsageFlags::Variable;
@@ -1294,22 +1307,28 @@ void FOptimusDataTypeRegistry::UnregisterAllTypes()
 	Registry.RegisteredTypes.Reset();
 }
 
-void FOptimusDataTypeRegistry::RegisterAssetRegistryCallbacks()
+void FOptimusDataTypeRegistry::RegisterEngineCallbacks()
 {
 	if (FAssetRegistryModule* AssetRegistryModule = FModuleManager::GetModulePtr<FAssetRegistryModule>(TEXT("AssetRegistry")))
 	{
+		AssetRegistryModule->Get().OnFilesLoaded().AddRaw(&Get(), &FOptimusDataTypeRegistry::OnFilesLoaded);
 		AssetRegistryModule->Get().OnAssetRemoved().AddRaw(&Get(), &FOptimusDataTypeRegistry::OnAssetRemoved);
 		AssetRegistryModule->Get().OnAssetRenamed().AddRaw(&Get(), &FOptimusDataTypeRegistry::OnAssetRenamed);
 	}
+
+	UE::Anim::AttributeTypes::GetOnAttributeTypesChanged().AddRaw(&Get(), &FOptimusDataTypeRegistry::OnAnimationAttributeRegistryChanged);
 }
 
-void FOptimusDataTypeRegistry::UnregisterAssetRegistryCallbacks()
+void FOptimusDataTypeRegistry::UnregisterEngineCallbacks()
 {
 	if (FAssetRegistryModule* AssetRegistryModule = FModuleManager::GetModulePtr<FAssetRegistryModule>(TEXT("AssetRegistry")))
 	{
+		AssetRegistryModule->Get().OnFilesLoaded().RemoveAll(&Get());
 		AssetRegistryModule->Get().OnAssetRemoved().RemoveAll(&Get());
 		AssetRegistryModule->Get().OnAssetRenamed().RemoveAll(&Get());
 	}
+	
+	UE::Anim::AttributeTypes::GetOnAttributeTypesChanged().RemoveAll(&Get());
 }
 
 
@@ -1331,6 +1350,29 @@ FOptimusDataTypeRegistry::FOptimusDataTypeRegistry()
 {
 }
 
+void FOptimusDataTypeRegistry::OnFilesLoaded()
+{
+	TArray<TWeakObjectPtr<const UScriptStruct>> AttributeTypes = UE::Anim::AttributeTypes::GetRegisteredTypes();
+	
+	for (TWeakObjectPtr<const UScriptStruct> Type : AttributeTypes)
+	{
+		UScriptStruct* ScriptStruct = const_cast<UScriptStruct*>(Type.Get());
+
+		TArray<UScriptStruct*> BuiltInAttributeTypeArray;
+		GetBuiltInAttributeTypes().GenerateValueArray(BuiltInAttributeTypeArray);
+		
+		if (ScriptStruct && !BuiltInAttributeTypeArray.Contains(ScriptStruct))
+		{
+			FOptimusDataTypeHandle DataType = FindType(Optimus::GetTypeName(ScriptStruct));
+
+			if (!DataType.IsValid())
+			{
+				RegisterStructType(ScriptStruct);
+			}
+		}
+	}
+}
+
 void FOptimusDataTypeRegistry::OnAssetRemoved(const FAssetData& InAssetData)
 {
 	if (InAssetData.AssetClassPath == UUserDefinedStruct::StaticClass()->GetClassPathName())
@@ -1359,6 +1401,25 @@ void FOptimusDataTypeRegistry::OnAssetRenamed(const FAssetData& InAssetData, con
 	}
 }
 
+void FOptimusDataTypeRegistry::OnAnimationAttributeRegistryChanged(const UScriptStruct* InScriptStruct,
+	bool bIsAdded)
+{
+	if (ensure(InScriptStruct) && ensure(InScriptStruct->IsA<UUserDefinedStruct>()))
+	{
+		
+		UScriptStruct* ScriptStruct = const_cast<UScriptStruct*>(InScriptStruct);
+		FOptimusDataTypeHandle DataType = FindType(Optimus::GetTypeName(ScriptStruct));
+		
+		if (bIsAdded)
+		{
+			if (!DataType.IsValid())
+			{
+				RegisterStructType(ScriptStruct);
+			}
+		}
+	}
+}
+
 FOptimusDataTypeRegistry::PropertyValueConvertFuncT FOptimusDataTypeRegistry::FindPropertyValueConvertFunc(
 	FName InTypeName
 	) const
@@ -1388,20 +1449,7 @@ const TArray<FOptimusDataTypeRegistry::FArrayMetadata>* FOptimusDataTypeRegistry
 
 UScriptStruct* FOptimusDataTypeRegistry::FindAttributeType(FName InTypeName) const
 {
-	FName TransformTypeName(*FString::Printf(TEXT("F%s"), *TBaseStructure<FTransform>::Get()->GetName()));
-	FName VectorTypeName(*FString::Printf(TEXT("F%s"), *TBaseStructure<FVector>::Get()->GetName()));
-	FName QuatTypeName(*FString::Printf(TEXT("F%s"), *TBaseStructure<FQuat>::Get()->GetName()));
-
-	static TMap<FName, UScriptStruct*> BuiltInAttributeTypes =
-		{
-			{FFloatProperty::StaticClass()->GetFName(), FFloatAnimationAttribute::StaticStruct()},
-			{FIntProperty::StaticClass()->GetFName(), FIntegerAnimationAttribute::StaticStruct()},
-			{TransformTypeName, FTransformAnimationAttribute::StaticStruct()},
-			{VectorTypeName, FVectorAnimationAttribute::StaticStruct()},
-			{QuatTypeName, FQuaternionAnimationAttribute::StaticStruct()},
-		};
-
-	if (UScriptStruct** AttributeType = BuiltInAttributeTypes.Find(InTypeName))
+	if (UScriptStruct* const* AttributeType = GetBuiltInAttributeTypes().Find(InTypeName))
 	{
 		return *AttributeType;
 	}
diff --git a/Engine/Plugins/Animation/DeformerGraph/Source/OptimusCore/Private/OptimusHelpers.cpp b/Engine/Plugins/Animation/DeformerGraph/Source/OptimusCore/Private/OptimusHelpers.cpp
--- a/Engine/Plugins/Animation/DeformerGraph/Source/OptimusCore/Private/OptimusHelpers.cpp
+++ b/Engine/Plugins/Animation/DeformerGraph/Source/OptimusCore/Private/OptimusHelpers.cpp
@@ -181,6 +181,6 @@ FName Optimus::GetTypeName(UScriptStruct* InStructType, bool bInShouldGetUniqueN
 		}
 	}
 	
-	return FName(*FString::Printf(TEXT("F%s"), *InStructType->GetName()));
+	return FName(*InStructType->GetStructCPPName());
 }
 
diff --git a/Engine/Plugins/Animation/DeformerGraph/Source/OptimusCore/Public/OptimusDataTypeRegistry.h b/Engine/Plugins/Animation/DeformerGraph/Source/OptimusCore/Public/OptimusDataTypeRegistry.h
--- a/Engine/Plugins/Animation/DeformerGraph/Source/OptimusCore/Public/OptimusDataTypeRegistry.h
+++ b/Engine/Plugins/Animation/DeformerGraph/Source/OptimusCore/Public/OptimusDataTypeRegistry.h
@@ -183,10 +183,10 @@ protected:
 	static void UnregisterAllTypes();
 
 	/** Call during module init to register asset registry callbacks for struct type registration*/	
-	static void RegisterAssetRegistryCallbacks();
+	static void RegisterEngineCallbacks();
 	
 	/** Call during module shutdown to unregister asset registry callbacks*/	
-	static void UnregisterAssetRegistryCallbacks();
+	static void UnregisterEngineCallbacks();
 
 	/** A helper function to return a property create function. The function can be unbound
 	 *  and that should be checked prior to calling
@@ -197,10 +197,15 @@ protected:
 private:
 	FOptimusDataTypeRegistry();
 
-	// Callbacks to update registry when user defined structs are removed/renamed
+	// Callbacks to update registry when user defined structs are loaded/removed/renamed
+	void OnFilesLoaded();
 	void OnAssetRemoved(const FAssetData& InAssetData);
 	void OnAssetRenamed(const FAssetData& InAssetData, const FString& InOldName);
 	
+	// Callbacks to update registry when list of user defined structs supported by the animation attribute system is changed
+	void OnAnimationAttributeRegistryChanged(const UScriptStruct* InScriptStruct, bool bIsAdded);
+
+	
 	bool RegisterType(
 		FName InTypeName,
 		TFunction<void(FOptimusDataType &)> InFillFunc,
@@ -210,7 +215,7 @@ private:
 		);
 
 	/** Check if a type can be registered */
-	static EOptimusDataTypeUsageFlags CanRegisterStructType(UScriptStruct* InStruct);
+	static EOptimusDataTypeUsageFlags GetStructTypeUsageFlag(UScriptStruct* InStruct);
 
 	struct FTypeInfo
 	{
diff --git a/Engine/Plugins/Animation/DeformerGraph/Source/OptimusEditor/Private/Widgets/SOptimusDataTypeSelector.cpp b/Engine/Plugins/Animation/DeformerGraph/Source/OptimusEditor/Private/Widgets/SOptimusDataTypeSelector.cpp
--- a/Engine/Plugins/Animation/DeformerGraph/Source/OptimusEditor/Private/Widgets/SOptimusDataTypeSelector.cpp
+++ b/Engine/Plugins/Animation/DeformerGraph/Source/OptimusEditor/Private/Widgets/SOptimusDataTypeSelector.cpp
@@ -9,6 +9,7 @@
 #include "EdGraphSchema_K2.h"
 #include "OptimusHelpers.h"
 #include "SListViewSelectorDropdownMenu.h"
+#include "Animation/AttributeTypes.h"
 #include "AssetRegistry/AssetRegistryModule.h"
 #include "Engine/UserDefinedStruct.h"
 #include "Widgets/Images/SImage.h"
@@ -148,6 +149,17 @@ FText SOptimusDataTypeSelector::GetTypeDescription(FOptimusDataTypeHandle InData
 		{
 			if (EnumHasAnyFlags(InDataType->UsageFlags, UsageMask.Get()))
 			{
+				if (UsageMask.Get() == EOptimusDataTypeUsageFlags::AnimAttributes)
+				{
+					if (UUserDefinedStruct* UserDefinedStruct = Cast<UUserDefinedStruct>(InDataType->TypeObject))
+					{
+						if (!UE::Anim::AttributeTypes::IsTypeRegistered(UserDefinedStruct))
+						{
+							return FText::FromString(TEXT("<Unregistered> ") + InDataType->DisplayName.ToString());
+						}
+					}
+				}
+
 				return InDataType->DisplayName;
 			}
 		}
@@ -159,6 +171,38 @@ FText SOptimusDataTypeSelector::GetTypeDescription(FOptimusDataTypeHandle InData
 }
 
 
+FText SOptimusDataTypeSelector::GetTypeTooltip(FOptimusDataTypeHandle InDataType) const
+{
+	if (!InDataType.IsValid())
+	{
+		return FText::GetEmpty();
+	}
+
+	if (!UsageMask.IsSet())
+	{
+		return FText::GetEmpty();
+	}
+
+	if (UsageMask.Get() == EOptimusDataTypeUsageFlags::AnimAttributes)
+	{
+		if (UUserDefinedStruct* UserDefinedStruct = Cast<UUserDefinedStruct>(InDataType->TypeObject))
+		{
+			if (!EnumHasAnyFlags(InDataType->UsageFlags, EOptimusDataTypeUsageFlags::AnimAttributes))
+			{
+				return FText::FromString(TEXT("Type contains unsupported members or nested arrays"));
+			}
+			else if (!UE::Anim::AttributeTypes::IsTypeRegistered(UserDefinedStruct))
+			{
+				return FText::FromString(TEXT("Please register the type in Project Settings - Animation - CustomAttributes - User Defined Struct Animation Attributes."));
+			}
+		}
+	}
+
+
+	return FText::GetEmpty();
+
+}
+
 FText SOptimusDataTypeSelector::GetTypeTooltip() const
 {
 	FText EditText;
@@ -167,26 +211,46 @@ FText SOptimusDataTypeSelector::GetTypeTooltip() const
 		EditText = LOCTEXT("DataTypeSelector", "Click to see the data type menu and change the current type.");
 	}
 
-	return FText::Format(LOCTEXT("TypeTooltip", "{0}Current Type: {1}"), EditText, GetTypeDescription());
+	return FText::Format(LOCTEXT("TypeTooltip", "{0} Current Type: {1}"), EditText, GetTypeDescription());
 }
 
 
 TSharedRef<SWidget> SOptimusDataTypeSelector::GetMenuContent()
 {
 	AllDataTypeItems.Reset();
+
+	TArray<FOptimusDataTypeHandle> UnsupportedUserDefinedStructs;
 	FOptimusDataTypeHandle SelectedItem;
 	for (FOptimusDataTypeHandle DataType : FOptimusDataTypeRegistry::Get().GetAllTypes())
 	{
-		if ( !UsageMask.IsSet() || UsageMask.Get() == EOptimusDataTypeUsageFlags::None || EnumHasAnyFlags(DataType->UsageFlags, UsageMask.Get()))
+		if (!UsageMask.IsSet() ||
+			UsageMask.Get() == EOptimusDataTypeUsageFlags::None ||
+			EnumHasAnyFlags(DataType->UsageFlags, UsageMask.Get()))
 		{
 			AllDataTypeItems.Add(DataType);
-			if (DataType == CurrentDataType.Get())
+		}
+		else if (UsageMask.Get() == EOptimusDataTypeUsageFlags::AnimAttributes)
+		{
+			if (UUserDefinedStruct* UserDefinedStruct = Cast<UUserDefinedStruct>(DataType->TypeObject.Get()))
 			{
-				SelectedItem = AllDataTypeItems.Last();
+				UnsupportedUserDefinedStructs.Add(DataType);	
 			}
 		}
 	}
 
+	// Show unsupported types for animation attribute at the end
+	// so that we can provide tooltips on why they are not supported
+	AllDataTypeItems.Append(UnsupportedUserDefinedStructs);
+
+	for (FOptimusDataTypeHandle DataType : AllDataTypeItems)
+	{
+		if (DataType == CurrentDataType.Get())
+		{
+			SelectedItem = DataType;
+			break;
+		}
+	}
+
 	ViewDataTypeItems = AllDataTypeItems;
 
 	if (!MenuContent.IsValid())
@@ -228,6 +292,9 @@ TSharedRef<SWidget> SOptimusDataTypeSelector::GetMenuContent()
 			];
 	}
 
+	// Refresh in case ViewDataTypeItems has changed
+	TypeListView->RebuildList();
+	
 	// Update the current selection
 	if (SelectedItem.IsValid())
 	{
@@ -270,6 +337,7 @@ TSharedRef<ITableRow> SOptimusDataTypeSelector::GenerateTypeListRow(
 				.Text_Lambda([InItem, this]() { return SOptimusDataTypeSelector::GetTypeDescription(InItem); })
 				.HighlightText(FilterText)
 				.Font(FAppStyle::GetFontStyle(TEXT("Kismet.TypePicker.NormalFont")))
+				.ToolTipText_Lambda([InItem, this](){ return SOptimusDataTypeSelector::GetTypeTooltip(InItem); })
 			]
 		];
 }
diff --git a/Engine/Plugins/Animation/DeformerGraph/Source/OptimusEditor/Private/Widgets/SOptimusDataTypeSelector.h b/Engine/Plugins/Animation/DeformerGraph/Source/OptimusEditor/Private/Widgets/SOptimusDataTypeSelector.h
--- a/Engine/Plugins/Animation/DeformerGraph/Source/OptimusEditor/Private/Widgets/SOptimusDataTypeSelector.h
+++ b/Engine/Plugins/Animation/DeformerGraph/Source/OptimusEditor/Private/Widgets/SOptimusDataTypeSelector.h
@@ -60,6 +60,7 @@ private:
 	FSlateColor GetTypeIconColor() const { return GetTypeIconColor(CurrentDataType.Get()); }
 	FText GetTypeDescription(FOptimusDataTypeHandle InDataType) const;
 	FText GetTypeDescription() const { return GetTypeDescription(CurrentDataType.Get()); }
+	FText GetTypeTooltip(FOptimusDataTypeHandle InDataType) const;
 	FText GetTypeTooltip() const;
 	
 	TSharedRef<SWidget> GetMenuContent();
