Optimus: Add a Get/Set resource node for expressing when resource is used to store data between kernels.
Add explicit set of write operation (including interlocked) on a resource node.
Remove auto inference of interlocked operations.
Remove bClearBeforeUse on buffer data interface. Buffers can be explicitly cleared.

diff --git a/Engine/Plugins/Animation/DeformerGraph/Shaders/Private/DataInterfaceRawBuffer.ush b/Engine/Plugins/Animation/DeformerGraph/Shaders/Private/DataInterfaceRawBuffer.ush
--- a/Engine/Plugins/Animation/DeformerGraph/Shaders/Private/DataInterfaceRawBuffer.ush
+++ b/Engine/Plugins/Animation/DeformerGraph/Shaders/Private/DataInterfaceRawBuffer.ush
@@ -31,23 +31,29 @@ DI_IMPL_WRITE(WriteValue, uint Index, BUFFER_TYPE Value)
 	DI_LOCAL(BufferUAV)[BufferIndex] = Value;
 }
 
-#ifdef BUFFER_TYPE_SUPPORTS_ATOMIC
-
-DI_IMPL_READ(ReadAtomicAdd, BUFFER_TYPE, uint Index, BUFFER_TYPE Value)
+DI_IMPL_WRITE(WriteAtomicAdd, uint Index, BUFFER_TYPE Value)
 {
+#ifdef BUFFER_TYPE_SUPPORTS_ATOMIC
 	uint BufferIndex = DI_LOCAL(StartOffset) + Index;
-	
 	BUFFER_TYPE OldValue;
 	InterlockedAdd(DI_LOCAL(BufferUAV)[BufferIndex], Value, OldValue);
-	return OldValue;
+#endif
 }
 
-DI_IMPL_WRITE(WriteAtomicAdd, uint Index, BUFFER_TYPE Value)
+DI_IMPL_WRITE(WriteAtomicMin, uint Index, BUFFER_TYPE Value)
 {
+#ifdef BUFFER_TYPE_SUPPORTS_ATOMIC
 	uint BufferIndex = DI_LOCAL(StartOffset) + Index;
-
 	BUFFER_TYPE OldValue;
-	InterlockedAdd(DI_LOCAL(BufferUAV)[BufferIndex], Value, OldValue);
+	InterlockedMin(DI_LOCAL(BufferUAV)[BufferIndex], Value, OldValue);
+#endif
 }
 
+DI_IMPL_WRITE(WriteAtomicMax, uint Index, BUFFER_TYPE Value)
+{
+#ifdef BUFFER_TYPE_SUPPORTS_ATOMIC
+	uint BufferIndex = DI_LOCAL(StartOffset) + Index;
+	BUFFER_TYPE OldValue;
+	InterlockedMin(DI_LOCAL(BufferUAV)[BufferIndex], Value, OldValue);
 #endif
+}
diff --git a/Engine/Plugins/Animation/DeformerGraph/Source/OptimusCore/Private/DataInterfaces/OptimusDataInterfaceRawBuffer.cpp b/Engine/Plugins/Animation/DeformerGraph/Source/OptimusCore/Private/DataInterfaces/OptimusDataInterfaceRawBuffer.cpp
--- a/Engine/Plugins/Animation/DeformerGraph/Source/OptimusCore/Private/DataInterfaces/OptimusDataInterfaceRawBuffer.cpp
+++ b/Engine/Plugins/Animation/DeformerGraph/Source/OptimusCore/Private/DataInterfaces/OptimusDataInterfaceRawBuffer.cpp
@@ -13,10 +13,6 @@
 #include "ComponentSources/OptimusSkinnedMeshComponentSource.h"
 
 
-const int32 UOptimusRawBufferDataInterface::ReadValueInputIndex = 1;
-const int32 UOptimusRawBufferDataInterface::WriteValueOutputIndex = 0;
-
-
 const UOptimusComponentSource* UOptimusRawBufferDataInterface::GetComponentSource() const
 {
 	if (const UOptimusComponentSourceBinding* ComponentSourceBindingPtr = ComponentSourceBinding.Get())
@@ -26,7 +22,6 @@ const UOptimusComponentSource* UOptimusRawBufferDataInterface::GetComponentSourc
 	return nullptr;
 }
 
-
 bool UOptimusRawBufferDataInterface::SupportsAtomics() const
 {
 	return ValueType->Type == EShaderFundamentalType::Int;
@@ -41,13 +36,11 @@ TArray<FOptimusCDIPinDefinition> UOptimusRawBufferDataInterface::GetPinDefinitio
 	return Defs;
 }
 
-
 TSubclassOf<UActorComponent> UOptimusRawBufferDataInterface::GetRequiredComponentClass() const
 {
 	return USceneComponent::StaticClass();
 }
 
-
 void UOptimusRawBufferDataInterface::GetSupportedInputs(TArray<FShaderFunctionDefinition>& OutFunctions) const
 {
 	OutFunctions.AddDefaulted_GetRef()
@@ -58,31 +51,40 @@ void UOptimusRawBufferDataInterface::GetSupportedInputs(TArray<FShaderFunctionDe
 		.SetName(TEXT("ReadValue"))
 		.AddReturnType(ValueType)
 		.AddParam(EShaderFundamentalType::Uint);
-
-	if (SupportsAtomics())
-	{
-		OutFunctions.AddDefaulted_GetRef()
-			.SetName(TEXT("WriteAtomicAdd"))
-			.AddReturnType(ValueType)
-			.AddParam(EShaderFundamentalType::Uint)
-			.AddParam(ValueType);
-	}
 }
 
 void UOptimusRawBufferDataInterface::GetSupportedOutputs(TArray<FShaderFunctionDefinition>& OutFunctions) const
 {
+	// Functions in order of EOptimusBufferWriteType
 	OutFunctions.AddDefaulted_GetRef()
 		.SetName(TEXT("WriteValue"))
 		.AddParam(EShaderFundamentalType::Uint)
 		.AddParam(ValueType);
 
-	if (SupportsAtomics())
-	{
-		OutFunctions.AddDefaulted_GetRef()
-			.SetName(TEXT("WriteAtomicAdd"))
-			.AddParam(EShaderFundamentalType::Uint)
-			.AddParam(ValueType);
-	}
+	OutFunctions.AddDefaulted_GetRef()
+		.SetName(TEXT("WriteAtomicAdd"))
+		.AddParam(EShaderFundamentalType::Uint)
+		.AddParam(ValueType);
+
+	OutFunctions.AddDefaulted_GetRef()
+		.SetName(TEXT("WriteAtomicMin"))
+		.AddParam(EShaderFundamentalType::Uint)
+		.AddParam(ValueType);
+
+	OutFunctions.AddDefaulted_GetRef()
+		.SetName(TEXT("WriteAtomicMax"))
+		.AddParam(EShaderFundamentalType::Uint)
+		.AddParam(ValueType);
+}
+
+int32 UOptimusRawBufferDataInterface::GetReadValueInputIndex()
+{
+	return 1;
+}
+
+int32 UOptimusRawBufferDataInterface::GetWriteValueOutputIndex(EOptimusBufferWriteType WriteType)
+{
+	return (int32)WriteType;
 }
 
 BEGIN_SHADER_PARAMETER_STRUCT(FTransientBufferDataInterfaceParameters, )
@@ -118,12 +120,16 @@ void UOptimusRawBufferDataInterface::GetHLSL(FString& OutHLSL) const
 	OutHLSL += TEXT("#define BUFFER_TYPE ");
 	OutHLSL += ValueType->ToString();
 	OutHLSL += TEXT(" \n");
+
 	if (SupportsAtomics()) { OutHLSL += TEXT("#define BUFFER_TYPE_SUPPORTS_ATOMIC 1\n"); }
 	if (UseSplitBuffers()) { OutHLSL += TEXT("#define BUFFER_SPLIT_READ_WRITE 1\n"); }
+
 	OutHLSL += TEXT("#include \"/Plugin/Optimus/Private/DataInterfaceRawBuffer.ush\"\n");
-	OutHLSL += TEXT("#undef BUFFER_TYPE\n");
+
 	if (SupportsAtomics()) { OutHLSL += TEXT("#undef BUFFER_TYPE_SUPPORTS_ATOMIC\n"); }
 	if (UseSplitBuffers()) { OutHLSL += TEXT("#undef BUFFER_SPLIT_READ_WRITE\n"); }
+
+	OutHLSL += TEXT("#undef BUFFER_TYPE\n");
 }
 
 FString UOptimusTransientBufferDataInterface::GetDisplayName() const
@@ -136,9 +142,6 @@ UComputeDataProvider* UOptimusTransientBufferDataInterface::CreateDataProvider(
 	) const
 {
 	UOptimusTransientBufferDataProvider *Provider = CreateProvider<UOptimusTransientBufferDataProvider>(InBinding);
-	
-	Provider->bClearBeforeUse = bClearBeforeUse;
-	
 	return Provider;
 }
 
@@ -289,7 +292,7 @@ FComputeDataProviderRenderProxy* UOptimusTransientBufferDataProvider::GetRenderP
 	TArray<int32> InvocationCounts;
 	GetInvocationElementCounts(InvocationCounts);
 	
-	return new FOptimusTransientBufferDataProviderProxy(InvocationCounts, ElementStride, bClearBeforeUse);
+	return new FOptimusTransientBufferDataProviderProxy(InvocationCounts, ElementStride);
 }
 
 
@@ -315,12 +318,10 @@ FComputeDataProviderRenderProxy* UOptimusPersistentBufferDataProvider::GetRender
 
 FOptimusTransientBufferDataProviderProxy::FOptimusTransientBufferDataProviderProxy(
 	TArray<int32> InInvocationElementCounts,
-	int32 InElementStride,
-	bool bInClearBeforeUse
+	int32 InElementStride
 	) :
 	InvocationElementCounts(InInvocationElementCounts),
-	ElementStride(InElementStride),
-	bClearBeforeUse(bInClearBeforeUse)
+	ElementStride(InElementStride)
 {
 }
 
@@ -332,11 +333,6 @@ void FOptimusTransientBufferDataProviderProxy::AllocateResources(FRDGBuilder& Gr
 		Buffer.Add(GraphBuilder.CreateBuffer(FRDGBufferDesc::CreateStructuredDesc(ElementStride, NumElements), TEXT("TransientBuffer"), ERDGBufferFlags::None));
 		BufferSRV.Add(GraphBuilder.CreateSRV(Buffer.Last()));
 		BufferUAV.Add(GraphBuilder.CreateUAV(Buffer.Last()));
-
-		if (bClearBeforeUse)
-		{
-			AddClearUAVPass(GraphBuilder, BufferUAV.Last(), 0);
-		}
 	}
 }
 
diff --git a/Engine/Plugins/Animation/DeformerGraph/Source/OptimusCore/Private/DataInterfaces/OptimusDataInterfaceRawBuffer.h b/Engine/Plugins/Animation/DeformerGraph/Source/OptimusCore/Private/DataInterfaces/OptimusDataInterfaceRawBuffer.h
--- a/Engine/Plugins/Animation/DeformerGraph/Source/OptimusCore/Private/DataInterfaces/OptimusDataInterfaceRawBuffer.h
+++ b/Engine/Plugins/Animation/DeformerGraph/Source/OptimusCore/Private/DataInterfaces/OptimusDataInterfaceRawBuffer.h
@@ -17,14 +17,25 @@ class UOptimusComponentSourceBinding;
 class UOptimusRawBufferDataProvider;
 class UOptimusComponentSource;
 
+/** Write to buffer operation types. */
+UENUM()
+enum class EOptimusBufferWriteType : uint8
+{
+	Write UMETA(ToolTip = "Write the value to resource buffer."),
+	WriteAtomicAdd UMETA(ToolTip = "AtomicAdd the value to the value in the resource buffer."),
+	WriteAtomicMin UMETA(ToolTip = "AtomicMin the value with the value in the resource buffer."),
+	WriteAtomicMax UMETA(ToolTip = "AtomicMax the value with the value in the resource buffer."),
+	Count UMETA(Hidden),
+};
+
 UCLASS(Abstract)
 class OPTIMUSCORE_API UOptimusRawBufferDataInterface : public UOptimusComputeDataInterface
 {
 	GENERATED_BODY()
 
 public:
-	static const int32 ReadValueInputIndex;
-	static const int32 WriteValueOutputIndex;
+	static int32 GetReadValueInputIndex();
+	static int32 GetWriteValueOutputIndex(EOptimusBufferWriteType WriteType);
 	
 	//~ Begin UOptimusComputeDataInterface Interface
 	TArray<FOptimusCDIPinDefinition> GetPinDefinitions() const override;
@@ -50,7 +61,7 @@ public:
 	/** The component source to query component domain validity and sizing */
 	UPROPERTY()
 	TWeakObjectPtr<UOptimusComponentSourceBinding> ComponentSourceBinding;
-	
+
 protected:
 	virtual bool UseSplitBuffers() const { return true; }
 
@@ -89,10 +100,6 @@ public:
 	void GetShaderParameters(TCHAR const* UID, FShaderParametersMetadataBuilder& InOutBuilder, FShaderParametersMetadataAllocations& InOutAllocations) const override;
 	UComputeDataProvider* CreateDataProvider(TObjectPtr<UObject> InBinding, uint64 InInputMask, uint64 InOutputMask) const override;
 	//~ End UComputeDataInterface Interface
-
-	/** Set to true if the buffer should be cleared prior to each render */ 
-	UPROPERTY()
-	bool bClearBeforeUse = true;
 };
 
 
@@ -168,9 +175,6 @@ public:
 	//~ Begin UComputeDataProvider Interface
 	FComputeDataProviderRenderProxy* GetRenderProxy() override;
 	//~ End UComputeDataProvider Interface
-
-	UPROPERTY()
-	bool bClearBeforeUse = true;
 };
 
 
@@ -204,8 +208,7 @@ class FOptimusTransientBufferDataProviderProxy :
 public:
 	FOptimusTransientBufferDataProviderProxy(
 		TArray<int32> InInvocationElementCounts,
-		int32 InElementStride,
-		bool bInClearBeforeUse
+		int32 InElementStride
 		);
 
 	//~ Begin FComputeDataProviderRenderProxy Interface
@@ -216,7 +219,6 @@ public:
 private:
 	const TArray<int32> InvocationElementCounts;
 	const int32 ElementStride;
-	const bool bClearBeforeUse;
 
 	TArray<FRDGBuffer*> Buffer;
 	TArray<FRDGBufferSRV*> BufferSRV;
diff --git a/Engine/Plugins/Animation/DeformerGraph/Source/OptimusCore/Private/Nodes/OptimusNode_ComputeKernelBase.cpp b/Engine/Plugins/Animation/DeformerGraph/Source/OptimusCore/Private/Nodes/OptimusNode_ComputeKernelBase.cpp
--- a/Engine/Plugins/Animation/DeformerGraph/Source/OptimusCore/Private/Nodes/OptimusNode_ComputeKernelBase.cpp
+++ b/Engine/Plugins/Animation/DeformerGraph/Source/OptimusCore/Private/Nodes/OptimusNode_ComputeKernelBase.cpp
@@ -395,7 +395,7 @@ void UOptimusNode_ComputeKernelBase::ProcessInputPinForComputeKernel(
 			// For transient buffers we need the function index as given by the
 			// ReadValue function. 
 			DataInterface = InLinkDataInterfaceMap[OutputPin];
-			DataInterfaceFuncIndex = UOptimusTransientBufferDataInterface::ReadValueInputIndex;
+			DataInterfaceFuncIndex = UOptimusTransientBufferDataInterface::GetReadValueInputIndex();
 
 			TArray<FShaderFunctionDefinition> ReadFunctions;
 			DataInterface->GetSupportedInputs(ReadFunctions);
@@ -554,14 +554,6 @@ void UOptimusNode_ComputeKernelBase::ProcessOutputPinForComputeKernel(
 			TArray<FShaderFunctionDefinition> WriteFunctions;
 			DataInterface->GetSupportedOutputs(WriteFunctions);
 
-			// This is a horrible hack for detecting interlocked writes
-			// TODO: Either express this via the kernel metadata or add full support for buffer data interface in graph editor.
-			int32 WriteValueOutputIndex = UOptimusTransientBufferDataInterface::WriteValueOutputIndex;
-			if (InOutputPin->GetName().Contains(TEXT("Interlocked")) && WriteFunctions.Num() > WriteValueOutputIndex + 1)
-			{
-				++WriteValueOutputIndex;
-			}
-
 			// Get the component binding from the upstream connection.
 			UOptimusComponentSourceBinding* ComponentBinding = nullptr;
 			TArray<UOptimusComponentSourceBinding*> ComponentBindings = GetOwningGraph()->GetComponentSourceBindingsForPin(InOutputPin).Array();
@@ -570,7 +562,8 @@ void UOptimusNode_ComputeKernelBase::ProcessOutputPinForComputeKernel(
 				ComponentBinding = ComponentBindings[0];
 			}
 			
-			WriteConnectionDefs.Add({DataInterface, ComponentBinding, WriteFunctions[WriteValueOutputIndex].Name, TEXT("Transient")});
+			const int32 OutputIndex = UOptimusTransientBufferDataInterface::GetWriteValueOutputIndex(EOptimusBufferWriteType::Write);
+			WriteConnectionDefs.Add({DataInterface, ComponentBinding, WriteFunctions[OutputIndex].Name, TEXT("Transient")});
 		}
 		
 		for (const FOptimusRoutedNodePin& ConnectedPin: ConnectedPins)
diff --git a/Engine/Plugins/Animation/DeformerGraph/Source/OptimusCore/Private/Nodes/OptimusNode_GetResource.cpp b/Engine/Plugins/Animation/DeformerGraph/Source/OptimusCore/Private/Nodes/OptimusNode_GetResource.cpp
--- a/Engine/Plugins/Animation/DeformerGraph/Source/OptimusCore/Private/Nodes/OptimusNode_GetResource.cpp
+++ b/Engine/Plugins/Animation/DeformerGraph/Source/OptimusCore/Private/Nodes/OptimusNode_GetResource.cpp
@@ -18,7 +18,7 @@ int32 UOptimusNode_GetResource::GetDataFunctionIndexFromPin(const UOptimusNodePi
 		return INDEX_NONE;
 	}
 
-	return UOptimusRawBufferDataInterface::ReadValueInputIndex;
+	return UOptimusRawBufferDataInterface::GetReadValueInputIndex();
 }
 
 void UOptimusNode_GetResource::ConstructNode()
@@ -27,7 +27,7 @@ void UOptimusNode_GetResource::ConstructNode()
 	{
 		// FIXME: Define context.
 		AddPinDirect(
-			Res->ResourceName,
+			GetResourcePinName(0),
 			EOptimusNodePinDirection::Output,
 			Res->DataDomain,
 			Res->DataType);
diff --git a/Engine/Plugins/Animation/DeformerGraph/Source/OptimusCore/Private/Nodes/OptimusNode_Resource.cpp b/Engine/Plugins/Animation/DeformerGraph/Source/OptimusCore/Private/Nodes/OptimusNode_Resource.cpp
new file mode 100644
--- /dev/null
+++ b/Engine/Plugins/Animation/DeformerGraph/Source/OptimusCore/Private/Nodes/OptimusNode_Resource.cpp
@@ -0,0 +1,60 @@
+// Copyright Epic Games, Inc. All Rights Reserved.
+
+#include "Nodes/OptimusNode_Resource.h"
+
+#include "OptimusNodePin.h"
+#include "OptimusResourceDescription.h"
+#include "DataInterfaces/OptimusDataInterfaceRawBuffer.h"
+
+
+int32 UOptimusNode_Resource::GetDataFunctionIndexFromPin(const UOptimusNodePin* InPin) const
+{
+	if (!InPin || InPin->GetParentPin() != nullptr)
+	{
+		return INDEX_NONE;
+	}
+	if (!ensure(GetPins().Contains(InPin)))
+	{
+		return INDEX_NONE;
+	}
+
+	return InPin->GetDirection() == EOptimusNodePinDirection::Output 
+		? UOptimusRawBufferDataInterface::GetReadValueInputIndex() 
+		: UOptimusRawBufferDataInterface::GetWriteValueOutputIndex(WriteType);
+}
+
+
+FName UOptimusNode_Resource::GetResourcePinName(int32 InPinIndex) const
+{
+	if (const UOptimusResourceDescription* Res = GetResourceDescription())
+	{
+		if (ensure(InPinIndex >= 0 && InPinIndex < 2))
+		{
+			FString Suffix = InPinIndex == 0 ? TEXT("Set") : TEXT("Get");
+			return FName(Suffix + Res->ResourceName.ToString());
+		}
+	}
+
+	return {};
+}
+
+
+void UOptimusNode_Resource::ConstructNode()
+{
+	if (const UOptimusResourceDescription* Res  = GetResourceDescription())
+	{
+		AddPinDirect(
+			GetResourcePinName(0),
+			EOptimusNodePinDirection::Input,
+			Res->DataDomain,
+			Res->DataType);
+
+		const FName GetterName("Get");
+
+		AddPinDirect(
+			GetResourcePinName(1),
+			EOptimusNodePinDirection::Output,
+			Res->DataDomain,
+			Res->DataType);
+	}
+}
diff --git a/Engine/Plugins/Animation/DeformerGraph/Source/OptimusCore/Private/Nodes/OptimusNode_Resource.h b/Engine/Plugins/Animation/DeformerGraph/Source/OptimusCore/Private/Nodes/OptimusNode_Resource.h
new file mode 100644
--- /dev/null
+++ b/Engine/Plugins/Animation/DeformerGraph/Source/OptimusCore/Private/Nodes/OptimusNode_Resource.h
@@ -0,0 +1,23 @@
+// Copyright Epic Games, Inc. All Rights Reserved.
+
+#pragma once
+
+#include "OptimusNode_ResourceAccessorBase.h"
+
+#include "OptimusNode_Resource.generated.h"
+
+
+UCLASS(Hidden)
+class UOptimusNode_Resource
+	: public UOptimusNode_ResourceAccessorBase
+{
+	GENERATED_BODY()
+
+	// UOptimusNode_Resource implementations
+	FName GetResourcePinName(int32 InPinIndex) const override;
+	// IOptimusDataInterfaceProvider implementations
+	int32 GetDataFunctionIndexFromPin(const UOptimusNodePin* InPin) const override;
+
+protected:
+	void ConstructNode() override;
+};
diff --git a/Engine/Plugins/Animation/DeformerGraph/Source/OptimusCore/Private/Nodes/OptimusNode_ResourceAccessorBase.cpp b/Engine/Plugins/Animation/DeformerGraph/Source/OptimusCore/Private/Nodes/OptimusNode_ResourceAccessorBase.cpp
--- a/Engine/Plugins/Animation/DeformerGraph/Source/OptimusCore/Private/Nodes/OptimusNode_ResourceAccessorBase.cpp
+++ b/Engine/Plugins/Animation/DeformerGraph/Source/OptimusCore/Private/Nodes/OptimusNode_ResourceAccessorBase.cpp
@@ -32,6 +32,18 @@ UOptimusResourceDescription* UOptimusNode_ResourceAccessorBase::GetResourceDescr
 }
 
 
+FName UOptimusNode_ResourceAccessorBase::GetResourcePinName(int32 InPinIndex) const
+{
+	UOptimusResourceDescription const* Description = GetResourceDescription();
+	if (!ensure(Description))
+	{
+		return {};
+	}
+	
+	return Description->ResourceName;
+}
+
+
 TOptional<FText> UOptimusNode_ResourceAccessorBase::ValidateForCompile() const
 {
 	const UOptimusResourceDescription* ResourceDescription = GetResourceDescription();
diff --git a/Engine/Plugins/Animation/DeformerGraph/Source/OptimusCore/Private/Nodes/OptimusNode_ResourceAccessorBase.h b/Engine/Plugins/Animation/DeformerGraph/Source/OptimusCore/Private/Nodes/OptimusNode_ResourceAccessorBase.h
--- a/Engine/Plugins/Animation/DeformerGraph/Source/OptimusCore/Private/Nodes/OptimusNode_ResourceAccessorBase.h
+++ b/Engine/Plugins/Animation/DeformerGraph/Source/OptimusCore/Private/Nodes/OptimusNode_ResourceAccessorBase.h
@@ -2,6 +2,7 @@
 
 #pragma once
 
+#include "DataInterfaces/OptimusDataInterfaceRawBuffer.h"
 #include "IOptimusComponentBindingProvider.h"
 #include "IOptimusDataInterfaceProvider.h"
 #include "OptimusNode.h"
@@ -25,7 +26,10 @@ public:
 	void SetResourceDescription(UOptimusResourceDescription* InResourceDesc);
 
 	UOptimusResourceDescription* GetResourceDescription() const;
-	
+
+	/** Returns name to use when creating each pin on the node. */
+	virtual FName GetResourcePinName(int32 InPinIndex) const;
+
 	// UOptimusNode overrides
 	FName GetNodeCategory() const override 
 	{
@@ -41,7 +45,11 @@ public:
 	UOptimusComponentSourceBinding* GetComponentBinding() const override;
 
 	
-private:
+protected:
 	UPROPERTY()
 	TWeakObjectPtr<UOptimusResourceDescription> ResourceDesc;
+
+	/** Logical operation when writing to the resource. */
+	UPROPERTY(EditAnywhere, Category = Resource)
+	EOptimusBufferWriteType WriteType = EOptimusBufferWriteType::Write;
 };
diff --git a/Engine/Plugins/Animation/DeformerGraph/Source/OptimusCore/Private/Nodes/OptimusNode_SetResource.cpp b/Engine/Plugins/Animation/DeformerGraph/Source/OptimusCore/Private/Nodes/OptimusNode_SetResource.cpp
--- a/Engine/Plugins/Animation/DeformerGraph/Source/OptimusCore/Private/Nodes/OptimusNode_SetResource.cpp
+++ b/Engine/Plugins/Animation/DeformerGraph/Source/OptimusCore/Private/Nodes/OptimusNode_SetResource.cpp
@@ -18,7 +18,7 @@ int32 UOptimusNode_SetResource::GetDataFunctionIndexFromPin(const UOptimusNodePi
 		return INDEX_NONE;
 	}
 
-	return UOptimusRawBufferDataInterface::WriteValueOutputIndex;
+	return UOptimusRawBufferDataInterface::GetWriteValueOutputIndex(WriteType);
 }
 
 
@@ -27,7 +27,7 @@ void UOptimusNode_SetResource::ConstructNode()
 	if (const UOptimusResourceDescription* Res = GetResourceDescription())
 	{
 		AddPinDirect(
-		    Res->ResourceName,
+		    GetResourcePinName(0),
 		    EOptimusNodePinDirection::Input,
 			Res->DataDomain,
 		    Res->DataType);
diff --git a/Engine/Plugins/Animation/DeformerGraph/Source/OptimusCore/Private/OptimusDeformer.cpp b/Engine/Plugins/Animation/DeformerGraph/Source/OptimusCore/Private/OptimusDeformer.cpp
--- a/Engine/Plugins/Animation/DeformerGraph/Source/OptimusCore/Private/OptimusDeformer.cpp
+++ b/Engine/Plugins/Animation/DeformerGraph/Source/OptimusCore/Private/OptimusDeformer.cpp
@@ -535,10 +535,7 @@ bool UOptimusDeformer::RemoveResource(UOptimusResourceDescription* InResourceDes
 		UOptimusNode_ResourceAccessorBase* ResourceNode = Cast<UOptimusNode_ResourceAccessorBase>(Node);
 		if (ResourceNode->GetResourceDescription() == InResourceDesc)
 		{
-			if (ensure(ResourceNode->GetPins().Num() == 1))
-			{
-				NodesByGraph.FindOrAdd(ResourceNode->GetOwningGraph()).Add(ResourceNode);
-			}
+			NodesByGraph.FindOrAdd(ResourceNode->GetOwningGraph()).Add(ResourceNode);
 		}	
 	}
 
@@ -595,9 +592,9 @@ bool UOptimusDeformer::RenameResource(
 		const UOptimusNode_ResourceAccessorBase* ResourceNode = Cast<UOptimusNode_ResourceAccessorBase>(Node);
 		if (ResourceNode->GetResourceDescription() == InResourceDesc)
 		{
-			if (ensure(ResourceNode->GetPins().Num() == 1))
+			for (int32 PinIndex = 0; PinIndex < ResourceNode->GetPins().Num(); ++PinIndex)
 			{
-				Action->AddSubAction<FOptimusNodeAction_SetPinName>(ResourceNode->GetPins()[0], InNewName);
+				Action->AddSubAction<FOptimusNodeAction_SetPinName>(ResourceNode->GetPins()[PinIndex], ResourceNode->GetResourcePinName(PinIndex));
 			}
 		}	
 	}	
@@ -652,25 +649,24 @@ bool UOptimusDeformer::SetResourceDataType(
 		const UOptimusNode_ResourceAccessorBase* ResourceNode = Cast<UOptimusNode_ResourceAccessorBase>(Node);
 		if (ResourceNode->GetResourceDescription() == InResourceDesc)
 		{
-			if (ensure(ResourceNode->GetPins().Num() == 1))
+			for (UOptimusNodePin* Pin : ResourceNode->GetPins())
 			{
-				UOptimusNodePin* Pin = ResourceNode->GetPins()[0];
-
 				// Update the pin type to match.
-				Action->AddSubAction<FOptimusNodeAction_SetPinType>(ResourceNode->GetPins()[0], InDataType);
+				Action->AddSubAction<FOptimusNodeAction_SetPinType>(Pin, InDataType);
 
-				// Collect _unique_ links (in case there's a resource->resource link, since that would otherwise
-				// show up twice).
+				// Collect _unique_ links (in case there's a resource->resource link, since that would otherwise show up twice).
 				const UOptimusNodeGraph* Graph = Pin->GetOwningNode()->GetOwningGraph();
 
 				for (UOptimusNodePin* ConnectedPin: Graph->GetConnectedPins(Pin))
-				if (Pin->GetDirection() == EOptimusNodePinDirection::Output)
 				{
-					Links.Add({Pin, ConnectedPin});
-				}
-				else
-				{
-					Links.Add({ConnectedPin, Pin});
+					if (Pin->GetDirection() == EOptimusNodePinDirection::Output)
+					{
+						Links.Add({Pin, ConnectedPin});
+					}
+					else
+					{
+						Links.Add({ConnectedPin, Pin});
+					}
 				}
 			}
 		}	
@@ -720,18 +716,17 @@ bool UOptimusDeformer::SetResourceDataDomain(
 		const UOptimusNode_ResourceAccessorBase* ResourceNode = Cast<UOptimusNode_ResourceAccessorBase>(Node);
 		if (ResourceNode->GetResourceDescription() == InResourceDesc)
 		{
-			if (ensure(ResourceNode->GetPins().Num() == 1))
+			for (UOptimusNodePin* Pin : ResourceNode->GetPins())
 			{
-				UOptimusNodePin* Pin = ResourceNode->GetPins()[0];
-
 				// Update the pin type to match.
-				Action->AddSubAction<FOptimusNodeAction_SetPinDataDomain>(ResourceNode->GetPins()[0], InDataDomain);
+				Action->AddSubAction<FOptimusNodeAction_SetPinDataDomain>(Pin, InDataDomain);
 
 				// Collect _unique_ links (in case there's a resource->resource link, since that would otherwise
 				// show up twice).
 				const UOptimusNodeGraph* Graph = Pin->GetOwningNode()->GetOwningGraph();
 
 				for (UOptimusNodePin* ConnectedPin: Graph->GetConnectedPins(Pin))
+				{
 					if (Pin->GetDirection() == EOptimusNodePinDirection::Output)
 					{
 						Links.Add({Pin, ConnectedPin});
@@ -740,6 +735,7 @@ bool UOptimusDeformer::SetResourceDataDomain(
 					{
 						Links.Add({ConnectedPin, Pin});
 					}
+				}
 			}
 		}	
 	}
@@ -1558,7 +1554,6 @@ UOptimusComputeGraph* UOptimusDeformer::CompileNodeGraphToComputeGraph(
 								return nullptr;
 							}
 
-							TransientBufferDI->bClearBeforeUse = true;
 							TransientBufferDI->ValueType = Pin->GetDataType()->ShaderValueType;
 							TransientBufferDI->DataDomain = Pin->GetDataDomain();
 							TransientBufferDI->ComponentSourceBinding = *ComponentSourceBindings.CreateConstIterator();
diff --git a/Engine/Plugins/Animation/DeformerGraph/Source/OptimusCore/Private/OptimusNodeGraph.cpp b/Engine/Plugins/Animation/DeformerGraph/Source/OptimusCore/Private/OptimusNodeGraph.cpp
--- a/Engine/Plugins/Animation/DeformerGraph/Source/OptimusCore/Private/OptimusNodeGraph.cpp
+++ b/Engine/Plugins/Animation/DeformerGraph/Source/OptimusCore/Private/OptimusNodeGraph.cpp
@@ -13,6 +13,7 @@
 #include "Actions/OptimusNodeActions.h"
 #include "Nodes/OptimusNode_GetResource.h"
 #include "Nodes/OptimusNode_GetVariable.h"
+#include "Nodes/OptimusNode_Resource.h"
 #include "Nodes/OptimusNode_SetResource.h"
 
 #include "Containers/Queue.h"
@@ -309,6 +310,19 @@ UOptimusNode* UOptimusNodeGraph::AddDataInterfaceNode(
 }
 
 
+UOptimusNode* UOptimusNodeGraph::AddResourceNode(
+	UOptimusResourceDescription* InResourceDesc,
+	const FVector2D& InPosition
+)
+{
+	return AddNodeInternal(UOptimusNode_Resource::StaticClass(), InPosition,
+		[InResourceDesc](UOptimusNode* InNode)
+		{
+			Cast<UOptimusNode_Resource>(InNode)->SetResourceDescription(InResourceDesc);
+		});
+}
+
+
 UOptimusNode* UOptimusNodeGraph::AddResourceGetNode(
 	UOptimusResourceDescription* InResourceDesc, 
 	const FVector2D& InPosition
diff --git a/Engine/Plugins/Animation/DeformerGraph/Source/OptimusCore/Public/OptimusNodeGraph.h b/Engine/Plugins/Animation/DeformerGraph/Source/OptimusCore/Public/OptimusNodeGraph.h
--- a/Engine/Plugins/Animation/DeformerGraph/Source/OptimusCore/Public/OptimusNodeGraph.h
+++ b/Engine/Plugins/Animation/DeformerGraph/Source/OptimusCore/Public/OptimusNodeGraph.h
@@ -125,6 +125,11 @@ public:
 		const FVector2D& InPosition
 	);
 
+	UFUNCTION(BlueprintCallable, Category = OptimusNodeGraph)
+	UOptimusNode* AddResourceNode(
+		UOptimusResourceDescription* InResourceDesc,
+		const FVector2D& InPosition);
+
 	UFUNCTION(BlueprintCallable, Category = OptimusNodeGraph)
 	UOptimusNode* AddResourceGetNode(
 	    UOptimusResourceDescription *InResourceDesc,
diff --git a/Engine/Plugins/Animation/DeformerGraph/Source/OptimusEditor/Private/Widgets/SOptimusEditorGraphExplorerActions.cpp b/Engine/Plugins/Animation/DeformerGraph/Source/OptimusEditor/Private/Widgets/SOptimusEditorGraphExplorerActions.cpp
--- a/Engine/Plugins/Animation/DeformerGraph/Source/OptimusEditor/Private/Widgets/SOptimusEditorGraphExplorerActions.cpp
+++ b/Engine/Plugins/Animation/DeformerGraph/Source/OptimusEditor/Private/Widgets/SOptimusEditorGraphExplorerActions.cpp
@@ -201,6 +201,16 @@ FReply FOptimusEditorGraphDragAction_Resource::DroppedOnPanel(
 
 	MenuBuilder.BeginSection("OptimusResourceDroppedOn", ResourceNameText);
 
+	MenuBuilder.AddMenuEntry(
+		FText::Format(LOCTEXT("CreateResource", "Get/Set {0}"), ResourceNameText),
+		FText::Format(LOCTEXT("CreateResourceToolTip", "Create Getter/Setter for resource '{0}'"), ResourceNameText),
+		FSlateIcon(),
+		FUIAction(
+			FExecuteAction::CreateLambda([ModelGraph, ResourceDesc, InGraphPosition]() {
+				ModelGraph->AddResourceNode(ResourceDesc, InGraphPosition);
+			}),
+			FCanExecuteAction()));
+
 	MenuBuilder.AddMenuEntry(
 	    FText::Format(LOCTEXT("CreateGetResource", "Get {0}"), ResourceNameText),
 	    FText::Format(LOCTEXT("CreateGetResourceToolTip", "Create Getter for resource '{0}'"), ResourceNameText),
diff --git a/Engine/Plugins/Runtime/ComputeFramework/Source/ComputeFramework/Private/ComputeGraph.cpp b/Engine/Plugins/Runtime/ComputeFramework/Source/ComputeFramework/Private/ComputeGraph.cpp
--- a/Engine/Plugins/Runtime/ComputeFramework/Source/ComputeFramework/Private/ComputeGraph.cpp
+++ b/Engine/Plugins/Runtime/ComputeFramework/Source/ComputeFramework/Private/ComputeGraph.cpp
@@ -522,19 +522,31 @@ FString UComputeGraph::BuildKernelSource(
 			TCHAR const* WrapNamespace = GraphEdge.BindingFunctionNamespace.IsEmpty() ? nullptr : *GraphEdge.BindingFunctionNamespace;
 			if (GraphEdge.bKernelInput)
 			{
-				TArray<FShaderFunctionDefinition> DataProviderFunctions;
-				DataInterfaces[GraphEdge.DataInterfaceIndex]->GetSupportedInputs(DataProviderFunctions);
-				FShaderFunctionDefinition const& DataProviderFunction = DataProviderFunctions[GraphEdge.DataInterfaceBindingIndex];
-				FShaderFunctionDefinition const& KernelFunction = InKernelSource.ExternalInputs[GraphEdge.KernelBindingIndex];
-				GetFunctionShimHLSL(DataProviderFunction, KernelFunction, *NamePrefix, WrapNameOverride, WrapNamespace, HLSL);
+				if (ensure(DataInterfaces.IsValidIndex(GraphEdge.DataInterfaceIndex)))
+				{
+					TArray<FShaderFunctionDefinition> DataProviderFunctions;
+					DataInterfaces[GraphEdge.DataInterfaceIndex]->GetSupportedInputs(DataProviderFunctions);
+					if (ensure(DataProviderFunctions.IsValidIndex(GraphEdge.DataInterfaceBindingIndex)) && ensure(InKernelSource.ExternalInputs.IsValidIndex(GraphEdge.KernelBindingIndex)))
+					{
+						FShaderFunctionDefinition const& DataProviderFunction = DataProviderFunctions[GraphEdge.DataInterfaceBindingIndex];
+						FShaderFunctionDefinition const& KernelFunction = InKernelSource.ExternalInputs[GraphEdge.KernelBindingIndex];
+						GetFunctionShimHLSL(DataProviderFunction, KernelFunction, *NamePrefix, WrapNameOverride, WrapNamespace, HLSL);
+					}
+				}
 			}
 			else
 			{
-				TArray<FShaderFunctionDefinition> DataProviderFunctions;
-				DataInterfaces[GraphEdge.DataInterfaceIndex]->GetSupportedOutputs(DataProviderFunctions);
-				FShaderFunctionDefinition const&  DataProviderFunction = DataProviderFunctions[GraphEdge.DataInterfaceBindingIndex];
-				FShaderFunctionDefinition const& KernelFunction = InKernelSource.ExternalOutputs[GraphEdge.KernelBindingIndex];
-				GetFunctionShimHLSL(DataProviderFunction, KernelFunction, *NamePrefix, WrapNameOverride, WrapNamespace, HLSL);
+				if (ensure(DataInterfaces.IsValidIndex(GraphEdge.DataInterfaceIndex)))
+				{
+					TArray<FShaderFunctionDefinition> DataProviderFunctions;
+					DataInterfaces[GraphEdge.DataInterfaceIndex]->GetSupportedOutputs(DataProviderFunctions);
+					if (ensure(DataProviderFunctions.IsValidIndex(GraphEdge.DataInterfaceBindingIndex)) && ensure(InKernelSource.ExternalOutputs.IsValidIndex(GraphEdge.KernelBindingIndex)))
+					{
+						FShaderFunctionDefinition const&  DataProviderFunction = DataProviderFunctions[GraphEdge.DataInterfaceBindingIndex];
+						FShaderFunctionDefinition const& KernelFunction = InKernelSource.ExternalOutputs[GraphEdge.KernelBindingIndex];
+						GetFunctionShimHLSL(DataProviderFunction, KernelFunction, *NamePrefix, WrapNameOverride, WrapNamespace, HLSL);
+					}
+				}
 			}
 		}
 	}
