Optimus: Handle unaligned buffer types by aliasing with a buffer of a matching scalar type.
This is to deal with cases where some shader compilers needs structured buffer types to align on 16 bytes.
Implemented packing logic for vector3 types. User defined structures may need their own packing function.

diff --git a/Engine/Plugins/Animation/DeformerGraph/Shaders/Private/DataInterfaceRawBuffer.ush b/Engine/Plugins/Animation/DeformerGraph/Shaders/Private/DataInterfaceRawBuffer.ush
--- a/Engine/Plugins/Animation/DeformerGraph/Shaders/Private/DataInterfaceRawBuffer.ush
+++ b/Engine/Plugins/Animation/DeformerGraph/Shaders/Private/DataInterfaceRawBuffer.ush
@@ -2,10 +2,14 @@
 
 #include "/Plugin/ComputeFramework/Private/ComputeKernelCommon.ush"
 
+// BUFFER_TYPE is type of underlying structure.
+// PUBLIC_TYPE is the exposed type that calling code expects to use.
+// Where BUFFER_TYPE is different to PUBLIC_TYPE we set BUFFER_TYPE_RAW and expect ReadRawBuffer() and WriteRawBuffer() to be implemented somewhere.
+
 uint DI_LOCAL(StartOffset);
 uint DI_LOCAL(BufferSize);
 
-#ifdef BUFFER_SPLIT_READ_WRITE
+#if BUFFER_SPLIT_READ_WRITE
 StructuredBuffer<BUFFER_TYPE> DI_LOCAL(BufferSRV);
 #endif
 RWStructuredBuffer<BUFFER_TYPE> DI_LOCAL(BufferUAV);
@@ -15,45 +19,59 @@ DI_IMPL_READ(ReadNumValues, uint, )
 	return DI_LOCAL(BufferSize);
 }
 
-DI_IMPL_READ(ReadValue, BUFFER_TYPE, uint Index)
+DI_IMPL_READ(ReadValue, PUBLIC_TYPE, uint Index)
 {
 	uint BufferIndex = DI_LOCAL(StartOffset) + Index;
-#ifdef BUFFER_SPLIT_READ_WRITE
-	return DI_LOCAL(BufferSRV)[BufferIndex];
+	PUBLIC_TYPE ReturnValue;
+#if BUFFER_SPLIT_READ_WRITE
+#if BUFFER_TYPE_RAW
+	ReadRawBuffer(DI_LOCAL(BufferSRV), BufferIndex, ReturnValue);
+#else
+	ReturnValue = DI_LOCAL(BufferSRV)[BufferIndex];
+#endif
+#else
+#if BUFFER_TYPE_RAW
+	ReadRawBuffer(DI_LOCAL(BufferUAV), BufferIndex, ReturnValue);
 #else
-	return DI_LOCAL(BufferUAV)[BufferIndex];
+	ReturnValue = DI_LOCAL(BufferUAV)[BufferIndex];
 #endif
+#endif
+	return ReturnValue;
 }
 
-DI_IMPL_WRITE(WriteValue, uint Index, BUFFER_TYPE Value)
+DI_IMPL_WRITE(WriteValue, uint Index, PUBLIC_TYPE Value)
 {
 	uint BufferIndex = DI_LOCAL(StartOffset) + Index;
+#if BUFFER_TYPE_RAW
+	WriteRawBuffer(DI_LOCAL(BufferUAV), BufferIndex, Value);
+#else
 	DI_LOCAL(BufferUAV)[BufferIndex] = Value;
+#endif
 }
 
-DI_IMPL_WRITE(WriteAtomicAdd, uint Index, BUFFER_TYPE Value)
+DI_IMPL_WRITE(WriteAtomicAdd, uint Index, PUBLIC_TYPE Value)
 {
-#ifdef BUFFER_TYPE_SUPPORTS_ATOMIC
+#if BUFFER_TYPE_SUPPORTS_ATOMIC
 	uint BufferIndex = DI_LOCAL(StartOffset) + Index;
-	BUFFER_TYPE OldValue;
+	PUBLIC_TYPE OldValue;
 	InterlockedAdd(DI_LOCAL(BufferUAV)[BufferIndex], Value, OldValue);
 #endif
 }
 
-DI_IMPL_WRITE(WriteAtomicMin, uint Index, BUFFER_TYPE Value)
+DI_IMPL_WRITE(WriteAtomicMin, uint Index, PUBLIC_TYPE Value)
 {
-#ifdef BUFFER_TYPE_SUPPORTS_ATOMIC
+#if BUFFER_TYPE_SUPPORTS_ATOMIC
 	uint BufferIndex = DI_LOCAL(StartOffset) + Index;
-	BUFFER_TYPE OldValue;
+	PUBLIC_TYPE OldValue;
 	InterlockedMin(DI_LOCAL(BufferUAV)[BufferIndex], Value, OldValue);
 #endif
 }
 
-DI_IMPL_WRITE(WriteAtomicMax, uint Index, BUFFER_TYPE Value)
+DI_IMPL_WRITE(WriteAtomicMax, uint Index, PUBLIC_TYPE Value)
 {
-#ifdef BUFFER_TYPE_SUPPORTS_ATOMIC
+#if BUFFER_TYPE_SUPPORTS_ATOMIC
 	uint BufferIndex = DI_LOCAL(StartOffset) + Index;
-	BUFFER_TYPE OldValue;
-	InterlockedMin(DI_LOCAL(BufferUAV)[BufferIndex], Value, OldValue);
+	PUBLIC_TYPE OldValue;
+	InterlockedMax(DI_LOCAL(BufferUAV)[BufferIndex], Value, OldValue);
 #endif
 }
diff --git a/Engine/Plugins/Animation/DeformerGraph/Source/OptimusCore/Private/DataInterfaces/OptimusDataInterfaceRawBuffer.cpp b/Engine/Plugins/Animation/DeformerGraph/Source/OptimusCore/Private/DataInterfaces/OptimusDataInterfaceRawBuffer.cpp
--- a/Engine/Plugins/Animation/DeformerGraph/Source/OptimusCore/Private/DataInterfaces/OptimusDataInterfaceRawBuffer.cpp
+++ b/Engine/Plugins/Animation/DeformerGraph/Source/OptimusCore/Private/DataInterfaces/OptimusDataInterfaceRawBuffer.cpp
@@ -27,6 +27,28 @@ bool UOptimusRawBufferDataInterface::SupportsAtomics() const
 	return ValueType->Type == EShaderFundamentalType::Int;
 }
 
+FString UOptimusRawBufferDataInterface::GetRawType() const
+{
+	// Currently the only case for using a raw typed buffer is for vectors with size 3.
+	// We _may_ also need something for user structures if we they don't obey alignment restrictions.
+	if (ValueType->DimensionType == EShaderFundamentalDimensionType::Vector && ValueType->VectorElemCount == 3)
+	{
+		return FShaderValueType::Get(ValueType->Type)->ToString();
+	}
+	return {};
+}
+
+int32 UOptimusRawBufferDataInterface::GetRawStride() const
+{
+	// Currently the only case for using a raw typed buffer is for vectors with size 3.
+	// We _may_ also need something for user structures if we they don't obey alignment restrictions.
+	if (ValueType->DimensionType == EShaderFundamentalDimensionType::Vector && ValueType->VectorElemCount == 3)
+	{
+		return 4;
+	}
+	return 0;
+}
+
 TArray<FOptimusCDIPinDefinition> UOptimusRawBufferDataInterface::GetPinDefinitions() const
 {
 	// FIXME: Multi-level support by proxying through a data interface.
@@ -117,19 +139,27 @@ void UOptimusRawBufferDataInterface::GetShaderHash(FString& InOutKey) const
 
 void UOptimusRawBufferDataInterface::GetHLSL(FString& OutHLSL) const
 {
-	OutHLSL += TEXT("#define BUFFER_TYPE ");
-	OutHLSL += ValueType->ToString();
-	OutHLSL += TEXT(" \n");
+	TStringBuilder<512> StringBuilder;
 
-	if (SupportsAtomics()) { OutHLSL += TEXT("#define BUFFER_TYPE_SUPPORTS_ATOMIC 1\n"); }
-	if (UseSplitBuffers()) { OutHLSL += TEXT("#define BUFFER_SPLIT_READ_WRITE 1\n"); }
+	const FString PublicType = ValueType->ToString();
+	const FString RawType = GetRawType();
+	const bool bUseRawType = !RawType.IsEmpty();
 
-	OutHLSL += TEXT("#include \"/Plugin/Optimus/Private/DataInterfaceRawBuffer.ush\"\n");
+	StringBuilder.Appendf(TEXT("#define PUBLIC_TYPE %s\n"), *PublicType);
+	StringBuilder.Appendf(TEXT("#define BUFFER_TYPE %s\n"), bUseRawType ? *RawType : *PublicType);
+	StringBuilder.Appendf(TEXT("#define BUFFER_TYPE_RAW %d\n"), bUseRawType ? 1 : 0);
+	StringBuilder.Appendf(TEXT("#define BUFFER_TYPE_SUPPORTS_ATOMIC %d\n"), SupportsAtomics() ? 1 : 0);
+	StringBuilder.Appendf(TEXT("#define BUFFER_SPLIT_READ_WRITE %d\n"), UseSplitBuffers() ? 1 : 0);
 
-	if (SupportsAtomics()) { OutHLSL += TEXT("#undef BUFFER_TYPE_SUPPORTS_ATOMIC\n"); }
-	if (UseSplitBuffers()) { OutHLSL += TEXT("#undef BUFFER_SPLIT_READ_WRITE\n"); }
+	StringBuilder.Append(TEXT("#include \"/Plugin/Optimus/Private/DataInterfaceRawBuffer.ush\"\n"));
+	
+	StringBuilder.Append(TEXT("#undef PUBLIC_TYPE\n"));
+	StringBuilder.Append(TEXT("#undef BUFFER_TYPE\n"));
+	StringBuilder.Append(TEXT("#undef BUFFER_TYPE_RAW\n"));
+	StringBuilder.Append(TEXT("#undef BUFFER_TYPE_SUPPORTS_ATOMIC\n"));
+	StringBuilder.Append(TEXT("#undef BUFFER_SPLIT_READ_WRITE\n"));
 
-	OutHLSL += TEXT("#undef BUFFER_TYPE\n");
+	OutHLSL += StringBuilder.ToString();
 }
 
 FString UOptimusTransientBufferDataInterface::GetDisplayName() const
@@ -142,6 +172,8 @@ UComputeDataProvider* UOptimusTransientBufferDataInterface::CreateDataProvider(
 	) const
 {
 	UOptimusTransientBufferDataProvider *Provider = CreateProvider<UOptimusTransientBufferDataProvider>(InBinding);
+	Provider->ElementStride = ValueType->GetResourceElementSize();
+	Provider->RawStride = GetRawStride();
 	return Provider;
 }
 
@@ -292,7 +324,7 @@ FComputeDataProviderRenderProxy* UOptimusTransientBufferDataProvider::GetRenderP
 	TArray<int32> InvocationCounts;
 	GetInvocationElementCounts(InvocationCounts);
 	
-	return new FOptimusTransientBufferDataProviderProxy(InvocationCounts, ElementStride);
+	return new FOptimusTransientBufferDataProviderProxy(InvocationCounts, ElementStride, RawStride);
 }
 
 
@@ -312,25 +344,32 @@ FComputeDataProviderRenderProxy* UOptimusPersistentBufferDataProvider::GetRender
 	TArray<int32> InvocationCounts;
 	GetInvocationElementCounts(InvocationCounts);
 	
-	return new FOptimusPersistentBufferDataProviderProxy(InvocationCounts, ElementStride, BufferPool, ResourceName, LODIndex);
+	return new FOptimusPersistentBufferDataProviderProxy(InvocationCounts, ElementStride, RawStride, BufferPool, ResourceName, LODIndex);
 }
 
 
 FOptimusTransientBufferDataProviderProxy::FOptimusTransientBufferDataProviderProxy(
 	TArray<int32> InInvocationElementCounts,
-	int32 InElementStride
+	int32 InElementStride,
+	int32 InRawStride
 	) :
 	InvocationElementCounts(InInvocationElementCounts),
-	ElementStride(InElementStride)
+	ElementStride(InElementStride),
+	RawStride(InRawStride)
 {
 }
 
 
 void FOptimusTransientBufferDataProviderProxy::AllocateResources(FRDGBuilder& GraphBuilder)
 {
+	// If we are using a raw type alias for the buffer then we need to adjust stride and count.
+	check(RawStride == 0 || ElementStride % RawStride == 0);
+	const int32 Stride = RawStride ? RawStride : ElementStride;
+	const int32 ElementStrideMultiplier = RawStride ? ElementStride / RawStride : 1;
+
 	for (const int32 NumElements: InvocationElementCounts)
 	{
-		Buffer.Add(GraphBuilder.CreateBuffer(FRDGBufferDesc::CreateStructuredDesc(ElementStride, NumElements), TEXT("TransientBuffer"), ERDGBufferFlags::None));
+		Buffer.Add(GraphBuilder.CreateBuffer(FRDGBufferDesc::CreateStructuredDesc(Stride, NumElements * ElementStrideMultiplier), TEXT("TransientBuffer"), ERDGBufferFlags::None));
 		BufferSRV.Add(GraphBuilder.CreateSRV(Buffer.Last()));
 		BufferUAV.Add(GraphBuilder.CreateUAV(Buffer.Last()));
 	}
@@ -359,12 +398,14 @@ void FOptimusTransientBufferDataProviderProxy::GatherDispatchData(FDispatchSetup
 FOptimusPersistentBufferDataProviderProxy::FOptimusPersistentBufferDataProviderProxy(
 	TArray<int32> InInvocationElementCounts,
 	int32 InElementStride,
+	int32 InRawStride,
 	TSharedPtr<FOptimusPersistentBufferPool> InBufferPool,
 	FName InResourceName,
 	int32 InLODIndex
 	) :
 	InvocationElementCounts(InInvocationElementCounts),
 	ElementStride(InElementStride),
+	RawStride(InRawStride),
 	BufferPool(InBufferPool),
 	ResourceName(InResourceName),
 	LODIndex(InLODIndex)
@@ -376,7 +417,7 @@ void FOptimusPersistentBufferDataProviderProxy::AllocateResources(
 	FRDGBuilder& GraphBuilder
 	)
 {
-	BufferPool->GetResourceBuffers(GraphBuilder, ResourceName, LODIndex, ElementStride, InvocationElementCounts, Buffers);
+	BufferPool->GetResourceBuffers(GraphBuilder, ResourceName, LODIndex, ElementStride, RawStride, InvocationElementCounts, Buffers);
 	BufferUAVs.Reserve(Buffers.Num());
 	for (FRDGBufferRef BufferRef : Buffers)
 	{
diff --git a/Engine/Plugins/Animation/DeformerGraph/Source/OptimusCore/Private/DataInterfaces/OptimusDataInterfaceRawBuffer.h b/Engine/Plugins/Animation/DeformerGraph/Source/OptimusCore/Private/DataInterfaces/OptimusDataInterfaceRawBuffer.h
--- a/Engine/Plugins/Animation/DeformerGraph/Source/OptimusCore/Private/DataInterfaces/OptimusDataInterfaceRawBuffer.h
+++ b/Engine/Plugins/Animation/DeformerGraph/Source/OptimusCore/Private/DataInterfaces/OptimusDataInterfaceRawBuffer.h
@@ -65,6 +65,8 @@ public:
 protected:
 	virtual bool UseSplitBuffers() const { return true; }
 
+	int32 GetRawStride() const;
+
 	template<typename U>
 	U* CreateProvider(TObjectPtr<UObject> InBinding) const
 	{
@@ -75,12 +77,14 @@ protected:
 			Provider->ComponentSource = GetComponentSource();
 			Provider->DataDomain = DataDomain;
 			Provider->ElementStride = ValueType->GetResourceElementSize();
+			Provider->RawStride = GetRawStride();
 		}
 		return Provider;
 	}
 private:
 	const UOptimusComponentSource* GetComponentSource() const;
 	bool SupportsAtomics() const;
+	FString GetRawType() const;
 };
 
 
@@ -162,6 +166,9 @@ public:
 
 	UPROPERTY()
 	int32 ElementStride = 4;
+
+	UPROPERTY()
+	int32 RawStride = 0;
 };
 
 
@@ -208,7 +215,8 @@ class FOptimusTransientBufferDataProviderProxy :
 public:
 	FOptimusTransientBufferDataProviderProxy(
 		TArray<int32> InInvocationElementCounts,
-		int32 InElementStride
+		int32 InElementStride,
+		int32 InRawStride
 		);
 
 	//~ Begin FComputeDataProviderRenderProxy Interface
@@ -219,6 +227,7 @@ public:
 private:
 	const TArray<int32> InvocationElementCounts;
 	const int32 ElementStride;
+	const int32 RawStride;
 
 	TArray<FRDGBuffer*> Buffer;
 	TArray<FRDGBufferSRV*> BufferSRV;
@@ -233,6 +242,7 @@ public:
 	FOptimusPersistentBufferDataProviderProxy(
 		TArray<int32> InInvocationElementCounts,
 		int32 InElementStride,
+		int32 InRawStride,
 		TSharedPtr<FOptimusPersistentBufferPool> InBufferPool,
 		FName InResourceName,
 		int32 InLODIndex
@@ -246,6 +256,7 @@ public:
 private:
 	const TArray<int32> InvocationElementCounts;
 	const int32 ElementStride;
+	const int32 RawStride;
 	const TSharedPtr<FOptimusPersistentBufferPool> BufferPool;
 	const FName ResourceName;
 	const int32 LODIndex;
diff --git a/Engine/Plugins/Animation/DeformerGraph/Source/OptimusCore/Private/OptimusDeformerInstance.cpp b/Engine/Plugins/Animation/DeformerGraph/Source/OptimusCore/Private/OptimusDeformerInstance.cpp
--- a/Engine/Plugins/Animation/DeformerGraph/Source/OptimusCore/Private/OptimusDeformerInstance.cpp
+++ b/Engine/Plugins/Animation/DeformerGraph/Source/OptimusCore/Private/OptimusDeformerInstance.cpp
@@ -27,6 +27,7 @@ void FOptimusPersistentBufferPool::GetResourceBuffers(
 	FName InResourceName,
 	int32 InLODIndex,
 	int32 InElementStride,
+	int32 InRawStride,
 	TArray<int32> const& InElementCounts,
 	TArray<FRDGBufferRef>& OutBuffers)
 {
@@ -40,9 +41,14 @@ void FOptimusPersistentBufferPool::GetResourceBuffers(
 		TArray<FOptimusPersistentStructuredBuffer> ResourceBuffers;
 		ResourceBuffers.Reserve(InElementCounts.Num());
 
+		// If we are using a raw type alias for the buffer then we need to adjust stride and count.
+		check(InRawStride == 0 || InElementStride % InRawStride == 0);
+		const int32 Stride = InRawStride ? InRawStride : InElementStride;
+		const int32 ElementStrideMultiplier = InRawStride ? InElementStride / InRawStride : 1;
+
 		for (int32 Index = 0; Index < InElementCounts.Num(); Index++)
 		{
-			FRDGBufferDesc BufferDesc = FRDGBufferDesc::CreateStructuredDesc(InElementStride, InElementCounts[Index]);
+			FRDGBufferDesc BufferDesc = FRDGBufferDesc::CreateStructuredDesc(Stride, InElementCounts[Index] * ElementStrideMultiplier);
 			FRDGBufferRef Buffer = GraphBuilder.CreateBuffer(BufferDesc, TEXT("FOptimusPersistentBuffer"), ERDGBufferFlags::None);
 			OutBuffers.Add(Buffer);
 
diff --git a/Engine/Plugins/Animation/DeformerGraph/Source/OptimusCore/Public/OptimusDeformerInstance.h b/Engine/Plugins/Animation/DeformerGraph/Source/OptimusCore/Public/OptimusDeformerInstance.h
--- a/Engine/Plugins/Animation/DeformerGraph/Source/OptimusCore/Public/OptimusDeformerInstance.h
+++ b/Engine/Plugins/Animation/DeformerGraph/Source/OptimusCore/Public/OptimusDeformerInstance.h
@@ -32,6 +32,7 @@ public:
 		FName InResourceName,
 		int32 InLODIndex,
 		int32 InElementStride,
+		int32 InRawStride,
 		TArray<int32> const& InElementCounts,
 		TArray<FRDGBuffer*>& OutBuffers );
 
diff --git a/Engine/Plugins/Runtime/ComputeFramework/Shaders/Private/ComputeKernelCommon.ush b/Engine/Plugins/Runtime/ComputeFramework/Shaders/Private/ComputeKernelCommon.ush
--- a/Engine/Plugins/Runtime/ComputeFramework/Shaders/Private/ComputeKernelCommon.ush
+++ b/Engine/Plugins/Runtime/ComputeFramework/Shaders/Private/ComputeKernelCommon.ush
@@ -34,3 +34,51 @@
 
 #define DI_IMPL_READ(Name, ReturnType, ...) ReturnType DI_LOCAL(Name)(__VA_ARGS__)
 #define DI_IMPL_WRITE(Name, ...) void DI_LOCAL(Name)(__VA_ARGS__)
+
+// Implement ReadRawBuffer/WriteRawBuffer functions for plain old types that need it.
+// Due to restrictions on some platforms we can't use StructureBuffer with unaligned types. That includes all vector<T, 3> types.
+// Similarly user structure types that don't follow alignment restrictions might have to implement similar functions to read from raw buffers.
+
+void ReadRawBuffer(in StructuredBuffer<float> InBuffer, in uint InBufferIndex, out float3 OutValue)
+{
+	for (uint Index = 0; Index < 3; ++Index)
+	{
+		OutValue[Index] = InBuffer[InBufferIndex * 3 + Index];
+	}
+}
+void ReadRawBuffer(in RWStructuredBuffer<float> InBuffer, in uint InBufferIndex, out float3 OutValue)
+{
+	for (uint Index = 0; Index < 3; ++Index)
+	{
+		OutValue[Index] = InBuffer[InBufferIndex * 3 + Index];
+	}
+}
+void WriteRawBuffer(in RWStructuredBuffer<float> InBuffer, in uint InBufferIndex, in float3 InValue)
+{
+	for (uint Index = 0; Index < 3; ++Index)
+	{
+		InBuffer[InBufferIndex * 3 + Index] = InValue[Index];
+	}
+}
+
+void ReadRawBuffer(in StructuredBuffer<int> InBuffer, in uint InBufferIndex, out int3 OutValue)
+{
+	for (uint Index = 0; Index < 3; ++Index)
+	{
+		OutValue[Index] = InBuffer[InBufferIndex * 3 + Index];
+	}
+}
+void ReadRawBuffer(in RWStructuredBuffer<int> InBuffer, in uint InBufferIndex, out int3 OutValue)
+{
+	for (uint Index = 0; Index < 3; ++Index)
+	{
+		OutValue[Index] = InBuffer[InBufferIndex * 3 + Index];
+	}
+}
+void WriteRawBuffer(in RWStructuredBuffer<int> InBuffer, in uint InBufferIndex, in int3 InValue)
+{
+	for (uint Index = 0; Index < 3; ++Index)
+	{
+		InBuffer[InBufferIndex * 3 + Index] = InValue[Index];
+	}
+}
