Optimus: Fix crash when placed mesh deformer instance is reset during graph compilation.

diff --git a/Engine/Plugins/Animation/DeformerGraph/Source/OptimusCore/Private/OptimusDeformerInstance.cpp b/Engine/Plugins/Animation/DeformerGraph/Source/OptimusCore/Private/OptimusDeformerInstance.cpp
--- a/Engine/Plugins/Animation/DeformerGraph/Source/OptimusCore/Private/OptimusDeformerInstance.cpp
+++ b/Engine/Plugins/Animation/DeformerGraph/Source/OptimusCore/Private/OptimusDeformerInstance.cpp
@@ -117,7 +117,7 @@ void UOptimusDeformerInstanceSettings::RefreshComponentBindings(
 		ExistingBindings.Add(Binding.ProviderName, Binding.ActorComponent.Get());
 	}
 	
-	AActor* Actor = InMeshComponent->GetOwner();
+	AActor* Actor = InMeshComponent != nullptr ? InMeshComponent->GetOwner() : nullptr;
 	TSet<UActorComponent*> ComponentsUsed;
 
 	const TArray<UOptimusComponentSourceBinding*>& ComponentBindings = InDeformer->GetComponentBindings();
