Optimus: Fix missing bindings on placed blueprint actors.
Still need to fix the issue that means that blueprint actors can't fully setup bindings.

diff --git a/Engine/Plugins/Animation/DeformerGraph/Source/OptimusCore/Private/OptimusDeformer.cpp b/Engine/Plugins/Animation/DeformerGraph/Source/OptimusCore/Private/OptimusDeformer.cpp
--- a/Engine/Plugins/Animation/DeformerGraph/Source/OptimusCore/Private/OptimusDeformer.cpp
+++ b/Engine/Plugins/Animation/DeformerGraph/Source/OptimusCore/Private/OptimusDeformer.cpp
@@ -2308,7 +2308,7 @@ UMeshDeformerInstance* UOptimusDeformer::CreateInstance(
 	Instance->SetupFromDeformer(this);
 
 	// Make sure all the instances know when we finish compiling so they can update their local state to match.
-	CompileEndDelegate.AddUObject(Instance, &UOptimusDeformerInstance::SetupFromDeformer, true);
+	CompileEndDelegate.AddUObject(Instance, &UOptimusDeformerInstance::SetupFromDeformer);
 	ConstantValueUpdateDelegate.AddUObject(Instance, &UOptimusDeformerInstance::SetConstantValueDirect);
 	SetAllInstancesCanbeActiveDelegate.AddUObject(Instance, &UOptimusDeformerInstance::SetCanBeActive);
 
diff --git a/Engine/Plugins/Animation/DeformerGraph/Source/OptimusCore/Private/OptimusDeformerInstance.cpp b/Engine/Plugins/Animation/DeformerGraph/Source/OptimusCore/Private/OptimusDeformerInstance.cpp
--- a/Engine/Plugins/Animation/DeformerGraph/Source/OptimusCore/Private/OptimusDeformerInstance.cpp
+++ b/Engine/Plugins/Animation/DeformerGraph/Source/OptimusCore/Private/OptimusDeformerInstance.cpp
@@ -269,18 +269,15 @@ void UOptimusDeformerInstance::SetInstanceSettings(UOptimusDeformerInstanceSetti
 }
 
 
-void UOptimusDeformerInstance::SetupFromDeformer(
-	UOptimusDeformer* InDeformer,
-	const bool bInRefreshBindings
-	)
+void UOptimusDeformerInstance::SetupFromDeformer(UOptimusDeformer* InDeformer)
 {
 	// If we're doing a recompile, ditch all stored render resources.
 	ReleaseResources();
 
-	// Update the component bindings before creating data providers. The bindings are in the same order
-	// as the component bindings in the deformer.
+	// Update the component bindings before creating data providers. 
+	// The bindings are in the same order as the component bindings in the deformer.
 	UOptimusDeformerInstanceSettings* InstanceSettingsPtr = InstanceSettings.Get(); 
-	if (bInRefreshBindings && InstanceSettingsPtr)
+	if (InstanceSettingsPtr)
 	{
 		InstanceSettingsPtr->RefreshComponentBindings(InDeformer, MeshComponent.Get());
 	}
diff --git a/Engine/Plugins/Animation/DeformerGraph/Source/OptimusCore/Public/OptimusDeformerInstance.h b/Engine/Plugins/Animation/DeformerGraph/Source/OptimusCore/Public/OptimusDeformerInstance.h
--- a/Engine/Plugins/Animation/DeformerGraph/Source/OptimusCore/Public/OptimusDeformerInstance.h
+++ b/Engine/Plugins/Animation/DeformerGraph/Source/OptimusCore/Public/OptimusDeformerInstance.h
@@ -154,15 +154,7 @@ public:
 	 * Needs to be called after the UOptimusDeformer creates this instance, and whenever the instance is invalidated.
 	 * Invalidation happens whenever any bound Data Providers become invalid.
 	 */
-	void SetupFromDeformer(
-		UOptimusDeformer* InDeformer,
-		const bool bInRefreshBindings = false
-		);
-
-	void SetupFromDeformerAndRefreshBindings(UOptimusDeformer* InDeformer)
-	{
-		SetupFromDeformer(InDeformer, true);
-	}
+	void SetupFromDeformer(UOptimusDeformer* InDeformer);
 
 	/** Set the value of a boolean variable. */
 	UFUNCTION(BlueprintPure, Category="Deformer", meta=(DisplayName="Set Variable (bool)"))
diff --git a/Engine/Plugins/Animation/DeformerGraph/Source/OptimusEditor/Private/OptimusDetailsCustomization.cpp b/Engine/Plugins/Animation/DeformerGraph/Source/OptimusEditor/Private/OptimusDetailsCustomization.cpp
--- a/Engine/Plugins/Animation/DeformerGraph/Source/OptimusEditor/Private/OptimusDetailsCustomization.cpp
+++ b/Engine/Plugins/Animation/DeformerGraph/Source/OptimusEditor/Private/OptimusDetailsCustomization.cpp
@@ -1593,7 +1593,7 @@ void FOptimusDeformerInstanceComponentBindingCustomization::CustomizeHeader(
 	.ValueContent()
 	[
 		SNew(SComboBox<FComponentHandle>)
-		.IsEnabled(Binding)
+		.IsEnabled(Binding && !Binding->IsPrimaryBinding())
 		.OptionsSource(&ComponentHandles)
 		.InitiallySelectedItem(SelectedComponentHandle)
 		.OnGenerateWidget_Lambda([](const FComponentHandle InComponentHandle)
