Optimus: Added variable default values.

diff --git a/Engine/Plugins/Animation/DeformerGraph/Source/OptimusCore/Private/DataInterfaces/OptimusDataInterfaceGraph.cpp b/Engine/Plugins/Animation/DeformerGraph/Source/OptimusCore/Private/DataInterfaces/OptimusDataInterfaceGraph.cpp
--- a/Engine/Plugins/Animation/DeformerGraph/Source/OptimusCore/Private/DataInterfaces/OptimusDataInterfaceGraph.cpp
+++ b/Engine/Plugins/Animation/DeformerGraph/Source/OptimusCore/Private/DataInterfaces/OptimusDataInterfaceGraph.cpp
@@ -137,7 +137,10 @@ FOptimusGraphDataProviderProxy::FOptimusGraphDataProviderProxy(UOptimusDeformerI
 		if (Variable.Value.Num())
 		{
 			// Use the constant value.
-			FMemory::Memcpy(&ParameterData[Variable.Offset], Variable.Value.GetData(), Variable.Value.Num());
+			if (ensure(ParameterData.Num() >= Variable.Offset + Variable.Value.Num()))
+			{
+				FMemory::Memcpy(&ParameterData[Variable.Offset], Variable.Value.GetData(), Variable.Value.Num());
+			}
 		}
 		else
 		{
@@ -149,8 +152,10 @@ FOptimusGraphDataProviderProxy::FOptimusGraphDataProviderProxy(UOptimusDeformerI
 				{
 					if (Variable.ValueType == VariableValue->DataType->ShaderValueType && Variable.Name == VariableValue->VariableName.GetPlainNameString())
 					{
-						TArrayView<uint8> ParameterEntry(&ParameterData[Variable.Offset], VariableValue->DataType->ShaderValueSize);
-						VariableValue->DataType->ConvertPropertyValueToShader(VariableValue->ValueData, ParameterEntry);
+						if (ensure(ParameterData.Num() >= Variable.Offset + VariableValue->ValueData.Num()))
+						{
+							FMemory::Memcpy(&ParameterData[Variable.Offset], VariableValue->ValueData.GetData(), VariableValue->ValueData.Num());
+						}
 						break;
 					}
 				}
diff --git a/Engine/Plugins/Animation/DeformerGraph/Source/OptimusCore/Private/OptimusDeformerInstance.cpp b/Engine/Plugins/Animation/DeformerGraph/Source/OptimusCore/Private/OptimusDeformerInstance.cpp
--- a/Engine/Plugins/Animation/DeformerGraph/Source/OptimusCore/Private/OptimusDeformerInstance.cpp
+++ b/Engine/Plugins/Animation/DeformerGraph/Source/OptimusCore/Private/OptimusDeformerInstance.cpp
@@ -350,16 +350,21 @@ void UOptimusDeformerInstance::SetupFromDeformer(UOptimusDeformer* InDeformer)
 	for (const UOptimusVariableDescription* VariableDescription : InDeformer->GetVariables())
 	{
 		if (!VariableDescription)
+		{
 			continue;
+		}
 		if (Visited.Contains(VariableDescription))
+		{
 			continue;
+		}
 		Visited.Add(VariableDescription);
+		
 		UOptimusVariableDescription* VariableDescriptionCopy = NewObject<UOptimusVariableDescription>();
 		VariableDescriptionCopy->Guid = VariableDescription->Guid;
 		VariableDescriptionCopy->VariableName = VariableDescription->VariableName;
 		VariableDescriptionCopy->DataType = VariableDescription->DataType;
+		VariableDescriptionCopy->DefaultValue = nullptr; // No need to copy the default value.
 		VariableDescriptionCopy->ValueData = VariableDescription->ValueData;
-		VariableDescriptionCopy->ResetValueDataSize();
 		Variables->Descriptions.Add(VariableDescriptionCopy);
 	}
 
@@ -373,6 +378,7 @@ void UOptimusDeformerInstance::SetupFromDeformer(UOptimusDeformer* InDeformer)
 void UOptimusDeformerInstance::SetCanBeActive(bool bInCanBeActive)
 {
 	bCanBeActive = bInCanBeActive;
+
 }
 
 void UOptimusDeformerInstance::AllocateResources()
diff --git a/Engine/Plugins/Animation/DeformerGraph/Source/OptimusCore/Private/OptimusVariableDescription.cpp b/Engine/Plugins/Animation/DeformerGraph/Source/OptimusCore/Private/OptimusVariableDescription.cpp
--- a/Engine/Plugins/Animation/DeformerGraph/Source/OptimusCore/Private/OptimusVariableDescription.cpp
+++ b/Engine/Plugins/Animation/DeformerGraph/Source/OptimusCore/Private/OptimusVariableDescription.cpp
@@ -4,10 +4,13 @@
 
 #include "OptimusDeformer.h"
 #include "OptimusHelpers.h"
+#include "OptimusValueContainer.h"
 
 
 void UOptimusVariableDescription::ResetValueDataSize()
 {
+	DefaultValue = UOptimusValueContainer::MakeValueContainer(this, DataType);
+
 	if (DataType->CanCreateProperty())
 	{
 		// Create a temporary property from the type so that we can get the size of the
@@ -51,12 +54,17 @@ void UOptimusVariableDescription::PostEditChangeProperty(FPropertyChangedEvent&
 		UOptimusDeformer* Deformer = GetOwningDeformer();
 		if (ensure(Deformer))
 		{
-			// Set the variable type again, so that we can remove any links that are now
-			// type-incompatible.
+			// Set the variable type again, so that we can remove any links that are now type-incompatible.
 			constexpr bool bForceChange = true;
 			Deformer->SetVariableDataType(this, DataType, bForceChange);
 		}
 	}
+	else if (PropertyName == GET_MEMBER_NAME_CHECKED(UOptimusVariableDescription, DefaultValue))
+	{
+		// Store the default shader value.
+		FShaderValueType::FValue Value = DefaultValue->GetShaderValue();
+		ValueData = MoveTemp(Value.ShaderValue);
+	}
 }
 
 
diff --git a/Engine/Plugins/Animation/DeformerGraph/Source/OptimusCore/Public/OptimusVariableDescription.h b/Engine/Plugins/Animation/DeformerGraph/Source/OptimusCore/Public/OptimusVariableDescription.h
--- a/Engine/Plugins/Animation/DeformerGraph/Source/OptimusCore/Public/OptimusVariableDescription.h
+++ b/Engine/Plugins/Animation/DeformerGraph/Source/OptimusCore/Public/OptimusVariableDescription.h
@@ -13,6 +13,7 @@
 
 
 class UOptimusDeformer;
+class UOptimusValueContainer;
 
 
 USTRUCT()
@@ -57,10 +58,6 @@ public:
 	UPROPERTY()
 	FGuid Guid;
 
-	/** The actual binary data of the value that was written (or the default value) */
-	UPROPERTY()
-	TArray<uint8> ValueData;
-	
 	/** Name of the variable */
 	UPROPERTY(EditAnywhere, Category = VariableDefinition)
 	FName VariableName;
@@ -69,6 +66,14 @@ public:
 	UPROPERTY(EditAnywhere, Category = VariableDefinition, meta=(UseInVariable))
 	FOptimusDataTypeRef DataType;
 
+	/** The default value for the variable. */
+	UPROPERTY(EditAnywhere, Category = VariableDefinition, meta = (EditInLine))
+	TObjectPtr<UOptimusValueContainer> DefaultValue = nullptr;
+
+	/** Cached shader value binary data. */
+	UPROPERTY()
+	TArray<uint8> ValueData;
+
 #if WITH_EDITOR
 	void PostEditChangeProperty(struct FPropertyChangedEvent& PropertyChangedEvent) override;
 	void PreEditUndo() override;
