Optimus: Fix up place where we were storing variable data as source type instead of shader type.
#rnx

diff --git a/Engine/Plugins/Animation/DeformerGraph/Source/OptimusCore/Private/OptimusDeformerInstance.cpp b/Engine/Plugins/Animation/DeformerGraph/Source/OptimusCore/Private/OptimusDeformerInstance.cpp
--- a/Engine/Plugins/Animation/DeformerGraph/Source/OptimusCore/Private/OptimusDeformerInstance.cpp
+++ b/Engine/Plugins/Animation/DeformerGraph/Source/OptimusCore/Private/OptimusDeformerInstance.cpp
@@ -447,10 +447,11 @@ namespace
 			if (VariableDesc->VariableName == InVariableName && VariableDesc->DataType == WantedType)
 			{
 				TUniquePtr<FProperty> Property(WantedType->CreateProperty(nullptr, NAME_None));
-				if (ensure(Property->GetSize() == sizeof(T)) &&
-					ensure(Property->GetSize() == VariableDesc->ValueData.Num()))
+				if (ensure(Property->GetSize() == sizeof(T)))
 				{
-					FPlatformMemory::Memcpy(VariableDesc->ValueData.GetData(), ValueBytes, sizeof(T));
+					FShaderValueType::FValue ValueResult = WantedType->MakeShaderValue();
+					WantedType->ConvertPropertyValueToShader(TArrayView<const uint8>((const uint8*)&InValue, sizeof(T)), ValueResult);
+					VariableDesc->ValueData = MoveTemp(ValueResult.ShaderValue);
 				}
 
 				return true;
