Get JiggleSpring deformer working.  Implement per vertex parameter remapping from import topology to render topology. Implement strain limiter. Fix math issues in the shader. Move bulk of the shader to the data interface ush file.

diff --git a/Engine/Plugins/Animation/DeformerGraph/Shaders/Private/DataInterfaceJiggleSpring.ush b/Engine/Plugins/Animation/DeformerGraph/Shaders/Private/DataInterfaceJiggleSpring.ush
--- a/Engine/Plugins/Animation/DeformerGraph/Shaders/Private/DataInterfaceJiggleSpring.ush
+++ b/Engine/Plugins/Animation/DeformerGraph/Shaders/Private/DataInterfaceJiggleSpring.ush
@@ -3,13 +3,21 @@
 #include "/Plugin/ComputeFramework/Private/ComputeKernelCommon.ush"
 
 uint DI_LOCAL(NumVertices);
+uint DI_LOCAL(BaseVertexIndex); // InputStreamStart
+
+uint DI_LOCAL(NumVertexMap);
 uint DI_LOCAL(NumStiffnessWeights);
 uint DI_LOCAL(NumDampingWeights);
+uint DI_LOCAL(NumMaxStretchWeights);
+
 float DI_LOCAL(Stiffness);
 float DI_LOCAL(Damping);
+float DI_LOCAL(MaxStretch);
+
 Buffer<float> DI_LOCAL(StiffnessWeightsBuffer);
 Buffer<float> DI_LOCAL(DampingWeightsBuffer);
-
+Buffer<uint> DI_LOCAL(VertexMapBuffer);
+Buffer<float> DI_LOCAL(MaxStretchWeightsBuffer);
 
 DI_IMPL_READ(ReadNumVertices, uint, )
 {
@@ -19,10 +27,21 @@ DI_IMPL_READ(ReadNumVertices, uint, )
 DI_IMPL_READ(ReadStiffness, float, uint VertexIndex)
 {
 #if ENABLE_DEFORMER_JIGGLE_SPRING
-	if(VertexIndex < DI_LOCAL(NumStiffnessWeights))
-		return DI_LOCAL(StiffnessWeightsBuffer)[VertexIndex] * DI_LOCAL(Stiffness);
+	if( //DI_LOCAL(NumStiffnessWeights) == DI_LOCAL(NumVertices) && 
+		VertexIndex < DI_LOCAL(NumStiffnessWeights))
+	{
+		uint BufferIndex = VertexIndex + DI_LOCAL(BaseVertexIndex);
+		return DI_LOCAL(StiffnessWeightsBuffer)[BufferIndex] * DI_LOCAL(Stiffness);
+	}
+	//else if(VertexIndex < DI_LOCAL(NumVertexMap))
+	//{
+	//	uint BufferIndex = DI_LOCAL(VertexMapBuffer)[VertexIndex + DI_LOCAL(BaseVertexIndex)];
+	//	return DI_LOCAL(StiffnessWeightsBuffer)[BufferIndex] * DI_LOCAL(Stiffness);
+	//}
 	else
+	{
 		return DI_LOCAL(Stiffness);
+	}
 #else
 	return 0.0;
 #endif
@@ -31,12 +50,206 @@ DI_IMPL_READ(ReadStiffness, float, uint VertexIndex)
 DI_IMPL_READ(ReadDamping, float, uint VertexIndex)
 {
 #if ENABLE_DEFORMER_JIGGLE_SPRING
-	if(VertexIndex < DI_LOCAL(NumDampingWeights))
-		return DI_LOCAL(DampingWeightsBuffer)[VertexIndex] * DI_LOCAL(Damping);
+	if( //DI_LOCAL(NumStiffnessWeights) == DI_LOCAL(NumVertices) && 
+		VertexIndex < DI_LOCAL(NumStiffnessWeights))
+	{
+		uint BufferIndex = VertexIndex + DI_LOCAL(BaseVertexIndex);
+		return DI_LOCAL(DampingWeightsBuffer)[BufferIndex] * DI_LOCAL(Damping);
+	}
+	//else if(VertexIndex < DI_LOCAL(NumVertexMap))
+	//{
+	//	uint BufferIndex = DI_LOCAL(VertexMapBuffer)[VertexIndex + DI_LOCAL(BaseVertexIndex)];
+	//	return DI_LOCAL(DampingWeightsBuffer)[BufferIndex] * DI_LOCAL(Damping);
+	//}
 	else
+	{
 		return DI_LOCAL(Damping);
+	}
 #else
 	return 0.0;
 #endif
 }
 
+DI_IMPL_READ(ReadBufferIndex, uint, uint VertexIndex)
+{
+#if ENABLE_DEFORMER_JIGGLE_SPRING
+	if(VertexIndex < DI_LOCAL(NumVertexMap))
+	{
+		uint BufferIndex = DI_LOCAL(VertexMapBuffer)[VertexIndex + DI_LOCAL(BaseVertexIndex)];
+		return BufferIndex;
+	}
+#else
+	return 0;
+#endif
+}
+
+DI_IMPL_READ(ReadNonZeroStiffnessNumber, uint, uint VertexIndex)
+{
+	uint Count = 0;
+#if ENABLE_DEFORMER_JIGGLE_SPRING
+	for(uint i = 0; i <= VertexIndex; i++)
+    {
+		if(DI_LOCAL(ReadStiffness)(VertexIndex) > 0.) 
+			Count++;
+	}
+#endif
+	return Count;
+}
+
+DI_IMPL_READ(ReadMaxStretch, float, uint VertexIndex)
+{
+#if ENABLE_DEFORMER_JIGGLE_SPRING
+	if(VertexIndex < DI_LOCAL(NumMaxStretchWeights))
+	{
+		uint BufferIndex = VertexIndex + DI_LOCAL(BaseVertexIndex);
+		return abs(DI_LOCAL(MaxStretchWeightsBuffer[BufferIndex]) * DI_LOCAL(MaxStretch));
+	}
+	return abs(DI_LOCAL(MaxStretch));
+#else
+	return 100.0;
+#endif
+}
+
+
+struct FAnalyticSpring
+{
+	float Dt; // struct can't be empty, so we have this member
+	
+	// Advance analytic spring sim by Dt.
+	void
+	AdvanceAnalyticSpringSim(
+		uint Index,
+		float3 PrevSkinPos, 
+		float3 CurrSkinPos, 
+		float4 SkinnedTangentX, 
+		float4 SkinnedTangentZ, 
+		float3 PrevDynPos, 
+		float3 PrevDynVel,
+		out float3 OutX,
+		out float3 OutV,
+		out float4 OutTanX,
+		out float4 OutTanZ)
+	{
+#if ENABLE_DEFORMER_JIGGLE_SPRING
+		// Sim parameters
+		float K = DI_LOCAL(ReadStiffness)(Index);
+		float D = DI_LOCAL(ReadDamping)(Index);
+		if(K <= 0.0 || D <= 0.0)
+		{
+			// Set outputs to skinning
+			OutX = CurrSkinPos;
+			OutV = float3(0,0,0);
+			OutTanX = SkinnedTangentX;
+			OutTanZ = SkinnedTangentZ;
+			return;
+		}
+		if(D == sqrt(4*K)) // avoid critical damping case
+			D += .01;
+
+		// History
+		float3 X = PrevDynPos;
+		float3 V = PrevDynVel;
+
+		// Interpolate (linear)
+		float3 Q = { 0.f, 0.f, 0.f };
+		float3 A = { 0.f, 0.f, 0.f };
+		float3 B = (CurrSkinPos - PrevSkinPos) / Dt; // src: (Xn+1 - Xn) / dt - Ryan wrong Dt here?  Finite diff...
+		float3 C = CurrSkinPos;
+
+		// Precalculations
+		float _Dt2 = Dt * Dt;
+		float _Dt3 = _Dt2 * Dt;
+		float _D2 = D * D;
+		float _D2m4K = _D2 - 4*K;
+
+		float3 _6QdK = 6 * Q / K;
+		float3 _2AdK = 2 * A / K;
+		float3 _6QDdK2 = 6 * Q * D / (K * K);
+
+		// Do spring
+		//float3 DynamicX;
+		if(_D2m4K > 0) // Over damped
+		{
+			float Delta_sqrt = sqrt(_D2m4K);
+			float Root1 = (-D + Delta_sqrt) / 2;
+			float Root2 = (-D - Delta_sqrt) / 2;
+			float _ExpRoot1Dt = exp(Root1 * Dt);
+			float _ExpRoot2Dt = exp(Root2 * Dt);
+
+			float3 C1 = (V - B + _6QdK - Root2 * (X - C + _2AdK - _6QDdK2)) / Delta_sqrt;
+			float3 C2 =(-V + B - _6QdK + Root1 * (X - C + _2AdK - _6QDdK2)) / Delta_sqrt;
+			float3 _C1ExpRoot1Dt = C1 * _ExpRoot1Dt;
+			float3 _C2ExpRoot2Dt = C2 * _ExpRoot2Dt;
+
+			X = _C1ExpRoot1Dt + _C2ExpRoot2Dt + Q*_Dt3 + A*_Dt2 + (B-6*Q/K)*Dt + C-2*A/K + _6QDdK2;
+			V = _C1ExpRoot1Dt*Root1 + _C2ExpRoot2Dt*Root2 + 3*Q*_Dt2 + 2*A*Dt + B-_6QdK;
+		}
+		else //if(_D2m4K < 0) // Under damped
+		{
+			float Delta_sqrt = sqrt(4*K - _D2);
+			float Beta = Delta_sqrt / 2;
+			float Alpha = -D / 2;
+
+			float3 C1 = X - C + _2AdK - _6QDdK2;
+			float3 C2 = (2 * V - 2 * B + 12 * Q / K + D * C1) / Delta_sqrt;
+
+			float _ExpAlphaDt = exp(Alpha * Dt);
+			float _BetaDt = Beta * Dt;
+			float _CosBetaDt = cos(_BetaDt);
+			float _SinBetaDt = sin(_BetaDt);
+
+			X = _ExpAlphaDt * (C1 * _CosBetaDt + C2 * _SinBetaDt) +
+				Q * _Dt3 + 
+				A * _Dt2 + 
+				(B-6*Q/K)*Dt + 
+				C-_2AdK + 
+				_6QDdK2;
+			V = C1 * Alpha * _ExpAlphaDt * _CosBetaDt -
+				C1 * _ExpAlphaDt * Beta * sin(Beta * Dt) + 
+				C2 * _ExpAlphaDt * _SinBetaDt + 
+				C2 * _ExpAlphaDt * Beta * _CosBetaDt +
+				3 * Q * _Dt2 + 
+				2 * A * Dt + 
+				B - _6QdK;
+		}
+		//else // Critically damped (which should be rare)
+		//{
+		//    float Root = -D / 2;
+		//    float3 C1 = X - C + _2AdK - _6QDdK2;
+		//    float3 C2 = V - B + 6 * Q / K + D / 2 * C1;
+		//    float _ExpRootDt = exp(Root * Dt);
+		//    DynamicX = C1 * _ExpRootDt + C2 * Dt * _ExpRootDt + Q * _Dt3 + A * _Dt2 + (B-6*Q/K) * Dt + C - _2AdK + _6QDdK2;
+		//    V = C1 * Root * _ExpRootDt + C2 * _ExpRootDt + Dt * Root * _ExpRootDt + 3 * Q * _Dt2 + 2 * A * Dt + B - _6QdK;
+		//}
+
+		// Clamp deformation
+		float3 Dir = X - CurrSkinPos;
+		float DirLen = length(Dir);
+		float MaxLen = DI_LOCAL(ReadMaxStretch)(Index);
+		if(DirLen > MaxLen)
+		{
+			Dir /= DirLen;
+			Dir *= MaxLen;
+			X = CurrSkinPos + Dir;
+		}
+
+		OutX = X;
+		OutV = V;
+		OutTanX = SkinnedTangentX;
+		OutTanZ = SkinnedTangentZ;
+#else
+		// Set outputs to skinning
+		OutX = CurrSkinPos;
+		OutV = float3(0,0,0);
+		OutTanX = SkinnedTangentX;
+		OutTanZ = SkinnedTangentZ;
+#endif
+	}
+
+};
+
+DI_IMPL_READ(ReadAnalyticSpring, FAnalyticSpring, )
+{
+	FAnalyticSpring Spring;
+	return Spring;
+}
\ No newline at end of file
diff --git a/Engine/Plugins/Animation/DeformerGraph/Source/OptimusCore/Private/DataInterfaces/OptimusDataInterfaceJiggleSpring.cpp b/Engine/Plugins/Animation/DeformerGraph/Source/OptimusCore/Private/DataInterfaces/OptimusDataInterfaceJiggleSpring.cpp
--- a/Engine/Plugins/Animation/DeformerGraph/Source/OptimusCore/Private/DataInterfaces/OptimusDataInterfaceJiggleSpring.cpp
+++ b/Engine/Plugins/Animation/DeformerGraph/Source/OptimusCore/Private/DataInterfaces/OptimusDataInterfaceJiggleSpring.cpp
@@ -6,13 +6,16 @@
 #include "ComputeFramework/ComputeKernelPermutationVector.h"
 #include "ComputeFramework/ShaderParamTypeDefinition.h"
 #include "GenericPlatform/GenericPlatformFile.h"
+#include "GenericPlatform/GenericPlatformMath.h"
 #include "HAL/PlatformFileManager.h"
 #include "Misc/FileHelper.h"
 #include "Misc/Paths.h"
 #include "OptimusDataDomain.h"
+#include "OptimusDataTypeRegistry.h"
 #include "ShaderParameterMetadataBuilder.h"
 #include "RenderGraphBuilder.h"
 #include "Rendering/SkeletalMeshLODRenderData.h"
+#include "Rendering/SkeletalMeshModel.h"
 #include "Rendering/SkeletalMeshRenderData.h"
 #include "SkeletalMeshDeformerHelpers.h"
 #include "SkeletalRenderPublic.h"
@@ -21,7 +24,136 @@
 // FOptimusJiggleSpringParameters
 //
 
-bool FOptimusJiggleSpringParameters::ReadWeightsFile(const FFilePath& FilePath, TArray<float>& Values) const
+bool 
+FOptimusJiggleSpringParameters::ReadWeightsFile(
+	const FFilePath& FilePath, 
+	//TArray<TTuple<FString,TArray<float>>>& SectionValues) const
+	TArray<FString>& SectionNames,
+	TArray<TArray<float>>& SectionValues) const
+{
+#if WITH_EDITOR
+	if (FilePath.FilePath.IsEmpty())
+	{
+		return false;
+	}
+
+	FPlatformFileManager& FileManager = FPlatformFileManager::Get();
+	IPlatformFile& PlatformFile = FileManager.GetPlatformFile();
+	if (!PlatformFile.FileExists(*FilePath.FilePath))
+	{
+		return false;
+	}
+	TSharedPtr<IFileHandle> File(PlatformFile.OpenRead(*FilePath.FilePath, false));
+	if (!File.IsValid())
+	{
+		return false;
+	}
+		
+	if (FilePath.FilePath.EndsWith(".bin"))
+	{
+		size_t NumValues = 0;
+		if (!File->Read(reinterpret_cast<uint8*>(&NumValues), sizeof(size_t)))
+			return false;
+		TArray<float> Values;
+		Values.SetNum(NumValues);
+
+		int64 BytesRemaining = File->Size() - File->Tell();
+		uint8 DataSize = BytesRemaining / NumValues;
+		if (DataSize == sizeof(float))
+		{
+			if (!File->Read(reinterpret_cast<uint8*>(&Values[0]), NumValues * DataSize))
+				return false;
+		}
+		else if (DataSize == sizeof(double))
+		{
+			TArray<double> Tmp;
+			Tmp.SetNum(NumValues);
+			if (!File->Read(reinterpret_cast<uint8*>(&Tmp[0]), NumValues * DataSize))
+				return false;
+			for (int32 i = 0; i < NumValues; i++)
+			Values[i] = Tmp[i];
+		}
+		else
+		{
+			return false;
+		}
+
+		if (Values.Num() == 14185) // ryan - slrpmnstr hack!
+		{
+			SectionNames.SetNum(2);
+			SectionValues.SetNum(2);
+			SectionValues[0].SetNum(13343);
+			for (int32 i = 0; i < 13343; i++)
+				SectionValues[0][i] = Values[i];
+			SectionValues[1].SetNum(842);
+			for (int32 i = 13343, j=0; i < 14185; i++, j++)
+				SectionValues[1][j] = Values[i];
+		}
+		else
+		{
+			SectionNames.Add(FilePath.FilePath);
+			SectionValues.Add(MoveTemp(Values));
+		}
+	}
+	else if (FilePath.FilePath.EndsWith(".jiggleweights"))
+	{
+		uint8 FileVersion = 0;
+		if (!File->Read(reinterpret_cast<uint8*>(&FileVersion), sizeof(uint8)))
+			return false;
+		uint64 NumSections = 0;
+		if (!File->Read(reinterpret_cast<uint8*>(&NumSections), sizeof(uint64)))
+			return false;
+		SectionNames.Reserve(NumSections);
+		SectionValues.Reserve(NumSections);
+		for (uint64 Section = 0; Section < NumSections; Section++)
+		{
+			uint64 NumChar = 0;
+			if (!File->Read(reinterpret_cast<uint8*>(&NumChar), sizeof(uint64)))
+				return false;
+			FString Name;
+			for (uint64 i = 0; i < NumChar; i++)
+			{
+				char Ch;
+				if (!File->Read(reinterpret_cast<uint8*>(&Ch), sizeof(char)))
+					return false;
+				Name.AppendChar(Ch);
+			}
+			uint8 DataTypeSize = 0;
+			if (!File->Read(reinterpret_cast<uint8*>(&DataTypeSize), sizeof(uint8)))
+				return false;
+			uint64 NumValues = 0;
+			if (!File->Read(reinterpret_cast<uint8*>(&NumValues), sizeof(uint64)))
+				return false;
+			TArray<float> Values; Values.AddUninitialized(NumValues);
+			if (DataTypeSize == sizeof(float))
+			{
+				if (!File->Read(reinterpret_cast<uint8*>(&Values[0]), DataTypeSize * NumValues))
+					return false;
+			}
+			else if (DataTypeSize == sizeof(double))
+			{
+				for (uint64 i = 0; i < NumValues; i++)
+				{
+					double Value;
+					if (!File->Read(reinterpret_cast<uint8*>(&Value), DataTypeSize))
+						return false;
+					Values[i] = Value;
+				}
+			}
+			else
+			{
+				return false;
+			}
+			SectionNames.Add(Name);
+			SectionValues.Add(MoveTemp(Values));
+		}
+	}
+#endif
+	return true;
+}
+
+bool
+FOptimusJiggleSpringParameters::ReadWeightsFile(const FFilePath& FilePath, TArray<FVector3f>& Positions, TArray<float>& Values) const
 {
 #if WITH_EDITOR
 	if (!FilePath.FilePath.IsEmpty())
@@ -37,26 +169,42 @@ bool FOptimusJiggleSpringParameters::ReadWeightsFile(const FFilePath& FilePath,
 		{
 			if (FilePath.FilePath.EndsWith(".bin"))
 			{
+				// NumValues here is NOT the number of [x, y, z, v] records, rather it is the total number of individual values.
 				size_t NumValues = 0;
 				if (!File->Read(reinterpret_cast<uint8*>(&NumValues), sizeof(size_t)))
 					return false;
-				Values.SetNum(NumValues);
+				NumValues /= 4; // Make NumValues reflect the number of [x, y, z, v] records.
 
 				int64 BytesRemaining = File->Size() - File->Tell();
-				uint8 DataSize = BytesRemaining / NumValues;
+				uint8 DataSize = BytesRemaining / (NumValues * 4);
+
 				if (DataSize == sizeof(float))
 				{
-					if (!File->Read(reinterpret_cast<uint8*>(&Values[0]), NumValues * DataSize))
-						return false;
+					Positions.SetNum(NumValues);
+					Values.SetNum(NumValues);
+					for (int32 i = 0; i < NumValues; i++)
+					{
+						if (!File->Read(reinterpret_cast<uint8*>(&Positions[i]), 3 * DataSize))
+							return false;
+						if (!File->Read(reinterpret_cast<uint8*>(&Values[i]), DataSize))
+							return false;
+					}
 				}
 				else if (DataSize == sizeof(double))
 				{
+					Positions.SetNum(NumValues);
+					Values.SetNum(NumValues);
 					TArray<double> Tmp;
-					Tmp.SetNum(NumValues);
-					if (!File->Read(reinterpret_cast<uint8*>(&Tmp[0]), NumValues * DataSize))
-						return false;
+					Tmp.SetNum(4);
 					for (int32 i = 0; i < NumValues; i++)
-						Values[i] = Tmp[i];
+					{
+						if (!File->Read(reinterpret_cast<uint8*>(&Tmp[0]), Tmp.Num() * DataSize))
+							return false;
+						Positions[i][0] = Tmp[0];
+						Positions[i][1] = Tmp[1];
+						Positions[i][2] = Tmp[2];
+						Values[i] = Tmp[3];
+					}
 				}
 				else
 				{
@@ -77,7 +225,6 @@ bool FOptimusJiggleSpringParameters::ReadWeightsFile(const FFilePath& FilePath,
 	return true;
 }
 
-
 //
 // Interface
 //
@@ -88,16 +235,32 @@ UOptimusJiggleSpringDataInterface::GetDisplayName() const
 	return TEXT("Jiggle Spring");
 }
 
+void
+UOptimusJiggleSpringDataInterface::RegisterTypes()
+{
+	FOptimusDataTypeRegistry::Get().RegisterType(
+		FName("FAnalyticSpring"),
+		FText::FromString(TEXT("FAnalyticSpring")),
+		FShaderValueType::Get(
+			FName("FAnalyticSpring"), 
+			// Initializer list for declaring members; Optimus doesn't allow empty structs, so we must have at least 1.
+			{ FShaderValueType::FStructElement(FName("Dt"), FShaderValueType::Get(EShaderFundamentalType::Float)) }), 
+		FName("FAnalyticSpring"),
+		nullptr,
+		FLinearColor(0.3f, 0.7f, 0.4f, 1.0f),
+		EOptimusDataTypeUsageFlags::None);
+}
+
 TArray<FOptimusCDIPinDefinition> 
 UOptimusJiggleSpringDataInterface::GetPinDefinitions() const
 {
 	TArray<FOptimusCDIPinDefinition> Defs;
-	Defs.Add({ "Stiffness", "ReadStiffness", Optimus::DomainName::Vertex, "ReadNumVertices" });
-	Defs.Add({ "Damping", "ReadDamping", Optimus::DomainName::Vertex, "ReadNumVertices" });
+	Defs.Add({ "AnalyticSpring", "ReadAnalyticSpring" });
 	return Defs;
 }
 
-TSubclassOf<UActorComponent> UOptimusJiggleSpringDataInterface::GetRequiredComponentClass() const
+TSubclassOf<UActorComponent> 
+UOptimusJiggleSpringDataInterface::GetRequiredComponentClass() const
 {
 	return USkinnedMeshComponent::StaticClass();
 }
@@ -105,6 +268,8 @@ TSubclassOf<UActorComponent> UOptimusJiggleSpringDataInterface::GetRequiredCompo
 void 
 UOptimusJiggleSpringDataInterface::GetSupportedInputs(TArray<FShaderFunctionDefinition>& OutFunctions) const
 {
+	// Note: Removing these declarations breaks existing assets, even if they don't use them.
+
 	OutFunctions.AddDefaulted_GetRef()
 		.SetName(TEXT("ReadNumVertices"))
 		.AddReturnType(EShaderFundamentalType::Uint);
@@ -118,16 +283,44 @@ UOptimusJiggleSpringDataInterface::GetSupportedInputs(TArray<FShaderFunctionDefi
 		.SetName(TEXT("ReadDamping"))
 		.AddReturnType(EShaderFundamentalType::Float)
 		.AddParam(EShaderFundamentalType::Uint);
+
+	// DEPRECATED
+	OutFunctions.AddDefaulted_GetRef()
+		.SetName(TEXT("ReadBufferIndex"))
+		.AddReturnType(EShaderFundamentalType::Uint)
+		.AddParam(EShaderFundamentalType::Uint);
+
+	// DEPRECATED
+	OutFunctions.AddDefaulted_GetRef()
+		.SetName(TEXT("ReadNonZeroStiffnessNumber"))
+		.AddReturnType(EShaderFundamentalType::Uint)
+		.AddParam(EShaderFundamentalType::Uint);
+
+	OutFunctions.AddDefaulted_GetRef()
+		.SetName(TEXT("ReadMaxStretch"))
+		.AddReturnType(EShaderFundamentalType::Float)
+		.AddParam(EShaderFundamentalType::Uint);
+
+	FShaderValueTypeHandle AnalyticSpringType = FOptimusDataTypeRegistry::Get().FindType(FName("FAnalyticSpring"))->ShaderValueType;
+	OutFunctions.AddDefaulted_GetRef()
+		.SetName(TEXT("ReadAnalyticSpring"))
+		.AddReturnType(AnalyticSpringType);	
 }
 
 BEGIN_SHADER_PARAMETER_STRUCT(FJiggleSpringDataInterfaceParameters, )
+	SHADER_PARAMETER(uint32, NumVertexMap)
 	SHADER_PARAMETER(uint32, NumVertices)
+	SHADER_PARAMETER(uint32, BaseVertexIndex)
 	SHADER_PARAMETER(uint32, NumStiffnessWeights)
 	SHADER_PARAMETER(uint32, NumDampingWeights)
+	SHADER_PARAMETER(uint32, NumMaxStretchWeights)
 	SHADER_PARAMETER(float, Stiffness)
 	SHADER_PARAMETER(float, Damping)
+	SHADER_PARAMETER(float, MaxStretch)
+	SHADER_PARAMETER_RDG_BUFFER_SRV(Buffer<uint32>, VertexMapBuffer)
 	SHADER_PARAMETER_RDG_BUFFER_SRV(Buffer<float>, StiffnessWeightsBuffer)
 	SHADER_PARAMETER_RDG_BUFFER_SRV(Buffer<float>, DampingWeightsBuffer)
+	SHADER_PARAMETER_RDG_BUFFER_SRV(Buffer<float>, MaxStretchWeightsBuffer)
 END_SHADER_PARAMETER_STRUCT()
 
 void 
@@ -159,12 +352,132 @@ UOptimusJiggleSpringDataInterface::GetHLSL(FString& OutHLSL) const
 	OutHLSL += TEXT("#include \"/Plugin/Optimus/Private/DataInterfaceJiggleSpring.ush\"\n");
 }
 
-UComputeDataProvider* 
+template <class T>
+void
+GetStats(const TArray<T>& Values, T& Min, double& Avg, T& Max, const bool SkipZero=false)
+{
+	Min = TNumericLimits<T>::Max();
+	Max = -TNumericLimits<T>::Max();
+	Avg = 0.0;
+	for (int32 i = 0; i < Values.Num(); i++)
+	{
+		if (SkipZero && Values[i] == 0.0) continue;
+		if (Values[i] < Min) Min = Values[i];
+		if (Values[i] > Max) Max = Values[i];
+		Avg += Values[i] / Values.Num();
+	}
+}
+
+UComputeDataProvider*
 UOptimusJiggleSpringDataInterface::CreateDataProvider(TObjectPtr<UObject> InBinding, uint64 InInputMask, uint64 InOutputMask) const
 {
+	// We're missing a transport mechanism that'll feed us TArray<float> via Blueprints, 
+	// so we need to read these files here for now.
+	// 
+	// Read files if weights are empty.
+	UOptimusJiggleSpringDataInterface* NCThis = const_cast<UOptimusJiggleSpringDataInterface*>(this);
+	if (!JiggleSpringParameters.StiffnessWeights.Num() && !JiggleSpringParameters.StiffnessWeightsFile.FilePath.IsEmpty())
+	{
+		UE_LOG(LogTemp, Log, TEXT("UOptimusJiggleSpringDataInterface::CreateDataProvider() - Reading stiffness file: '%s'"), *JiggleSpringParameters.StiffnessWeightsFile.FilePath);
+		if (!JiggleSpringParameters.ReadWeightsFile(JiggleSpringParameters.StiffnessWeightsFile, NCThis->JiggleSpringParameters.StiffnessWeightsNames, NCThis->JiggleSpringParameters.StiffnessWeights))
+		{
+			UE_LOG(LogTemp, Error, TEXT("Failed to read stiffness file: '%s'"), *JiggleSpringParameters.StiffnessWeightsFile.FilePath);
+		}
+		//float Min, Max; double Avg;
+		//GetStats(JiggleSpringParameters.StiffnessWeights, Min, Avg, Max, true);
+		//UE_LOG(LogTemp, Log, TEXT("UOptimusJiggleSpringDataInterface::CreateDataProvider() - Stiffness stats (min, avg, max): %g, %g, %g"), Min, Avg, Max);
+	}
+	if (!JiggleSpringParameters.DampingWeights.Num() && !JiggleSpringParameters.DampingWeightsFile.FilePath.IsEmpty())
+	{
+		UE_LOG(LogTemp, Log, TEXT("UOptimusJiggleSpringDataInterface::CreateDataProvider() - Reading damping file: '%s'"), *JiggleSpringParameters.DampingWeightsFile.FilePath);
+		if (!JiggleSpringParameters.ReadWeightsFile(JiggleSpringParameters.DampingWeightsFile, NCThis->JiggleSpringParameters.DampingWeightsNames, NCThis->JiggleSpringParameters.DampingWeights))
+		{
+			UE_LOG(LogTemp, Error, TEXT("Failed to read damping file: '%s'"), *JiggleSpringParameters.DampingWeightsFile.FilePath);
+		}
+		//float Min, Max; double Avg;
+		//GetStats(JiggleSpringParameters.DampingWeights, Min, Avg, Max, true);
+		//UE_LOG(LogTemp, Log, TEXT("UOptimusJiggleSpringDataInterface::CreateDataProvider() - Stiffness stats (min, avg, max): %g, %g, %g"), Min, Avg, Max);
+	}
+
+	// Copy JiggleSpringParameters for data provider.
 	UOptimusJiggleSpringDataProvider* Provider = NewObject<UOptimusJiggleSpringDataProvider>();
 	Provider->SkinnedMesh = Cast<USkinnedMeshComponent>(InBinding);
 	Provider->JiggleSpringParameters = JiggleSpringParameters;
+
+	if (JiggleSpringParameters.StiffnessWeights.Num() || JiggleSpringParameters.DampingWeights.Num())
+	{
+#if WITH_EDITORONLY_DATA
+		// TODO Ryan - Topology is only available with the editor.  Crazy talk, I know...  Save what we need in parameters?
+		const USkeletalMesh* SkeletalMesh = Provider->SkinnedMesh ? Cast<USkeletalMesh>(Provider->SkinnedMesh->GetSkinnedAsset()) : nullptr;
+		if (const FSkeletalMeshModel* ImportedModel = SkeletalMesh ? SkeletalMesh->GetImportedModel() : nullptr)
+		{
+			// We (may) need to reorder stiffness & damping values from the import geometry to the render geometry.
+			// Send along the map that does that translation to the shader.
+			Provider->JiggleSpringParameters.VertexMap.SetNum(ImportedModel->LODModels[0].MeshToImportVertexMap.Num());
+			for (int32 i = 0; i < ImportedModel->LODModels[0].MeshToImportVertexMap.Num(); i++)
+			{
+				Provider->JiggleSpringParameters.VertexMap[i] = static_cast<uint32>(ImportedModel->LODModels[0].MeshToImportVertexMap[i]);
+			}
+
+			if (JiggleSpringParameters.bUseAvgEdgeLengthForMaxStretchMap)
+			{
+				// Pry the mesh topology out of the LOD model.
+				TArray<FIntVector3> Triangles;
+				for (int32 i = 0; i < ImportedModel->LODModels[0].Sections.Num(); i++)
+				{
+					const auto& Section = ImportedModel->LODModels[0].Sections[i];
+					Triangles.Reserve(Triangles.Num() + Section.NumTriangles);
+					for (uint32 j = 0; j < Section.NumTriangles; j++)
+					{
+						FIntVector3 Tri;
+						for (int32 k = 0; k < 3; k++)
+						{
+							Tri[k] = ImportedModel->LODModels[0].IndexBuffer[Section.BaseIndex + ((j * 3 + k))];
+						}
+						Triangles.Add(Tri);
+					}
+				}
+
+				// Get the set of unique edges (don't care about winding order).
+				TSet<TPair<int32, int32>> Edges;
+				for (int i = 0; i < Triangles.Num(); i++)
+				{
+					const FIntVector3& Tri = Triangles[i];
+					for (int j = 0; j < 3; j++)
+					{
+						int32 a = Tri[j];
+						int32 b = Tri[(j + 1) % 3];
+						Edges.Add(TPair<int32, int32>(FMath::Min(a, b), FMath::Max(a, b)));
+					}
+				}
+
+				// Compute the per vertex average edge length.
+				TArray<float> AvgEdgeLength; AvgEdgeLength.SetNum(ImportedModel->LODModels[0].MeshToImportVertexMap.Num());
+				TArray<int32> NumEdges; NumEdges.SetNum(ImportedModel->LODModels[0].MeshToImportVertexMap.Num());
+				TArray<FSoftSkinVertex> Vertices;
+				ImportedModel->LODModels[0].GetVertices(Vertices);
+				for (auto it = Edges.CreateConstIterator(); !!it; ++it)
+				{
+					const TPair<int32, int32>& Edge = *it;
+					float EdgeLength = (Vertices[Edge.Key].Position - Vertices[Edge.Value].Position).Length();
+					AvgEdgeLength[Edge.Key] += EdgeLength;
+					AvgEdgeLength[Edge.Value] += EdgeLength;
+					NumEdges[Edge.Key]++;
+					NumEdges[Edge.Value]++;
+				}
+				for (int32 i = 0; i < NumEdges.Num(); i++)
+					if (NumEdges[i] > 0)
+						AvgEdgeLength[i] /= NumEdges[i];
+
+				Provider->JiggleSpringParameters.MaxStretchWeights = AvgEdgeLength;
+			}
+		}
+		else
+		{
+			// Despair?
+		}
+#endif
+	}
 	return Provider;
 }
 
@@ -183,11 +496,6 @@ UOptimusJiggleSpringDataProvider::IsValid() const
 FComputeDataProviderRenderProxy* 
 UOptimusJiggleSpringDataProvider::GetRenderProxy()
 {
-	// Not sure where file reading should happen...
-	if(!JiggleSpringParameters.StiffnessWeights.Num() && !JiggleSpringParameters.StiffnessWeightsFile.FilePath.IsEmpty())
-		JiggleSpringParameters.ReadWeightsFile(JiggleSpringParameters.StiffnessWeightsFile, JiggleSpringParameters.StiffnessWeights);
-	if(!JiggleSpringParameters.DampingWeights.Num() && !JiggleSpringParameters.DampingWeightsFile.FilePath.IsEmpty())
-		JiggleSpringParameters.ReadWeightsFile(JiggleSpringParameters.DampingWeightsFile, JiggleSpringParameters.DampingWeights);
 	return new FOptimusJiggleSpringDataProviderProxy(SkinnedMesh, JiggleSpringParameters);
 }
 
@@ -198,7 +506,9 @@ UOptimusJiggleSpringDataProvider::GetRenderProxy()
 FOptimusJiggleSpringDataProviderProxy::FOptimusJiggleSpringDataProviderProxy(
 	USkinnedMeshComponent* SkinnedMeshComponent, 
 	FOptimusJiggleSpringParameters const& InJiggleSpringParameters)
+	: SkinnedMeshComponent(SkinnedMeshComponent)
 {
+
 	SkeletalMeshObject = SkinnedMeshComponent->MeshObject;
 	JiggleSpringParameters = InJiggleSpringParameters;
 }
@@ -206,19 +516,76 @@ FOptimusJiggleSpringDataProviderProxy::FOptimusJiggleSpringDataProviderProxy(
 void
 FOptimusJiggleSpringDataProviderProxy::AllocateResources(FRDGBuilder& GraphBuilder)
 {
+	// Find the total number of vertices for this skeletal mesh
+	uint32 TotalNumVertices = 0;
+	const int32 LodIndex = SkeletalMeshObject->GetLOD();
+	FSkeletalMeshRenderData const& SkeletalMeshRenderData = SkeletalMeshObject->GetSkeletalMeshRenderData();
+	FSkeletalMeshLODRenderData const* LodRenderData = &SkeletalMeshRenderData.LODRenderData[LodIndex];
+	for (int32 i = 0; i < LodRenderData->RenderSections.Num(); i++)
+	{
+		TotalNumVertices += LodRenderData->RenderSections[i].GetNumVertices();
+	}
+
+	// Find the material property index that matches this skeletal mesh's name or num vertices.
+	SectionIndex = INDEX_NONE;
+	for (int32 i = 0; i < JiggleSpringParameters.StiffnessWeights.Num(); i++)
+	{
+		const FString& Name = JiggleSpringParameters.StiffnessWeightsNames.Num() > i ? JiggleSpringParameters.StiffnessWeightsNames[i] : FString();
+		if (SkeletalMeshObject->GetDebugName().ToString() == Name ||
+			TotalNumVertices == JiggleSpringParameters.StiffnessWeights[i].Num())
+		{
+			SectionIndex = i;
+		}
+	}
+
+	// VertexMapBuffer - maps from render vertex index to import vertex index.
+	// We only need this if the material property files were solved for on the import geometry,
+	// which likely differs from the render geometry.
+	if (SectionIndex != INDEX_NONE &&
+		JiggleSpringParameters.VertexMap.Num() &&
+		(JiggleSpringParameters.StiffnessWeights[SectionIndex].Num() != TotalNumVertices ||
+		 JiggleSpringParameters.DampingWeights[SectionIndex].Num() != TotalNumVertices))
+	{
+		VertexMapBuffer =
+			GraphBuilder.CreateBuffer(
+				FRDGBufferDesc::CreateStructuredDesc(sizeof(uint32), JiggleSpringParameters.VertexMap.Num()),
+				TEXT("JiggleSpring.VertexMap"));
+		VertexMapBufferSRV =
+			GraphBuilder.CreateSRV(VertexMapBuffer);
+		GraphBuilder.QueueBufferUpload(
+			VertexMapBuffer,
+			JiggleSpringParameters.VertexMap.GetData(),
+			JiggleSpringParameters.VertexMap.Num() * sizeof(uint32),
+			ERDGInitialDataFlags::None);
+	}
+	else
+	{
+		VertexMapBuffer =
+			GraphBuilder.CreateBuffer(
+				FRDGBufferDesc::CreateStructuredDesc(sizeof(uint32), 1),
+				TEXT("JiggleSpring.VertexMap"));
+		VertexMapBufferSRV =
+			GraphBuilder.CreateSRV(VertexMapBuffer);
+		GraphBuilder.QueueBufferUpload(
+			VertexMapBuffer,
+			&NullIntBuffer,
+			1 * sizeof(uint32),
+			ERDGInitialDataFlags::None);
+	}
 	// Stiffness weights
-	if (JiggleSpringParameters.StiffnessWeights.Num())
+	if (SectionIndex != INDEX_NONE &&
+		JiggleSpringParameters.StiffnessWeights[SectionIndex].Num())
 	{
 		StiffnessWeightsBuffer =
 			GraphBuilder.CreateBuffer(
-				FRDGBufferDesc::CreateStructuredDesc(sizeof(float), JiggleSpringParameters.StiffnessWeights.Num()),
+				FRDGBufferDesc::CreateStructuredDesc(sizeof(float), JiggleSpringParameters.StiffnessWeights[SectionIndex].Num()),
 				TEXT("JiggleSpring.StiffnessWeights"));
 		StiffnessWeightsBufferSRV =
 			GraphBuilder.CreateSRV(StiffnessWeightsBuffer);
 		GraphBuilder.QueueBufferUpload(
 			StiffnessWeightsBuffer,
-			JiggleSpringParameters.StiffnessWeights.GetData(),
-			JiggleSpringParameters.StiffnessWeights.Num() * sizeof(float),
+			JiggleSpringParameters.StiffnessWeights[SectionIndex].GetData(),
+			JiggleSpringParameters.StiffnessWeights[SectionIndex].Num() * sizeof(float),
 			ERDGInitialDataFlags::None);
 	}
 	else
@@ -236,18 +603,19 @@ FOptimusJiggleSpringDataProviderProxy::AllocateResources(FRDGBuilder& GraphBuild
 			ERDGInitialDataFlags::None);
 	}
 	// Damping weights
-	if (JiggleSpringParameters.DampingWeights.Num())
+	if (SectionIndex != INDEX_NONE &&
+		JiggleSpringParameters.DampingWeights[SectionIndex].Num())
 	{
 		DampingWeightsBuffer =
 			GraphBuilder.CreateBuffer(
-				FRDGBufferDesc::CreateStructuredDesc(sizeof(float), JiggleSpringParameters.DampingWeights.Num()),
+				FRDGBufferDesc::CreateStructuredDesc(sizeof(float), JiggleSpringParameters.DampingWeights[SectionIndex].Num()),
 				TEXT("JiggleSpring.DampingWeights"));
 		DampingWeightsBufferSRV =
 			GraphBuilder.CreateSRV(DampingWeightsBuffer);
 		GraphBuilder.QueueBufferUpload(
 			DampingWeightsBuffer,
-			JiggleSpringParameters.DampingWeights.GetData(),
-			JiggleSpringParameters.DampingWeights.Num() * sizeof(float),
+			JiggleSpringParameters.DampingWeights[SectionIndex].GetData(),
+			JiggleSpringParameters.DampingWeights[SectionIndex].Num() * sizeof(float),
 			ERDGInitialDataFlags::None);
 	}
 	else
@@ -264,6 +632,35 @@ FOptimusJiggleSpringDataProviderProxy::AllocateResources(FRDGBuilder& GraphBuild
 			1 * sizeof(float),
 			ERDGInitialDataFlags::None);
 	}
+	// Max stretch weights
+	if (JiggleSpringParameters.MaxStretchWeights.Num())
+	{
+		MaxStretchWeightsBuffer =
+			GraphBuilder.CreateBuffer(
+				FRDGBufferDesc::CreateStructuredDesc(sizeof(float), JiggleSpringParameters.MaxStretchWeights.Num()),
+				TEXT("JiggleSpring.MaxStretchWeights"));
+		MaxStretchWeightsBufferSRV =
+			GraphBuilder.CreateSRV(MaxStretchWeightsBuffer);
+		GraphBuilder.QueueBufferUpload(
+			MaxStretchWeightsBuffer,
+			JiggleSpringParameters.MaxStretchWeights.GetData(),
+			JiggleSpringParameters.MaxStretchWeights.Num() * sizeof(float),
+			ERDGInitialDataFlags::None);
+	}
+	else
+	{
+		MaxStretchWeightsBuffer =
+			GraphBuilder.CreateBuffer(
+				FRDGBufferDesc::CreateStructuredDesc(sizeof(float), 1),
+				TEXT("JiggleSpring.MaxStretchWeights"));
+		MaxStretchWeightsBufferSRV =
+			GraphBuilder.CreateSRV(MaxStretchWeightsBuffer);
+		GraphBuilder.QueueBufferUpload(
+			MaxStretchWeightsBuffer,
+			&NullFloatBuffer,
+			1 * sizeof(float),
+			ERDGInitialDataFlags::None);
+	}
 }
 
 struct FJiggleSpringDataInterfacePermutationIds
@@ -303,18 +700,23 @@ void FOptimusJiggleSpringDataProviderProxy::GatherDispatchData(FDispatchSetup co
 
 		FJiggleSpringDataInterfaceParameters* Parameters = 
 			(FJiggleSpringDataInterfaceParameters*)(InOutDispatchData.ParameterBuffer + InDispatchSetup.ParameterBufferOffset + InDispatchSetup.ParameterBufferStride * InvocationIndex);
-		Parameters->NumVertices = RenderSection.NumVertices;
+		Parameters->NumVertices = RenderSection.GetNumVertices();
+		Parameters->BaseVertexIndex = RenderSection.BaseVertexIndex;
+		Parameters->NumVertexMap = JiggleSpringParameters.VertexMap.Num();
+
+		Parameters->NumStiffnessWeights = SectionIndex != INDEX_NONE ? JiggleSpringParameters.StiffnessWeights[SectionIndex].Num() : 0;
+		Parameters->NumDampingWeights = SectionIndex != INDEX_NONE ? JiggleSpringParameters.DampingWeights[SectionIndex].Num() : 0;
 
-		Parameters->NumStiffnessWeights = JiggleSpringParameters.StiffnessWeights.Num();
-		Parameters->NumDampingWeights = JiggleSpringParameters.DampingWeights.Num();
+		Parameters->VertexMapBuffer = VertexMapBufferSRV;
+		Parameters->StiffnessWeightsBuffer = StiffnessWeightsBufferSRV;
+		Parameters->DampingWeightsBuffer = DampingWeightsBufferSRV;
 
 		Parameters->Stiffness = JiggleSpringParameters.BaselineStiffness;
 		Parameters->Damping = JiggleSpringParameters.BaselineDamping;
 
-		Parameters->StiffnessWeightsBuffer = StiffnessWeightsBufferSRV;
-		Parameters->DampingWeightsBuffer = DampingWeightsBufferSRV;
+		Parameters->MaxStretch = JiggleSpringParameters.MaxStretchMultiplier;
+		Parameters->MaxStretchWeightsBuffer = MaxStretchWeightsBufferSRV;
 
 		InOutDispatchData.PermutationId[InvocationIndex] |= PermutationIds.EnableDeformerJiggleSpring;
 	}
-
 }
diff --git a/Engine/Plugins/Animation/DeformerGraph/Source/OptimusCore/Private/DataInterfaces/OptimusDataInterfaceJiggleSpring.h b/Engine/Plugins/Animation/DeformerGraph/Source/OptimusCore/Private/DataInterfaces/OptimusDataInterfaceJiggleSpring.h
--- a/Engine/Plugins/Animation/DeformerGraph/Source/OptimusCore/Private/DataInterfaces/OptimusDataInterfaceJiggleSpring.h
+++ b/Engine/Plugins/Animation/DeformerGraph/Source/OptimusCore/Private/DataInterfaces/OptimusDataInterfaceJiggleSpring.h
@@ -23,9 +23,20 @@ struct FOptimusJiggleSpringParameters
 	UPROPERTY(EditAnywhere, Category = Deformer)
 	float BaselineStiffness = 100.0f;
 
-	/** Per vertex spring stiffness. */
-	UPROPERTY(EditAnywhere, Category = Deformer)
-	TArray<float> StiffnessWeights;
+	// Notes on stiffness and damping weights:
+	// Until we can transport TArray<float> to the shader via blueprints, we're stuck
+	// reading these values from standalone files.  These files contain values for all
+	// surfaces, and so we need a way to associate values to surface.  So we optionally
+	// match by name (which may be lost during surface export to file), and if that 
+	// fails, we match by vertex to value count.  At the point when we can set per
+	// skeletal mesh values via blueprints, the names arrays go away and the 
+	// TArray<TArray<float>> for values turn into TArray<float>, and is exposed as a
+	// UPROPERTY.
+
+	/** Per vertex spring stiffness surface names. */
+	TArray<FString> StiffnessWeightsNames;
+	//UPROPERTY(EditAnywhere, Category = Deformer)
+	TArray<TArray<float>> StiffnessWeights;
 
 	/** Stiffness weights file. */
 	UPROPERTY(EditAnywhere, Category = Deformer)
@@ -35,15 +46,35 @@ struct FOptimusJiggleSpringParameters
 	UPROPERTY(EditAnywhere, Category = Deformer)
 	float BaselineDamping = 10.0f;
 
-	/** Per vertex spring damping. */
-	UPROPERTY(EditAnywhere, Category = Deformer)
-	TArray<float> DampingWeights;
+	/** Per vertex spring damping surface names. */
+	TArray<FString> DampingWeightsNames;
+	//UPROPERTY(EditAnywhere, Category = Deformer)
+	TArray<TArray<float>> DampingWeights;
 
 	/** Damping weights file. */
 	UPROPERTY(EditAnywhere, Category = Deformer)
 	FFilePath DampingWeightsFile;
 
-	bool ReadWeightsFile(const FFilePath& FilePath, TArray<float>& Values) const;
+	/** Map render vertices to import indices. */
+	TArray<int32> VertexMap;
+
+	/** Multiplier on the max stretch distance per-vertex value if specified, 
+	 *  or the uniform max stretch value if no per-vertex map is specified. 
+	 */
+	UPROPERTY(EditAnywhere, Category = Deformer)
+	float MaxStretchMultiplier = 3.0f;
+
+	/** Use the per vertex average edge length as the max stretch distance map. */
+	UPROPERTY(EditAnywhere, Category = Deformer)
+	bool bUseAvgEdgeLengthForMaxStretchMap = true;
+
+	/** Per vertex maximum stretch distance. */
+	UPROPERTY(EditAnywhere, Category = Deformer)
+	TArray<float> MaxStretchWeights;
+
+	//bool ReadWeightsFile(const FFilePath& FilePath, TArray<TTuple<FString,TArray<float>>>& Values) const;
+	bool ReadWeightsFile(const FFilePath& FilePath, TArray<FString>& Names, TArray<TArray<float>>& Values) const;
+	bool ReadWeightsFile(const FFilePath& FilePath, TArray<FVector3f>& Positions, TArray<float>& Values) const;
 };
 
 
@@ -66,6 +97,7 @@ public:
 	FString GetDisplayName() const override;
 	TArray<FOptimusCDIPinDefinition> GetPinDefinitions() const override;
 	TSubclassOf<UActorComponent> GetRequiredComponentClass() const override;
+	void RegisterTypes() override;
 	//~ End UOptimusComputeDataInterface Interface
 	
 	//~ Begin UComputeDataInterface Interface
@@ -114,12 +146,23 @@ public:
 	//~ End FComputeDataProviderRenderProxy Interface
 
 private:
+	USkinnedMeshComponent* SkinnedMeshComponent = nullptr;
 	FSkeletalMeshObject* SkeletalMeshObject = nullptr;
 	FOptimusJiggleSpringParameters JiggleSpringParameters;
+
+	// If the stiffness and damping weights came from file, then they're divided into multiple sections;
+	// 1 section for each skeletal mesh.  This integer identifies which one we're currently using.
+	int32 SectionIndex = INDEX_NONE;
+
+	FRDGBuffer* VertexMapBuffer = nullptr;
+	FRDGBufferSRV* VertexMapBufferSRV = nullptr;
 	FRDGBuffer* StiffnessWeightsBuffer = nullptr;
 	FRDGBufferSRV* StiffnessWeightsBufferSRV = nullptr;
 	FRDGBuffer* DampingWeightsBuffer = nullptr;
 	FRDGBufferSRV* DampingWeightsBufferSRV = nullptr;
+	FRDGBuffer* MaxStretchWeightsBuffer = nullptr;
+	FRDGBufferSRV* MaxStretchWeightsBufferSRV = nullptr;
 	float NullFloatBuffer = 0.0f;
+	int32 NullIntBuffer = 0;
 };
 
