Remove deprecated hlsl function.

diff --git a/Engine/Plugins/Animation/DeformerGraph/Shaders/Private/DataInterfaceJiggleSpring.ush b/Engine/Plugins/Animation/DeformerGraph/Shaders/Private/DataInterfaceJiggleSpring.ush
--- a/Engine/Plugins/Animation/DeformerGraph/Shaders/Private/DataInterfaceJiggleSpring.ush
+++ b/Engine/Plugins/Animation/DeformerGraph/Shaders/Private/DataInterfaceJiggleSpring.ush
@@ -83,19 +83,6 @@ DI_IMPL_READ(ReadBufferIndex, uint, uint VertexIndex)
 #endif
 }
 
-DI_IMPL_READ(ReadNonZeroStiffnessNumber, uint, uint VertexIndex)
-{
-	uint Count = 0;
-#if ENABLE_DEFORMER_JIGGLE_SPRING
-	for(uint i = 0; i <= VertexIndex; i++)
-    {
-		if(DI_LOCAL(ReadStiffness)(VertexIndex) > 0.) 
-			Count++;
-	}
-#endif
-	return Count;
-}
-
 DI_IMPL_READ(ReadMaxStretch, float, uint VertexIndex)
 {
 #if ENABLE_DEFORMER_JIGGLE_SPRING
diff --git a/Engine/Plugins/Animation/DeformerGraph/Source/OptimusCore/Private/DataInterfaces/OptimusDataInterfaceJiggleSpring.cpp b/Engine/Plugins/Animation/DeformerGraph/Source/OptimusCore/Private/DataInterfaces/OptimusDataInterfaceJiggleSpring.cpp
--- a/Engine/Plugins/Animation/DeformerGraph/Source/OptimusCore/Private/DataInterfaces/OptimusDataInterfaceJiggleSpring.cpp
+++ b/Engine/Plugins/Animation/DeformerGraph/Source/OptimusCore/Private/DataInterfaces/OptimusDataInterfaceJiggleSpring.cpp
@@ -290,12 +290,6 @@ UOptimusJiggleSpringDataInterface::GetSupportedInputs(TArray<FShaderFunctionDefi
 		.AddReturnType(EShaderFundamentalType::Uint)
 		.AddParam(EShaderFundamentalType::Uint);
 
-	// DEPRECATED
-	OutFunctions.AddDefaulted_GetRef()
-		.SetName(TEXT("ReadNonZeroStiffnessNumber"))
-		.AddReturnType(EShaderFundamentalType::Uint)
-		.AddParam(EShaderFundamentalType::Uint);
-
 	OutFunctions.AddDefaulted_GetRef()
 		.SetName(TEXT("ReadMaxStretch"))
 		.AddReturnType(EShaderFundamentalType::Float)
