Optimus: fixed a typo

diff --git a/Engine/Plugins/Animation/DeformerGraph/Source/OptimusEditor/Private/OptimusEditorClipboard.cpp b/Engine/Plugins/Animation/DeformerGraph/Source/OptimusEditor/Private/OptimusEditorClipboard.cpp
--- a/Engine/Plugins/Animation/DeformerGraph/Source/OptimusEditor/Private/OptimusEditorClipboard.cpp
+++ b/Engine/Plugins/Animation/DeformerGraph/Source/OptimusEditor/Private/OptimusEditorClipboard.cpp
@@ -48,7 +48,7 @@ UOptimusClipboardContent* UOptimusClipboardContent::Create(
 				ClipboardLink.NodeOutputPinName = Link->GetNodeOutputPin()->GetUniqueName().ToString();
 			
 				ClipboardLink.NodeInputIndex = NodeIndexMap[InputNode];
-				ClipboardLink.NodeInputPinName = Link->GetNodeOutputPin()->GetUniqueName().ToString();
+				ClipboardLink.NodeInputPinName = Link->GetNodeInputPin()->GetUniqueName().ToString();
 
 				Content->NodeLinks.Add(ClipboardLink);
 			}
