Optimus: Fix crash when LOD changes and we are using resources that are sized by LOD.
The buffer provider needs to get the current LOD from the component source so that it can use the correct key in the resource manager.

diff --git a/Engine/Plugins/Animation/DeformerGraph/Source/OptimusCore/Private/ComponentSources/OptimusSkinnedMeshComponentSource.cpp b/Engine/Plugins/Animation/DeformerGraph/Source/OptimusCore/Private/ComponentSources/OptimusSkinnedMeshComponentSource.cpp
--- a/Engine/Plugins/Animation/DeformerGraph/Source/OptimusCore/Private/ComponentSources/OptimusSkinnedMeshComponentSource.cpp
+++ b/Engine/Plugins/Animation/DeformerGraph/Source/OptimusCore/Private/ComponentSources/OptimusSkinnedMeshComponentSource.cpp
@@ -29,10 +29,17 @@ TArray<FName> UOptimusSkinnedMeshComponentSource::GetExecutionDomains() const
 	return {Domains::Vertex, Domains::Triangle};
 }
 
+int32 UOptimusSkinnedMeshComponentSource::GetLodIndex(const UActorComponent* InComponent) const
+{
+	const USkinnedMeshComponent* SkinnedMeshComponent = Cast<USkinnedMeshComponent>(InComponent);
+	const FSkeletalMeshObject* SkeletalMeshObject = SkinnedMeshComponent ? SkinnedMeshComponent->MeshObject : nullptr;
+	return SkeletalMeshObject ? SkeletalMeshObject->GetLOD() : 0;
+}
 
 bool UOptimusSkinnedMeshComponentSource::GetComponentElementCountsForExecutionDomain(
 	FName InDomainName,
 	const UActorComponent* InComponent,
+	int32 InLodIndex,
 	TArray<int32>& OutInvocationElementCounts
 	) const
 {
@@ -43,9 +50,8 @@ bool UOptimusSkinnedMeshComponentSource::GetComponentElementCountsForExecutionDo
 	}
 
 	const FSkeletalMeshObject* SkeletalMeshObject = SkinnedMeshComponent->MeshObject;
-	const int32 LodIndex = SkeletalMeshObject->GetLOD();
 	FSkeletalMeshRenderData const& SkeletalMeshRenderData = SkeletalMeshObject->GetSkeletalMeshRenderData();
-	FSkeletalMeshLODRenderData const* LodRenderData = &SkeletalMeshRenderData.LODRenderData[LodIndex];
+	FSkeletalMeshLODRenderData const* LodRenderData = &SkeletalMeshRenderData.LODRenderData[InLodIndex];
 
 	OutInvocationElementCounts.Reset();
 #if 0
diff --git a/Engine/Plugins/Animation/DeformerGraph/Source/OptimusCore/Private/ComponentSources/OptimusSkinnedMeshComponentSource.h b/Engine/Plugins/Animation/DeformerGraph/Source/OptimusCore/Private/ComponentSources/OptimusSkinnedMeshComponentSource.h
--- a/Engine/Plugins/Animation/DeformerGraph/Source/OptimusCore/Private/ComponentSources/OptimusSkinnedMeshComponentSource.h
+++ b/Engine/Plugins/Animation/DeformerGraph/Source/OptimusCore/Private/ComponentSources/OptimusSkinnedMeshComponentSource.h
@@ -24,5 +24,6 @@ public:
 	FName GetBindingName() const override { return FName("SkinnedMesh"); }
 	TSubclassOf<UActorComponent> GetComponentClass() const override;
 	TArray<FName> GetExecutionDomains() const override;
-	bool GetComponentElementCountsForExecutionDomain(FName InDomainName, const UActorComponent* InComponent, TArray<int32>& OutInvocationElementCounts) const override;
+	int32 GetLodIndex(const UActorComponent* InComponent) const override;
+	bool GetComponentElementCountsForExecutionDomain(FName InDomainName, const UActorComponent* InComponent, int32 InLodIndex, TArray<int32>& OutInvocationElementCounts) const override;
 };
diff --git a/Engine/Plugins/Animation/DeformerGraph/Source/OptimusCore/Private/DataInterfaces/OptimusDataInterfaceRawBuffer.cpp b/Engine/Plugins/Animation/DeformerGraph/Source/OptimusCore/Private/DataInterfaces/OptimusDataInterfaceRawBuffer.cpp
--- a/Engine/Plugins/Animation/DeformerGraph/Source/OptimusCore/Private/DataInterfaces/OptimusDataInterfaceRawBuffer.cpp
+++ b/Engine/Plugins/Animation/DeformerGraph/Source/OptimusCore/Private/DataInterfaces/OptimusDataInterfaceRawBuffer.cpp
@@ -238,7 +238,8 @@ bool UOptimusRawBufferDataProvider::IsValid() const
 }
 
 
-bool UOptimusRawBufferDataProvider::GetInvocationElementCounts(
+bool UOptimusRawBufferDataProvider::GetLodAndInvocationElementCounts(
+	int32& OutLodIndex,
 	TArray<int32>& OutInvocationElementCounts
 	) const
 {
@@ -254,12 +255,14 @@ bool UOptimusRawBufferDataProvider::GetInvocationElementCounts(
 		return false;
 	}
 
+	OutLodIndex = ComponentSourcePtr->GetLodIndex(ComponentPtr);
+
 	switch(DataDomain.Type)
 	{
 	case EOptimusDataDomainType::Dimensional:
 		{
 			const FName ExecutionDomain = !DataDomain.DimensionNames.IsEmpty() ? DataDomain.DimensionNames[0] : NAME_None;
-			if (ComponentSourcePtr->GetComponentElementCountsForExecutionDomain(ExecutionDomain, ComponentPtr, OutInvocationElementCounts))
+			if (ComponentSourcePtr->GetComponentElementCountsForExecutionDomain(ExecutionDomain, ComponentPtr, OutLodIndex, OutInvocationElementCounts))
 			{
 				if (DataDomain.DimensionNames.Num() == 1)
 				{
@@ -283,7 +286,7 @@ bool UOptimusRawBufferDataProvider::GetInvocationElementCounts(
 				EngineConstants.Add(ExecutionDomain, 0);
 
 				TArray<int32>& ElementCounts = ElementCountsPerDomain.FindOrAdd(ExecutionDomain);
-				if (!ComponentSourcePtr->GetComponentElementCountsForExecutionDomain(ExecutionDomain, ComponentPtr, ElementCounts))
+				if (!ComponentSourcePtr->GetComponentElementCountsForExecutionDomain(ExecutionDomain, ComponentPtr, OutLodIndex, ElementCounts))
 				{
 					return false;
 				}
@@ -321,8 +324,9 @@ bool UOptimusRawBufferDataProvider::GetInvocationElementCounts(
 
 FComputeDataProviderRenderProxy* UOptimusTransientBufferDataProvider::GetRenderProxy()
 {
+	int32 LodIndex;
 	TArray<int32> InvocationCounts;
-	GetInvocationElementCounts(InvocationCounts);
+	GetLodAndInvocationElementCounts(LodIndex, InvocationCounts);
 	
 	return new FOptimusTransientBufferDataProviderProxy(InvocationCounts, ElementStride, RawStride);
 }
@@ -341,10 +345,11 @@ bool UOptimusPersistentBufferDataProvider::IsValid() const
 
 FComputeDataProviderRenderProxy* UOptimusPersistentBufferDataProvider::GetRenderProxy()
 {
+	int32 LodIndex = 0;
 	TArray<int32> InvocationCounts;
-	GetInvocationElementCounts(InvocationCounts);
+	GetLodAndInvocationElementCounts(LodIndex, InvocationCounts);
 	
-	return new FOptimusPersistentBufferDataProviderProxy(InvocationCounts, ElementStride, RawStride, BufferPool, ResourceName, LODIndex);
+	return new FOptimusPersistentBufferDataProviderProxy(InvocationCounts, ElementStride, RawStride, BufferPool, ResourceName, LodIndex);
 }
 
 
diff --git a/Engine/Plugins/Animation/DeformerGraph/Source/OptimusCore/Private/DataInterfaces/OptimusDataInterfaceRawBuffer.h b/Engine/Plugins/Animation/DeformerGraph/Source/OptimusCore/Private/DataInterfaces/OptimusDataInterfaceRawBuffer.h
--- a/Engine/Plugins/Animation/DeformerGraph/Source/OptimusCore/Private/DataInterfaces/OptimusDataInterfaceRawBuffer.h
+++ b/Engine/Plugins/Animation/DeformerGraph/Source/OptimusCore/Private/DataInterfaces/OptimusDataInterfaceRawBuffer.h
@@ -149,7 +149,8 @@ public:
 	 *  object and a data domain. Uses the execution domains to compute the information. Returns false if the
 	 *  data domain is not valid for computation.
 	 */
-	bool GetInvocationElementCounts(
+	bool GetLodAndInvocationElementCounts(
+		int32& OutLodIndex,
 		TArray<int32>& OutInvocationElementCounts
 		) const;
 	
@@ -202,9 +203,6 @@ public:
 	 */
 	TSharedPtr<FOptimusPersistentBufferPool> BufferPool;
 
-	/** LOD Index of the primary mesh */
-	int32 LODIndex = 0;
-
 	/** The resource this buffer is provider to */
 	FName ResourceName;
 };
diff --git a/Engine/Plugins/Animation/DeformerGraph/Source/OptimusCore/Private/OptimusDeformerInstance.cpp b/Engine/Plugins/Animation/DeformerGraph/Source/OptimusCore/Private/OptimusDeformerInstance.cpp
--- a/Engine/Plugins/Animation/DeformerGraph/Source/OptimusCore/Private/OptimusDeformerInstance.cpp
+++ b/Engine/Plugins/Animation/DeformerGraph/Source/OptimusCore/Private/OptimusDeformerInstance.cpp
@@ -332,7 +332,6 @@ void UOptimusDeformerInstance::SetupFromDeformer(UOptimusDeformer* InDeformer)
 			if (UOptimusPersistentBufferDataProvider* PersistentBufferProvider = Cast<UOptimusPersistentBufferDataProvider>(DataProvider))
 			{
 				PersistentBufferProvider->BufferPool = BufferPool;
-				PersistentBufferProvider->LODIndex = LODIndex;
 			}
 		}
 
diff --git a/Engine/Plugins/Animation/DeformerGraph/Source/OptimusCore/Public/OptimusComponentSource.h b/Engine/Plugins/Animation/DeformerGraph/Source/OptimusCore/Public/OptimusComponentSource.h
--- a/Engine/Plugins/Animation/DeformerGraph/Source/OptimusCore/Public/OptimusComponentSource.h
+++ b/Engine/Plugins/Animation/DeformerGraph/Source/OptimusCore/Public/OptimusComponentSource.h
@@ -37,12 +37,19 @@ public:
 	virtual TArray<FName> GetExecutionDomains() const
 	PURE_VIRTUAL(UOptimusComponentSource::GetExecutionContexts, return {}; );
 
-	/** For a given execution domain, return the range of the domain as given by the component that is associated
-	 *  with this source.
+	/** Returns the current LOD index associated with a component associated with this source. */
+	virtual int32 GetLodIndex(const UActorComponent* InComponent) const
+	{
+		return 0;
+	}
+
+	/** For a given execution domain, and component LOD index, return the range of the domain as given by 
+	 *  the component that is associated with this source.
 	 */
 	virtual bool GetComponentElementCountsForExecutionDomain(
 		FName InDomainName,
 		const UActorComponent* InComponent,
+		int32 InLodIndex,
 		TArray<int32>& OutInvocationElementCounts
 		) const
 	{
diff --git a/Engine/Plugins/Animation/MLDeformer/MLDeformerFramework/Source/MLDeformerFramework/Private/MLDeformerComponentSource.cpp b/Engine/Plugins/Animation/MLDeformer/MLDeformerFramework/Source/MLDeformerFramework/Private/MLDeformerComponentSource.cpp
--- a/Engine/Plugins/Animation/MLDeformer/MLDeformerFramework/Source/MLDeformerFramework/Private/MLDeformerComponentSource.cpp
+++ b/Engine/Plugins/Animation/MLDeformer/MLDeformerFramework/Source/MLDeformerFramework/Private/MLDeformerComponentSource.cpp
@@ -33,9 +33,20 @@ TArray<FName> UMLDeformerComponentSource::GetExecutionDomains() const
 }
 
 
+int32 UMLDeformerComponentSource::GetLodIndex(const UActorComponent* InComponent) const
+{
+	const UMLDeformerComponent* DeformerComponent = Cast<UMLDeformerComponent>(InComponent);
+	const UMLDeformerModelInstance* ModelInstance = DeformerComponent ? DeformerComponent->GetModelInstance() : nullptr;
+	const USkeletalMeshComponent* SkeletalMeshComponent = ModelInstance ? ModelInstance->GetSkeletalMeshComponent() : nullptr;
+	const FSkeletalMeshObject* SkeletalMeshObject = SkeletalMeshComponent ? SkeletalMeshComponent->MeshObject : nullptr;
+	return SkeletalMeshObject ? SkeletalMeshObject->GetLOD() : 0;
+}
+
+
 bool UMLDeformerComponentSource::GetComponentElementCountsForExecutionDomain(
 	FName InDomainName,
 	const UActorComponent* InComponent,
+	int32 InLodIndex,
 	TArray<int32>& OutElementCounts
 	) const
 {
@@ -61,9 +72,8 @@ bool UMLDeformerComponentSource::GetComponentElementCountsForExecutionDomain(
 	{
 		return false;
 	}
-	const int32 LodIndex = SkeletalMeshObject->GetLOD();
 	FSkeletalMeshRenderData const& SkeletalMeshRenderData = SkeletalMeshObject->GetSkeletalMeshRenderData();
-	FSkeletalMeshLODRenderData const* LodRenderData = &SkeletalMeshRenderData.LODRenderData[LodIndex];
+	FSkeletalMeshLODRenderData const* LodRenderData = &SkeletalMeshRenderData.LODRenderData[InLodIndex];
 		
 	OutElementCounts.Reset(LodRenderData->RenderSections.Num());
 	for (FSkelMeshRenderSection const& RenderSection: LodRenderData->RenderSections)
diff --git a/Engine/Plugins/Animation/MLDeformer/MLDeformerFramework/Source/MLDeformerFramework/Private/MLDeformerComponentSource.h b/Engine/Plugins/Animation/MLDeformer/MLDeformerFramework/Source/MLDeformerFramework/Private/MLDeformerComponentSource.h
--- a/Engine/Plugins/Animation/MLDeformer/MLDeformerFramework/Source/MLDeformerFramework/Private/MLDeformerComponentSource.h
+++ b/Engine/Plugins/Animation/MLDeformer/MLDeformerFramework/Source/MLDeformerFramework/Private/MLDeformerComponentSource.h
@@ -23,5 +23,6 @@ public:
 	FName GetBindingName() const override { return FName("MLDeformer"); }
 	TSubclassOf<UActorComponent> GetComponentClass() const override;
 	TArray<FName> GetExecutionDomains() const override;
-	bool GetComponentElementCountsForExecutionDomain(FName InDomainName, const UActorComponent* InComponent, TArray<int32>& OutElementCounts) const override;
+	int32 GetLodIndex(const UActorComponent* InComponent) const override;
+	bool GetComponentElementCountsForExecutionDomain(FName InDomainName, const UActorComponent* InComponent, int32 InLodIndex, TArray<int32>& OutElementCounts) const override;
 };
