ComputeFramework: Move to use HLSL templates for data interface shader code.
This makes the code slightly less "magic" and easier to read or write.
It's also the same pattern that Niagara uses, so might make it easier for people working in the two systems.
#rnx

diff --git a/Engine/Plugins/Animation/DeformerGraph/Shaders/Private/DataInterfaceCloth.ush b/Engine/Plugins/Animation/DeformerGraph/Shaders/Private/DataInterfaceCloth.ush
--- a/Engine/Plugins/Animation/DeformerGraph/Shaders/Private/DataInterfaceCloth.ush
+++ b/Engine/Plugins/Animation/DeformerGraph/Shaders/Private/DataInterfaceCloth.ush
@@ -1,29 +1,7 @@
 // Copyright Epic Games, Inc. All Rights Reserved.
 
-#include "/Plugin/ComputeFramework/Private/ComputeKernelCommon.ush"
-
-uint DI_LOCAL(NumVertices);
-uint DI_LOCAL(InputStreamStart);
-uint DI_LOCAL(NumInfluencesPerVertex);
-float DI_LOCAL(ClothBlendWeight);
-/** Transform from cloth space (relative to cloth root bone) to local(component) space */
-float4x4 DI_LOCAL(ClothToLocal);
-Buffer<float4> DI_LOCAL(ClothBuffer);
-Buffer<float2> DI_LOCAL(ClothPositionsAndNormalsBuffer);
-
-DI_IMPL_READ(ReadNumVertices, uint, )
-{
-	return DI_LOCAL(NumVertices);
-}
-
-DI_IMPL_READ(ReadClothToLocal, float4x4, )
-{
-#if ENABLE_DEFORMER_CLOTH
-	return DI_LOCAL(ClothToLocal);
-#else
-	return float4x4(1, 0, 0, 0, 0, 1, 0, 0, 0, 0, 1, 0, 0, 0, 0, 1);
-#endif
-}
+#ifndef DATA_INTERFACE_CLOTH_ONCE
+#define DATA_INTERFACE_CLOTH_ONCE 1
 
 struct FVertexTriangleInfluence
 {
@@ -34,20 +12,52 @@ struct FVertexTriangleInfluence
 	float Weight;
 };
 
-FVertexTriangleInfluence DI_LOCAL(GetClothInfluence)(uint VertexIndex, uint InfluenceIndex)
+struct FClothResult
 {
-	const uint VertexOffset = DI_LOCAL(NumInfluencesPerVertex) * VertexIndex + DI_LOCAL(InputStreamStart);
+	float ClothWeight;
+	float3 Position;
+	float3 TangentX;
+	float3 TangentZ;
+};
+
+#endif // DATA_INTERFACE_CLOTH_ONCE
+
+uint {DataInterfaceName}_NumVertices;
+uint {DataInterfaceName}_InputStreamStart;
+uint {DataInterfaceName}_NumInfluencesPerVertex;
+float {DataInterfaceName}_ClothBlendWeight;
+float4x4 {DataInterfaceName}_ClothToLocal;
+Buffer<float4> {DataInterfaceName}_ClothBuffer;
+Buffer<float2> {DataInterfaceName}_ClothPositionsAndNormalsBuffer;
+
+uint ReadNumVertices_{DataInterfaceName}()
+{
+	return {DataInterfaceName}_NumVertices;
+}
+
+float4x4 ReadClothToLocal_{DataInterfaceName}()
+{
+#if ENABLE_DEFORMER_CLOTH
+	return {DataInterfaceName}_ClothToLocal;
+#else
+	return float4x4(1, 0, 0, 0, 0, 1, 0, 0, 0, 0, 1, 0, 0, 0, 0, 1);
+#endif
+}
+
+FVertexTriangleInfluence GetClothInfluence_{DataInterfaceName}(uint VertexIndex, uint InfluenceIndex)
+{
+	const uint VertexOffset = {DataInterfaceName}_NumInfluencesPerVertex * VertexIndex + {DataInterfaceName}_InputStreamStart;
 
 	const uint NUM_FLOAT4S_PER_VERTEX_INFLUENCE = 4; // Stride of GPUSkinApexCloth
 	const uint Offset = (VertexOffset + InfluenceIndex) * NUM_FLOAT4S_PER_VERTEX_INFLUENCE;
 
 	FVertexTriangleInfluence Influence;
 
-	Influence.PositionBaryCoordsAndDist = DI_LOCAL(ClothBuffer)[Offset];
-	Influence.NormalBaryCoordsAndDist = DI_LOCAL(ClothBuffer)[Offset + 1];
-	Influence.TangentBaryCoordsAndDist = DI_LOCAL(ClothBuffer)[Offset + 2];
+	Influence.PositionBaryCoordsAndDist = {DataInterfaceName}_ClothBuffer[Offset];
+	Influence.NormalBaryCoordsAndDist = {DataInterfaceName}_ClothBuffer[Offset + 1];
+	Influence.TangentBaryCoordsAndDist = {DataInterfaceName}_ClothBuffer[Offset + 2];
 
-	uint4 PackedIndices = asuint(DI_LOCAL(ClothBuffer)[Offset + 3]);
+	uint4 PackedIndices = asuint({DataInterfaceName}_ClothBuffer[Offset + 3]);
 	Influence.SourceMeshVertIndices.yw = (PackedIndices.xy >> 16) & 0xffff;
 	Influence.SourceMeshVertIndices.xz = PackedIndices.xy & 0xffff;
 	
@@ -56,15 +66,7 @@ FVertexTriangleInfluence DI_LOCAL(GetClothInfluence)(uint VertexIndex, uint Infl
 	return Influence;
 }
 
-struct FClothResult
-{
-	float ClothWeight;
-	float3 Position;
-	float3 TangentX;
-	float3 TangentZ;
-};
-
-FClothResult DI_LOCAL(GetClothResult)(uint VertexIndex)
+FClothResult GetClothResult_{DataInterfaceName}(uint VertexIndex)
 {
 	float ClothWeight = 0;
 	float3 SimulatedPosition = float3(0, 0, 0);
@@ -74,26 +76,26 @@ FClothResult DI_LOCAL(GetClothResult)(uint VertexIndex)
 	int NumInfluences = 0;
 	float SumWeights = 0;
 
-	for (uint i = 0; i < DI_LOCAL(NumInfluencesPerVertex); ++i)
+	for (uint i = 0; i < {DataInterfaceName}_NumInfluencesPerVertex; ++i)
 	{
-		const FVertexTriangleInfluence Influence = DI_LOCAL(GetClothInfluence)(VertexIndex, i);
+		const FVertexTriangleInfluence Influence = GetClothInfluence_{DataInterfaceName}(VertexIndex, i);
 
 		if (Influence.SourceMeshVertIndices.w < 0xFFFF)
 		{
 			++NumInfluences;
 
-			float3 A = float3(DI_LOCAL(ClothPositionsAndNormalsBuffer)[Influence.SourceMeshVertIndices.x * 3], DI_LOCAL(ClothPositionsAndNormalsBuffer)[Influence.SourceMeshVertIndices.x * 3 + 1].x);
-			float3 B = float3(DI_LOCAL(ClothPositionsAndNormalsBuffer)[Influence.SourceMeshVertIndices.y * 3], DI_LOCAL(ClothPositionsAndNormalsBuffer)[Influence.SourceMeshVertIndices.y * 3 + 1].x);
-			float3 C = float3(DI_LOCAL(ClothPositionsAndNormalsBuffer)[Influence.SourceMeshVertIndices.z * 3], DI_LOCAL(ClothPositionsAndNormalsBuffer)[Influence.SourceMeshVertIndices.z * 3 + 1].x);
+			float3 A = float3({DataInterfaceName}_ClothPositionsAndNormalsBuffer[Influence.SourceMeshVertIndices.x * 3], {DataInterfaceName}_ClothPositionsAndNormalsBuffer[Influence.SourceMeshVertIndices.x * 3 + 1].x);
+			float3 B = float3({DataInterfaceName}_ClothPositionsAndNormalsBuffer[Influence.SourceMeshVertIndices.y * 3], {DataInterfaceName}_ClothPositionsAndNormalsBuffer[Influence.SourceMeshVertIndices.y * 3 + 1].x);
+			float3 C = float3({DataInterfaceName}_ClothPositionsAndNormalsBuffer[Influence.SourceMeshVertIndices.z * 3], {DataInterfaceName}_ClothPositionsAndNormalsBuffer[Influence.SourceMeshVertIndices.z * 3 + 1].x);
 
-			float3 NA = float3(DI_LOCAL(ClothPositionsAndNormalsBuffer)[Influence.SourceMeshVertIndices.x * 3 + 1].y, DI_LOCAL(ClothPositionsAndNormalsBuffer)[Influence.SourceMeshVertIndices.x * 3 + 2]);
-			float3 NB = float3(DI_LOCAL(ClothPositionsAndNormalsBuffer)[Influence.SourceMeshVertIndices.y * 3 + 1].y, DI_LOCAL(ClothPositionsAndNormalsBuffer)[Influence.SourceMeshVertIndices.y * 3 + 2]);
-			float3 NC = float3(DI_LOCAL(ClothPositionsAndNormalsBuffer)[Influence.SourceMeshVertIndices.z * 3 + 1].y, DI_LOCAL(ClothPositionsAndNormalsBuffer)[Influence.SourceMeshVertIndices.z * 3 + 2]);
+			float3 NA = float3({DataInterfaceName}_ClothPositionsAndNormalsBuffer[Influence.SourceMeshVertIndices.x * 3 + 1].y, {DataInterfaceName}_ClothPositionsAndNormalsBuffer[Influence.SourceMeshVertIndices.x * 3 + 2]);
+			float3 NB = float3({DataInterfaceName}_ClothPositionsAndNormalsBuffer[Influence.SourceMeshVertIndices.y * 3 + 1].y, {DataInterfaceName}_ClothPositionsAndNormalsBuffer[Influence.SourceMeshVertIndices.y * 3 + 2]);
+			float3 NC = float3({DataInterfaceName}_ClothPositionsAndNormalsBuffer[Influence.SourceMeshVertIndices.z * 3 + 1].y, {DataInterfaceName}_ClothPositionsAndNormalsBuffer[Influence.SourceMeshVertIndices.z * 3 + 2]);
 
-			ClothWeight += DI_LOCAL(ClothBlendWeight) * (1.0f - (Influence.SourceMeshVertIndices.w / 65535.0f));
+			ClothWeight += {DataInterfaceName}_ClothBlendWeight * (1.0f - (Influence.SourceMeshVertIndices.w / 65535.0f));
 
 			float Weight = 1.0f;
-			if (DI_LOCAL(NumInfluencesPerVertex) > 1)
+			if ({DataInterfaceName}_NumInfluencesPerVertex > 1)
 			{
 				// Weight is packed in the last coordinate
 				Weight = Influence.Weight;
@@ -139,7 +141,7 @@ FClothResult DI_LOCAL(GetClothResult)(uint VertexIndex)
 		TangentPosition *= InvWeight;
 		NormalPosition *= InvWeight;
 
-		Result.ClothWeight = ClothWeight / (float)DI_LOCAL(NumInfluencesPerVertex);
+		Result.ClothWeight = ClothWeight / (float) {DataInterfaceName}_NumInfluencesPerVertex;
 
 		Result.Position = SimulatedPosition;
 		
@@ -147,8 +149,8 @@ FClothResult DI_LOCAL(GetClothResult)(uint VertexIndex)
 		Result.TangentZ = normalize(NormalPosition - SimulatedPosition);
 
 		// Cloth data are all in world space so need to change into local space
-		Result.TangentX = mul(Result.TangentX, (half3x3)DI_LOCAL(ClothToLocal));
-		Result.TangentZ = mul(Result.TangentZ, (half3x3)DI_LOCAL(ClothToLocal));
+		Result.TangentX = mul(Result.TangentX, (half3x3) {DataInterfaceName}_ClothToLocal);
+		Result.TangentZ = mul(Result.TangentZ, (half3x3) {DataInterfaceName}_ClothToLocal);
 	}
 
 	return Result;
@@ -158,40 +160,40 @@ FClothResult DI_LOCAL(GetClothResult)(uint VertexIndex)
 // There is the hope that the shader compiler may optimize here.
 // But really we will need to work out a way for Compute Kernels to do some shared work once and carry the results of that in an opaque context struct.
 
-DI_IMPL_READ(ReadClothWeight, float, uint VertexIndex)
+float ReadClothWeight_{DataInterfaceName}(uint VertexIndex)
 {
 #if ENABLE_DEFORMER_CLOTH
-	FClothResult Result = DI_LOCAL(GetClothResult)(VertexIndex);
+	FClothResult Result = GetClothResult_{DataInterfaceName}(VertexIndex);
 	return Result.ClothWeight;
 #else
 	return 0;
 #endif
 }
 
-DI_IMPL_READ(ReadClothPosition, float3, uint VertexIndex)
+float3 ReadClothPosition_{DataInterfaceName}(uint VertexIndex)
 {
 #if ENABLE_DEFORMER_CLOTH
-	FClothResult Result = DI_LOCAL(GetClothResult)(VertexIndex);
+	FClothResult Result = GetClothResult_{DataInterfaceName}(VertexIndex);
 	return Result.Position;
 #else
 	return float3(0, 0, 0);
 #endif
 }
 
-DI_IMPL_READ(ReadClothTangentX, float3, uint VertexIndex)
+float3 ReadClothTangentX_{DataInterfaceName}(uint VertexIndex)
 {
 #if ENABLE_DEFORMER_CLOTH
-	FClothResult Result = DI_LOCAL(GetClothResult)(VertexIndex);
+	FClothResult Result = GetClothResult_{DataInterfaceName}(VertexIndex);
 	return Result.TangentX;
 #else
 	return float3(1, 0, 0);
 #endif
 }
 
-DI_IMPL_READ(ReadClothTangentZ, float3, uint VertexIndex)
+float3 ReadClothTangentZ_{DataInterfaceName}(uint VertexIndex)
 {
 #if ENABLE_DEFORMER_CLOTH
-	FClothResult Result = DI_LOCAL(GetClothResult)(VertexIndex);
+	FClothResult Result = GetClothResult_{DataInterfaceName}(VertexIndex);
 	return Result.TangentZ;
 #else
 	return float3(0, 0, 1);
diff --git a/Engine/Plugins/Animation/DeformerGraph/Shaders/Private/DataInterfaceConnectivity.ush b/Engine/Plugins/Animation/DeformerGraph/Shaders/Private/DataInterfaceConnectivity.ush
--- a/Engine/Plugins/Animation/DeformerGraph/Shaders/Private/DataInterfaceConnectivity.ush
+++ b/Engine/Plugins/Animation/DeformerGraph/Shaders/Private/DataInterfaceConnectivity.ush
@@ -1,26 +1,24 @@
 // Copyright Epic Games, Inc. All Rights Reserved.
 
-#include "/Plugin/ComputeFramework/Private/ComputeKernelCommon.ush"
+uint {DataInterfaceName}_NumVertices;
+uint {DataInterfaceName}_InputStreamStart;
+uint {DataInterfaceName}_MaxConnectedVertexCount;
+Buffer<uint> {DataInterfaceName}_ConnectivityBuffer;
 
-uint DI_LOCAL(NumVertices);
-uint DI_LOCAL(InputStreamStart);
-uint DI_LOCAL(MaxConnectedVertexCount);
-Buffer<uint> DI_LOCAL(ConnectivityBuffer);
-
-DI_IMPL_READ(ReadNumVertices, uint, )
+uint ReadNumVertices_{DataInterfaceName}()
 {
-	return DI_LOCAL(NumVertices);
+	return {DataInterfaceName}_NumVertices;
 }
 
-DI_IMPL_READ(ReadNumConnectedVertices, uint, uint VertexIndex)
+uint ReadNumConnectedVertices_{DataInterfaceName}(uint VertexIndex)
 {
-	return DI_LOCAL(MaxConnectedVertexCount);
+	return {DataInterfaceName}_MaxConnectedVertexCount;
 }
 
-DI_IMPL_READ(ReadConnectedVertex, uint, uint VertexIndex, uint ConnectedVertexSubIndex)
+uint ReadConnectedVertex_{DataInterfaceName}(uint VertexIndex, uint ConnectedVertexSubIndex)
 {
-	uint BaseIndex = (DI_LOCAL(InputStreamStart) + VertexIndex) * DI_LOCAL(MaxConnectedVertexCount);
-	uint ConnectedVertexIndex = DI_LOCAL(ConnectivityBuffer)[BaseIndex + ConnectedVertexSubIndex];
-	uint BaseOffset = ConnectedVertexIndex == 0xFFFFFFFF ? 0 : DI_LOCAL(InputStreamStart);
+	uint BaseIndex = ({DataInterfaceName}_InputStreamStart + VertexIndex) * {DataInterfaceName}_MaxConnectedVertexCount;
+	uint ConnectedVertexIndex = {DataInterfaceName}_ConnectivityBuffer[BaseIndex + ConnectedVertexSubIndex];
+	uint BaseOffset = ConnectedVertexIndex == 0xFFFFFFFF ? 0 : {DataInterfaceName}_InputStreamStart;
 	return ConnectedVertexIndex - BaseOffset;
 }
diff --git a/Engine/Plugins/Animation/DeformerGraph/Shaders/Private/DataInterfaceDebugDraw.ush b/Engine/Plugins/Animation/DeformerGraph/Shaders/Private/DataInterfaceDebugDraw.ush
--- a/Engine/Plugins/Animation/DeformerGraph/Shaders/Private/DataInterfaceDebugDraw.ush
+++ b/Engine/Plugins/Animation/DeformerGraph/Shaders/Private/DataInterfaceDebugDraw.ush
@@ -1,7 +1,5 @@
 // Copyright Epic Games, Inc. All Rights Reserved.
 
-#include "/Plugin/ComputeFramework/Private/ComputeKernelCommon.ush"
-
 //todo[CF]:
 // Only included for [FLATTEN] used by MiniFontCommon.ush
 // We need to factor MiniFontCommon.ush out of ShaderPrint.ush
@@ -11,12 +9,11 @@
 // This is a temporary hack to insert our buffer name into ShaderPrint.ush
 // We need ShaderPrint to be flexible enough to take uniforms and buffers as parameters instead of using the hard coded uniform buffer.
 // Until that is done we can't support multiple of this data interface per kernel (although that would be an unusual requirement).
-#define ShaderPrint_StateBuffer		DI_LOCAL(ShaderPrint_StateBuffer)
-#define ShaderPrint_RWEntryBuffer	DI_LOCAL(ShaderPrint_RWEntryBuffer)
-
+#define ShaderPrint_StateBuffer		{DataInterfaceName}_ShaderPrint_StateBuffer
+#define ShaderPrint_RWEntryBuffer	{DataInterfaceName}_ShaderPrint_RWEntryBuffer
 #include "/Engine/Private/ShaderPrint.ush"
 
-float4x4 DI_LOCAL(LocalToWorld);
+float4x4 {DataInterfaceName}_LocalToWorld;
 
 struct FDebugDraw
 {
@@ -53,9 +50,9 @@ struct FDebugDraw
 	}
 };
 
-DI_IMPL_READ(ReadDebugDraw, FDebugDraw, )
+FDebugDraw ReadDebugDraw_{DataInterfaceName}()
 {
 	FDebugDraw DebugDraw;
-	DebugDraw.LocalToWorld = DI_LOCAL(LocalToWorld);
+	DebugDraw.LocalToWorld = {DataInterfaceName}_LocalToWorld;
 	return DebugDraw;
 }
diff --git a/Engine/Plugins/Animation/DeformerGraph/Shaders/Private/DataInterfaceJiggleSpring.ush b/Engine/Plugins/Animation/DeformerGraph/Shaders/Private/DataInterfaceJiggleSpring.ush
--- a/Engine/Plugins/Animation/DeformerGraph/Shaders/Private/DataInterfaceJiggleSpring.ush
+++ b/Engine/Plugins/Animation/DeformerGraph/Shaders/Private/DataInterfaceJiggleSpring.ush
@@ -1,81 +1,79 @@
 // Copyright Epic Games, Inc. All Rights Reserved.
 
-#include "/Plugin/ComputeFramework/Private/ComputeKernelCommon.ush"
+uint {DataInterfaceName}_NumVertices;
+uint {DataInterfaceName}_BaseVertexIndex; // InputStreamStart
 
-uint DI_LOCAL(NumVertices);
-uint DI_LOCAL(BaseVertexIndex); // InputStreamStart
+uint {DataInterfaceName}_NumVertexMap;
+uint {DataInterfaceName}_NumStiffnessWeights;
+uint {DataInterfaceName}_NumDampingWeights;
+uint {DataInterfaceName}_NumMaxStretchWeights;
 
-uint DI_LOCAL(NumVertexMap);
-uint DI_LOCAL(NumStiffnessWeights);
-uint DI_LOCAL(NumDampingWeights);
-uint DI_LOCAL(NumMaxStretchWeights);
+float {DataInterfaceName}_Stiffness;
+float {DataInterfaceName}_Damping;
+float {DataInterfaceName}_MaxStretch;
 
-float DI_LOCAL(Stiffness);
-float DI_LOCAL(Damping);
-float DI_LOCAL(MaxStretch);
+Buffer<float> {DataInterfaceName}_StiffnessWeightsBuffer;
+Buffer<float> {DataInterfaceName}_DampingWeightsBuffer;
+Buffer<uint> {DataInterfaceName}_VertexMapBuffer;
+Buffer<float> {DataInterfaceName}_MaxStretchWeightsBuffer;
 
-Buffer<float> DI_LOCAL(StiffnessWeightsBuffer);
-Buffer<float> DI_LOCAL(DampingWeightsBuffer);
-Buffer<uint> DI_LOCAL(VertexMapBuffer);
-Buffer<float> DI_LOCAL(MaxStretchWeightsBuffer);
-
-DI_IMPL_READ(ReadNumVertices, uint, )
+uint ReadNumVertices_{DataInterfaceName}()
 {
-	return DI_LOCAL(NumVertices);
+	return {DataInterfaceName}_NumVertices;
 }
 
-DI_IMPL_READ(ReadStiffness, float, uint VertexIndex)
+float ReadStiffness_{DataInterfaceName}(uint VertexIndex)
 {
 #if ENABLE_DEFORMER_JIGGLE_SPRING
-	if( //DI_LOCAL(NumStiffnessWeights) == DI_LOCAL(NumVertices) && 
-		VertexIndex < DI_LOCAL(NumStiffnessWeights))
+	if( //{DataInterfaceName}_NumStiffnessWeights == {DataInterfaceName}_NumVertices && 
+		VertexIndex < {DataInterfaceName}_NumStiffnessWeights)
 	{
-		uint BufferIndex = VertexIndex + DI_LOCAL(BaseVertexIndex);
-		return DI_LOCAL(StiffnessWeightsBuffer)[BufferIndex] * DI_LOCAL(Stiffness);
+		uint BufferIndex = VertexIndex + {DataInterfaceName}_BaseVertexIndex;
+		return {DataInterfaceName}_StiffnessWeightsBuffer[BufferIndex] * {DataInterfaceName}_Stiffness;
 	}
-	//else if(VertexIndex < DI_LOCAL(NumVertexMap))
+	//else if(VertexIndex < {DataInterfaceName}_NumVertexMap)
 	//{
-	//	uint BufferIndex = DI_LOCAL(VertexMapBuffer)[VertexIndex + DI_LOCAL(BaseVertexIndex)];
-	//	return DI_LOCAL(StiffnessWeightsBuffer)[BufferIndex] * DI_LOCAL(Stiffness);
+	//	uint BufferIndex = {DataInterfaceName}_VertexMapBuffer[VertexIndex + {DataInterfaceName}_BaseVertexIndex];
+	//	return {DataInterfaceName}_StiffnessWeightsBuffer[BufferIndex] * {DataInterfaceName}_Stiffness;
 	//}
 	else
 	{
-		return DI_LOCAL(Stiffness);
+		return {DataInterfaceName}_Stiffness;
 	}
 #else
 	return 0.0;
 #endif
 }
 
-DI_IMPL_READ(ReadDamping, float, uint VertexIndex)
+float ReadDamping_{DataInterfaceName}(uint VertexIndex)
 {
 #if ENABLE_DEFORMER_JIGGLE_SPRING
-	if( //DI_LOCAL(NumStiffnessWeights) == DI_LOCAL(NumVertices) && 
-		VertexIndex < DI_LOCAL(NumStiffnessWeights))
+	if( //{DataInterfaceName}_NumStiffnessWeights == {DataInterfaceName}_NumVertices && 
+		VertexIndex < {DataInterfaceName}_NumStiffnessWeights)
 	{
-		uint BufferIndex = VertexIndex + DI_LOCAL(BaseVertexIndex);
-		return DI_LOCAL(DampingWeightsBuffer)[BufferIndex] * DI_LOCAL(Damping);
+		uint BufferIndex = VertexIndex + {DataInterfaceName}_BaseVertexIndex;
+		return {DataInterfaceName}_DampingWeightsBuffer[BufferIndex] * {DataInterfaceName}_Damping;
 	}
-	//else if(VertexIndex < DI_LOCAL(NumVertexMap))
+	//else if(VertexIndex < {DataInterfaceName}_NumVertexMap)
 	//{
-	//	uint BufferIndex = DI_LOCAL(VertexMapBuffer)[VertexIndex + DI_LOCAL(BaseVertexIndex)];
-	//	return DI_LOCAL(DampingWeightsBuffer)[BufferIndex] * DI_LOCAL(Damping);
+	//	uint BufferIndex = {DataInterfaceName}_VertexMapBuffer[VertexIndex + {DataInterfaceName}_BaseVertexIndex];
+	//	return {DataInterfaceName}_DampingWeightsBuffer[BufferIndex] * {DataInterfaceName}_Damping;
 	//}
 	else
 	{
-		return DI_LOCAL(Damping);
+		return {DataInterfaceName}_Damping;
 	}
 #else
 	return 0.0;
 #endif
 }
 
-DI_IMPL_READ(ReadBufferIndex, uint, uint VertexIndex)
+uint ReadBufferIndex_{DataInterfaceName}(uint VertexIndex)
 {
 #if ENABLE_DEFORMER_JIGGLE_SPRING
-	if(VertexIndex < DI_LOCAL(NumVertexMap))
+	if(VertexIndex < {DataInterfaceName}_NumVertexMap)
 	{
-		uint BufferIndex = DI_LOCAL(VertexMapBuffer)[VertexIndex + DI_LOCAL(BaseVertexIndex)];
+		uint BufferIndex = {DataInterfaceName}_VertexMapBuffer[VertexIndex + {DataInterfaceName}_BaseVertexIndex];
 		return BufferIndex;
 	}
 #else
@@ -83,15 +81,15 @@ DI_IMPL_READ(ReadBufferIndex, uint, uint VertexIndex)
 #endif
 }
 
-DI_IMPL_READ(ReadMaxStretch, float, uint VertexIndex)
+float ReadMaxStretch_{DataInterfaceName}(uint VertexIndex)
 {
 #if ENABLE_DEFORMER_JIGGLE_SPRING
-	if(VertexIndex < DI_LOCAL(NumMaxStretchWeights))
+	if(VertexIndex < {DataInterfaceName}_NumMaxStretchWeights)
 	{
-		uint BufferIndex = VertexIndex + DI_LOCAL(BaseVertexIndex);
-		return abs(DI_LOCAL(MaxStretchWeightsBuffer[BufferIndex]) * DI_LOCAL(MaxStretch));
+		uint BufferIndex = VertexIndex + {DataInterfaceName}_BaseVertexIndex;
+		return abs({DataInterfaceName}_MaxStretchWeightsBuffer[BufferIndex] * {DataInterfaceName}_MaxStretch);
 	}
-	return abs(DI_LOCAL(MaxStretch));
+	return abs({DataInterfaceName}_MaxStretch);
 #else
 	return 100.0;
 #endif
@@ -119,8 +117,8 @@ struct FAnalyticSpring
 	{
 #if ENABLE_DEFORMER_JIGGLE_SPRING
 		// Sim parameters
-		float K = DI_LOCAL(ReadStiffness)(Index);
-		float D = DI_LOCAL(ReadDamping)(Index);
+		float K = ReadStiffness_{DataInterfaceName}(Index);
+		float D = ReadDamping_{DataInterfaceName}(Index);
 		if(K <= 0.0 || D <= 0.0)
 		{
 			// Set outputs to skinning
@@ -212,7 +210,7 @@ struct FAnalyticSpring
 		// Clamp deformation
 		float3 Dir = X - CurrSkinPos;
 		float DirLen = length(Dir);
-		float MaxLen = DI_LOCAL(ReadMaxStretch)(Index);
+		float MaxLen = ReadMaxStretch_{DataInterfaceName}(Index);
 		if(DirLen > MaxLen)
 		{
 			Dir /= DirLen;
@@ -235,7 +233,7 @@ struct FAnalyticSpring
 
 };
 
-DI_IMPL_READ(ReadAnalyticSpring, FAnalyticSpring, )
+FAnalyticSpring ReadAnalyticSpring_{DataInterfaceName}()
 {
 	FAnalyticSpring Spring;
 	return Spring;
diff --git a/Engine/Plugins/Animation/DeformerGraph/Shaders/Private/DataInterfaceMorphTarget.ush b/Engine/Plugins/Animation/DeformerGraph/Shaders/Private/DataInterfaceMorphTarget.ush
--- a/Engine/Plugins/Animation/DeformerGraph/Shaders/Private/DataInterfaceMorphTarget.ush
+++ b/Engine/Plugins/Animation/DeformerGraph/Shaders/Private/DataInterfaceMorphTarget.ush
@@ -1,32 +1,29 @@
 // Copyright Epic Games, Inc. All Rights Reserved.
 
-#include "/Plugin/ComputeFramework/Private/ComputeKernelCommon.ush"
+uint {DataInterfaceName}_NumVertices;
+uint {DataInterfaceName}_InputStreamStart;
+Buffer<float> {DataInterfaceName}_MorphBuffer;
 
-uint DI_LOCAL(NumVertices);
-uint DI_LOCAL(InputStreamStart);
-
-Buffer<float> DI_LOCAL(MorphBuffer);
-
-DI_IMPL_READ(ReadNumVertices, uint, )
+uint ReadNumVertices_{DataInterfaceName}()
 {
-	return DI_LOCAL(NumVertices);
+	return {DataInterfaceName}_NumVertices;
 }
 
-DI_IMPL_READ(ReadDeltaPosition, float3, uint VertexIndex)
+float3 ReadDeltaPosition_{DataInterfaceName}(uint VertexIndex)
 {
 #if ENABLE_DEFORMER_MORPHTARGET
-	uint BufferIndex = (DI_LOCAL(InputStreamStart) + VertexIndex) * 6;
-	return float3(DI_LOCAL(MorphBuffer)[BufferIndex], DI_LOCAL(MorphBuffer)[BufferIndex + 1], DI_LOCAL(MorphBuffer)[BufferIndex + 2]);
+	uint BufferIndex = ({DataInterfaceName}_InputStreamStart + VertexIndex) * 6;
+	return float3({DataInterfaceName}_MorphBuffer[BufferIndex], {DataInterfaceName}_MorphBuffer[BufferIndex + 1], {DataInterfaceName}_MorphBuffer[BufferIndex + 2]);
 #else
 	return float3(0, 0, 0);
 #endif
 }
 
-DI_IMPL_READ(ReadDeltaNormal, float3, uint VertexIndex)
+float3 ReadDeltaNormal_{DataInterfaceName}(uint VertexIndex)
 {
 #if ENABLE_DEFORMER_MORPHTARGET
-	uint BufferIndex = (DI_LOCAL(InputStreamStart) + VertexIndex) * 6;
-	return float3(DI_LOCAL(MorphBuffer)[BufferIndex + 3], DI_LOCAL(MorphBuffer)[BufferIndex + 4], DI_LOCAL(MorphBuffer)[BufferIndex + 5]);
+	uint BufferIndex = ({DataInterfaceName}_InputStreamStart + VertexIndex) * 6;
+	return float3({DataInterfaceName}_MorphBuffer[BufferIndex + 3], {DataInterfaceName}_MorphBuffer[BufferIndex + 4], {DataInterfaceName}_MorphBuffer[BufferIndex + 5]);
 #else
 	return float3(0, 0, 0);
 #endif
diff --git a/Engine/Plugins/Animation/DeformerGraph/Shaders/Private/DataInterfaceRawBuffer.ush b/Engine/Plugins/Animation/DeformerGraph/Shaders/Private/DataInterfaceRawBuffer.ush
--- a/Engine/Plugins/Animation/DeformerGraph/Shaders/Private/DataInterfaceRawBuffer.ush
+++ b/Engine/Plugins/Animation/DeformerGraph/Shaders/Private/DataInterfaceRawBuffer.ush
@@ -1,77 +1,89 @@
 // Copyright Epic Games, Inc. All Rights Reserved.
 
-#include "/Plugin/ComputeFramework/Private/ComputeKernelCommon.ush"
+#include "/Plugin/ComputeFramework/Private/BufferAlias.ush"
 
 // BUFFER_TYPE is type of underlying structure.
 // PUBLIC_TYPE is the exposed type that calling code expects to use.
 // Where BUFFER_TYPE is different to PUBLIC_TYPE we set BUFFER_TYPE_RAW and expect ReadRawBuffer() and WriteRawBuffer() to be implemented somewhere.
+#define PUBLIC_TYPE {PublicType}
+#define BUFFER_TYPE {BufferType}
+#define BUFFER_TYPE_RAW {BufferTypeRaw}
+// BUFFER_TYPE_SUPPORTS_ATOMIC is set if we can enable atomic operations on the PUBLIC_TYPE.
+#define BUFFER_TYPE_SUPPORTS_ATOMIC {SupportAtomic}
+// BUFFER_SPLIT_READ_WRITE is set if we can use an SRV to read the buffer. It is unset if we might need to read and write the buffer in the same kernel.
+#define BUFFER_SPLIT_READ_WRITE {SplitReadWrite}
 
-uint DI_LOCAL(StartOffset);
-uint DI_LOCAL(BufferSize);
-
+uint {DataInterfaceName}_StartOffset;
+uint {DataInterfaceName}_BufferSize;
 #if BUFFER_SPLIT_READ_WRITE
-StructuredBuffer<BUFFER_TYPE> DI_LOCAL(BufferSRV);
+StructuredBuffer<BUFFER_TYPE> {DataInterfaceName}_BufferSRV;
 #endif
-RWStructuredBuffer<BUFFER_TYPE> DI_LOCAL(BufferUAV);
+RWStructuredBuffer<BUFFER_TYPE> {DataInterfaceName}_BufferUAV;
 
-DI_IMPL_READ(ReadNumValues, uint, )
+uint ReadNumValues_{DataInterfaceName}()
 {
-	return DI_LOCAL(BufferSize);
+	return {DataInterfaceName}_BufferSize;
 }
 
-DI_IMPL_READ(ReadValue, PUBLIC_TYPE, uint Index)
+PUBLIC_TYPE ReadValue_{DataInterfaceName}(uint Index)
 {
-	uint BufferIndex = DI_LOCAL(StartOffset) + Index;
+	uint BufferIndex = {DataInterfaceName}_StartOffset + Index;
 	PUBLIC_TYPE ReturnValue;
 #if BUFFER_SPLIT_READ_WRITE
 #if BUFFER_TYPE_RAW
-	ReadRawBuffer(DI_LOCAL(BufferSRV), BufferIndex, ReturnValue);
+	ReadRawBuffer({DataInterfaceName}_BufferSRV, BufferIndex, ReturnValue);
 #else
-	ReturnValue = DI_LOCAL(BufferSRV)[BufferIndex];
+	ReturnValue = {DataInterfaceName}_BufferSRV[BufferIndex];
 #endif
 #else
 #if BUFFER_TYPE_RAW
-	ReadRawBuffer(DI_LOCAL(BufferUAV), BufferIndex, ReturnValue);
+	ReadRawBuffer({DataInterfaceName}_BufferUAV, BufferIndex, ReturnValue);
 #else
-	ReturnValue = DI_LOCAL(BufferUAV)[BufferIndex];
+	ReturnValue = {DataInterfaceName}_BufferUAV[BufferIndex];
 #endif
 #endif
 	return ReturnValue;
 }
 
-DI_IMPL_WRITE(WriteValue, uint Index, PUBLIC_TYPE Value)
+void WriteValue_{DataInterfaceName}(uint Index, PUBLIC_TYPE Value)
 {
-	uint BufferIndex = DI_LOCAL(StartOffset) + Index;
+	uint BufferIndex = {DataInterfaceName}_StartOffset + Index;
 #if BUFFER_TYPE_RAW
-	WriteRawBuffer(DI_LOCAL(BufferUAV), BufferIndex, Value);
+	WriteRawBuffer({DataInterfaceName}_BufferUAV, BufferIndex, Value);
 #else
-	DI_LOCAL(BufferUAV)[BufferIndex] = Value;
+	{DataInterfaceName}_BufferUAV[BufferIndex] = Value;
 #endif
 }
 
-DI_IMPL_WRITE(WriteAtomicAdd, uint Index, PUBLIC_TYPE Value)
+void WriteAtomicAdd_{DataInterfaceName}(uint Index, PUBLIC_TYPE Value)
 {
 #if BUFFER_TYPE_SUPPORTS_ATOMIC
-	uint BufferIndex = DI_LOCAL(StartOffset) + Index;
+	uint BufferIndex = {DataInterfaceName}_StartOffset + Index;
 	PUBLIC_TYPE OldValue;
-	InterlockedAdd(DI_LOCAL(BufferUAV)[BufferIndex], Value, OldValue);
+	InterlockedAdd({DataInterfaceName}_BufferUAV[BufferIndex], Value, OldValue);
 #endif
 }
 
-DI_IMPL_WRITE(WriteAtomicMin, uint Index, PUBLIC_TYPE Value)
+void WriteAtomicMin_{DataInterfaceName}(uint Index, PUBLIC_TYPE Value)
 {
 #if BUFFER_TYPE_SUPPORTS_ATOMIC
-	uint BufferIndex = DI_LOCAL(StartOffset) + Index;
+	uint BufferIndex = {DataInterfaceName}_StartOffset + Index;
 	PUBLIC_TYPE OldValue;
-	InterlockedMin(DI_LOCAL(BufferUAV)[BufferIndex], Value, OldValue);
+	InterlockedMin({DataInterfaceName}_BufferUAV[BufferIndex], Value, OldValue);
 #endif
 }
 
-DI_IMPL_WRITE(WriteAtomicMax, uint Index, PUBLIC_TYPE Value)
+void WriteAtomicMax_{DataInterfaceName}(uint Index, PUBLIC_TYPE Value)
 {
 #if BUFFER_TYPE_SUPPORTS_ATOMIC
-	uint BufferIndex = DI_LOCAL(StartOffset) + Index;
+	uint BufferIndex = {DataInterfaceName}_StartOffset + Index;
 	PUBLIC_TYPE OldValue;
-	InterlockedMax(DI_LOCAL(BufferUAV)[BufferIndex], Value, OldValue);
+	InterlockedMax({DataInterfaceName}_BufferUAV[BufferIndex], Value, OldValue);
 #endif
 }
+
+#undef PUBLIC_TYPE
+#undef BUFFER_TYPE
+#undef BUFFER_TYPE_RAW
+#undef BUFFER_TYPE_SUPPORTS_ATOMIC
+#undef BUFFER_SPLIT_READ_WRITE
diff --git a/Engine/Plugins/Animation/DeformerGraph/Shaders/Private/DataInterfaceScene.ush b/Engine/Plugins/Animation/DeformerGraph/Shaders/Private/DataInterfaceScene.ush
--- a/Engine/Plugins/Animation/DeformerGraph/Shaders/Private/DataInterfaceScene.ush
+++ b/Engine/Plugins/Animation/DeformerGraph/Shaders/Private/DataInterfaceScene.ush
@@ -1,22 +1,20 @@
 // Copyright Epic Games, Inc. All Rights Reserved.
 
-#include "/Plugin/ComputeFramework/Private/ComputeKernelCommon.ush"
+float {DataInterfaceName}_GameTime;
+float {DataInterfaceName}_GameTimeDelta;
+uint {DataInterfaceName}_FrameNumber;
 
-float DI_LOCAL(GameTime);
-float DI_LOCAL(GameTimeDelta);
-uint DI_LOCAL(FrameNumber);
-
-DI_IMPL_READ(ReadGameTime, float, )
+float ReadGameTime_{DataInterfaceName}()
 {
-	return DI_LOCAL(GameTime);
+	return {DataInterfaceName}_GameTime;
 }
 
-DI_IMPL_READ(ReadGameTimeDelta, float, )
+float ReadGameTimeDelta_{DataInterfaceName}()
 {
-	return DI_LOCAL(GameTimeDelta);
+	return {DataInterfaceName}_GameTimeDelta;
 }
 
-DI_IMPL_READ(ReadFrameNumber, uint, )
+float ReadFrameNumber_{DataInterfaceName}()
 {
-	return DI_LOCAL(FrameNumber);
+	return {DataInterfaceName}_FrameNumber;
 }
diff --git a/Engine/Plugins/Animation/DeformerGraph/Shaders/Private/DataInterfaceSkeletalMeshRead.ush b/Engine/Plugins/Animation/DeformerGraph/Shaders/Private/DataInterfaceSkeletalMeshRead.ush
deleted file mode 100644
diff --git a/Engine/Plugins/Animation/DeformerGraph/Shaders/Private/DataInterfaceSkeleton.ush b/Engine/Plugins/Animation/DeformerGraph/Shaders/Private/DataInterfaceSkeleton.ush
--- a/Engine/Plugins/Animation/DeformerGraph/Shaders/Private/DataInterfaceSkeleton.ush
+++ b/Engine/Plugins/Animation/DeformerGraph/Shaders/Private/DataInterfaceSkeleton.ush
@@ -1,45 +1,37 @@
 // Copyright Epic Games, Inc. All Rights Reserved.
 
-#include "/Plugin/ComputeFramework/Private/ComputeKernelCommon.ush"
-#include "/Engine/Private/GpuSkinCommon.ush"
-
-// The following should come from project settings and be set on compute graph kernel compilation.
-// GPUSKIN_BONE_INDEX_UINT16
-// GPUSKIN_UNLIMITED_BONE_INFLUENCE
-
-uint DI_LOCAL(NumVertices);
-uint DI_LOCAL(InputStreamStart);
-uint DI_LOCAL(NumBoneInfluences);
-uint DI_LOCAL(InputWeightStart);
-uint DI_LOCAL(InputWeightStride);
-uint DI_LOCAL(InputWeightIndexSize);
-
-Buffer<float4> DI_LOCAL(BoneMatrices);
-Buffer<uint> DI_LOCAL(InputWeightStream);
-Buffer<uint> DI_LOCAL(InputWeightLookupStream);
-
-DI_IMPL_READ(ReadNumVertices, uint, )
+uint {DataInterfaceName}_NumVertices;
+uint {DataInterfaceName}_InputStreamStart;
+uint {DataInterfaceName}_NumBoneInfluences;
+uint {DataInterfaceName}_InputWeightStart;
+uint {DataInterfaceName}_InputWeightStride;
+uint {DataInterfaceName}_InputWeightIndexSize;
+Buffer<float4> {DataInterfaceName}_BoneMatrices;
+Buffer<uint> {DataInterfaceName}_InputWeightStream;
+Buffer<uint> {DataInterfaceName}_InputWeightLookupStream;
+
+uint ReadNumVertices_{DataInterfaceName}()
 {
-	return DI_LOCAL(NumVertices);
+	return {DataInterfaceName}_NumVertices;
 }
 
 #if !GPUSKIN_UNLIMITED_BONE_INFLUENCE
 
-int4 DI_LOCAL(GetBlendIndices)(uint StreamOffset, uint ExtraInfluenceIndex)
+int4 GetBlendIndices_{DataInterfaceName}(uint StreamOffset, uint ExtraInfluenceIndex)
 {
 	int4 UnpackedBlendIndices = 0;
 
 #if GPUSKIN_BONE_INDEX_UINT16
 	StreamOffset += GPUSKIN_VB_OFFSET_INFLUENCEBONES + 2 * ExtraInfluenceIndex;
-	uint PackedBlendIndices = DI_LOCAL(InputWeightStream)[StreamOffset];
+	uint PackedBlendIndices = {DataInterfaceName}_InputWeightStream[StreamOffset];
 	UnpackedBlendIndices.x = PackedBlendIndices & 0xffff;
 	UnpackedBlendIndices.y = PackedBlendIndices >> 16 & 0xffff;
-	PackedBlendIndices = DI_LOCAL(InputWeightStream)[StreamOffset + 1];
+	PackedBlendIndices = {DataInterfaceName}_InputWeightStream[StreamOffset + 1];
 	UnpackedBlendIndices.z = PackedBlendIndices & 0xffff;
 	UnpackedBlendIndices.w = PackedBlendIndices >> 16 & 0xffff;
 #else // GPUSKIN_BONE_INDEX_UINT16
 	StreamOffset += GPUSKIN_VB_OFFSET_INFLUENCEBONES + ExtraInfluenceIndex;
-	uint PackedBlendIndices = DI_LOCAL(InputWeightStream)[StreamOffset];
+	uint PackedBlendIndices = {DataInterfaceName}_InputWeightStream[StreamOffset];
 	UnpackedBlendIndices.x = PackedBlendIndices & 0xff;
 	UnpackedBlendIndices.y = PackedBlendIndices >> 8 & 0xff;
 	UnpackedBlendIndices.z = PackedBlendIndices >> 16 & 0xff;
@@ -49,10 +41,10 @@ int4 DI_LOCAL(GetBlendIndices)(uint StreamOffset, uint ExtraInfluenceIndex)
 	return UnpackedBlendIndices;
 }
 
-float4 DI_LOCAL(GetBlendWeights)(uint StreamOffset, uint ExtraInfluenceIndex)
+float4 GetBlendWeights_{DataInterfaceName}(uint StreamOffset, uint ExtraInfluenceIndex)
 {
-	StreamOffset += (((DI_LOCAL(NumBoneInfluences) + 3) / 4) + ExtraInfluenceIndex) * (1 + GPUSKIN_BONE_INDEX_UINT16);
-	uint PackedBlendWeights = DI_LOCAL(InputWeightStream)[StreamOffset];
+	StreamOffset += ((({DataInterfaceName}_NumBoneInfluences + 3) / 4) + ExtraInfluenceIndex) * (1 + GPUSKIN_BONE_INDEX_UINT16);
+	uint PackedBlendWeights = {DataInterfaceName}_InputWeightStream[StreamOffset];
 	
 	float4 UnpackedBlendWeights = 0;
 	UnpackedBlendWeights.x = float(PackedBlendWeights & 0xff) / 255.0f;
@@ -64,25 +56,25 @@ float4 DI_LOCAL(GetBlendWeights)(uint StreamOffset, uint ExtraInfluenceIndex)
 
 #endif // GPUSKIN_UNLIMITED_BONE_INFLUENCE
 
-float3x4 DI_LOCAL(GetBoneMatrix)(uint BoneIndex)
+float3x4 GetBoneMatrix_{DataInterfaceName}(uint BoneIndex)
 {
 	uint BufferIndex = BoneIndex * 3;
-	return float3x4(DI_LOCAL(BoneMatrices)[BufferIndex], DI_LOCAL(BoneMatrices)[BufferIndex + 1], DI_LOCAL(BoneMatrices)[BufferIndex + 2]);
+	return float3x4({DataInterfaceName}_BoneMatrices[BufferIndex], {DataInterfaceName}_BoneMatrices[BufferIndex + 1], {DataInterfaceName}_BoneMatrices[BufferIndex + 2]);
 }
 
-DI_IMPL_READ(ReadNumBones, uint, uint VertexIndex)
+uint ReadNumBones_{DataInterfaceName}(uint VertexIndex)
 {
 #if !ENABLE_DEFORMER_BONES
 	return 1;
 #elif GPUSKIN_UNLIMITED_BONE_INFLUENCE
-	uint BlendOffsetCount = DI_LOCAL(InputWeightLookupStream)[VertexIndex];
+	uint BlendOffsetCount = {DataInterfaceName}_InputWeightLookupStream[VertexIndex];
 	return BlendOffsetCount & 0xff;
 #else // !GPUSKIN_UNLIMITED_BONE_INFLUENCE
-	return DI_LOCAL(NumBoneInfluences);
+	return {DataInterfaceName}_NumBoneInfluences;
 #endif
 }
 
-DI_IMPL_READ(ReadBoneMatrix, float3x4, uint VertexIndex, uint BoneIndex)
+float3x4 ReadBoneMatrix_{DataInterfaceName}(uint VertexIndex, uint BoneIndex)
 {
 #if !ENABLE_DEFORMER_BONES
 
@@ -90,28 +82,28 @@ DI_IMPL_READ(ReadBoneMatrix, float3x4, uint VertexIndex, uint BoneIndex)
 
 #elif GPUSKIN_UNLIMITED_BONE_INFLUENCE
 	
-	uint BlendOffsetCount = DI_LOCAL(InputWeightLookupStream)[VertexIndex];
+	uint BlendOffsetCount = {DataInterfaceName}_InputWeightLookupStream[VertexIndex];
 	int StreamOffset = BlendOffsetCount >> 8;
 
-	int BoneIndexOffset = StreamOffset + (DI_LOCAL(InputWeightIndexSize) * BoneIndex);
-	int BoneGlobalIndex = DI_LOCAL(InputWeightStream)[BoneIndexOffset];
-	if (DI_LOCAL(InputWeightIndexSize) > 1)
+	int BoneIndexOffset = StreamOffset + ({DataInterfaceName}_InputWeightIndexSize * BoneIndex);
+	int BoneGlobalIndex = {DataInterfaceName}_InputWeightStream[BoneIndexOffset];
+	if ({DataInterfaceName}_InputWeightIndexSize > 1)
 	{
-		BoneGlobalIndex = DI_LOCAL(InputWeightStream)[BoneIndexOffset + 1] << 8 | BoneGlobalIndex;
+		BoneGlobalIndex = {DataInterfaceName}_InputWeightStream[BoneIndexOffset + 1] << 8 | BoneGlobalIndex;
 	}
 	
-	return DI_LOCAL(GetBoneMatrix)(BoneGlobalIndex);
+	return GetBoneMatrix_{DataInterfaceName}(BoneGlobalIndex);
 
 #else // !GPUSKIN_UNLIMITED_BONE_INFLUENCE
 	
-	uint StreamOffset = DI_LOCAL(InputWeightStart) + (VertexIndex * (DI_LOCAL(InputWeightStride) / 4));
-	int BlendIndex = DI_LOCAL(GetBlendIndices)(StreamOffset, BoneIndex / 4)[BoneIndex & 0x3];
-	return DI_LOCAL(GetBoneMatrix)(BlendIndex);
+	uint StreamOffset = {DataInterfaceName}_InputWeightStart + (VertexIndex * ({DataInterfaceName}_InputWeightStride / 4));
+	int BlendIndex = GetBlendIndices_{DataInterfaceName}(StreamOffset, BoneIndex / 4)[BoneIndex & 0x3];
+	return GetBoneMatrix_{DataInterfaceName}(BlendIndex);
 
 #endif
 }
 
-DI_IMPL_READ(ReadBoneWeight, float, uint VertexIndex, uint BoneIndex)
+float ReadBoneWeight_{DataInterfaceName}(uint VertexIndex, uint BoneIndex)
 {
 #if !ENABLE_DEFORMER_BONES
 
@@ -119,23 +111,23 @@ DI_IMPL_READ(ReadBoneWeight, float, uint VertexIndex, uint BoneIndex)
 
 #elif GPUSKIN_UNLIMITED_BONE_INFLUENCE
 	
-	uint BlendOffsetCount = DI_LOCAL(InputWeightLookupStream)[VertexIndex];
+	uint BlendOffsetCount = {DataInterfaceName}_InputWeightLookupStream[VertexIndex];
 	int NumBoneInfluencesLocal = BlendOffsetCount & 0xff;
 	int StreamOffset = BlendOffsetCount >> 8;
-	int WeightsOffset = StreamOffset + (DI_LOCAL(InputWeightIndexSize) * NumBoneInfluencesLocal);
+	int WeightsOffset = StreamOffset + ({DataInterfaceName}_InputWeightIndexSize * NumBoneInfluencesLocal);
 
-	return float(DI_LOCAL(InputWeightStream)[WeightsOffset + BoneIndex]) / 255.0f;
+	return float({DataInterfaceName}_InputWeightStream[WeightsOffset + BoneIndex]) / 255.0f;
 
 #else // !GPUSKIN_UNLIMITED_BONE_INFLUENCE
 	
-	uint StreamOffset = DI_LOCAL(InputWeightStart) + (VertexIndex * (DI_LOCAL(InputWeightStride) / 4));
-	float BlendWeight = DI_LOCAL(GetBlendWeights)(StreamOffset, BoneIndex / 4)[BoneIndex & 0x3];
+	uint StreamOffset = {DataInterfaceName}_InputWeightStart + (VertexIndex * ({DataInterfaceName}_InputWeightStride / 4));
+	float BlendWeight = GetBlendWeights_{DataInterfaceName}(StreamOffset, BoneIndex / 4)[BoneIndex & 0x3];
 	return BlendWeight;
 
 #endif
 }
 
-DI_IMPL_READ(ReadWeightedBoneMatrix, float3x4, uint VertexIndex)
+float3x4 ReadWeightedBoneMatrix_{DataInterfaceName}(uint VertexIndex)
 {
 #if !ENABLE_DEFORMER_BONES
 
@@ -145,22 +137,22 @@ DI_IMPL_READ(ReadWeightedBoneMatrix, float3x4, uint VertexIndex)
 
 	float3x4 BlendMatrix = (float3x4)0;
 
-	uint BufferIndex = DI_LOCAL(InputStreamStart) + VertexIndex;
-	uint BlendOffsetCount = DI_LOCAL(InputWeightLookupStream)[BufferIndex];
+	uint BufferIndex = {DataInterfaceName}_InputStreamStart + VertexIndex;
+	uint BlendOffsetCount = {DataInterfaceName}_InputWeightLookupStream[BufferIndex];
 	int NumBoneInfluencesLocal = BlendOffsetCount & 0xff;
 	int IndicesOffset = BlendOffsetCount >> 8;
-	int WeightsOffset = IndicesOffset + (DI_LOCAL(InputWeightIndexSize) * NumBoneInfluencesLocal);
+	int WeightsOffset = IndicesOffset + ({DataInterfaceName}_InputWeightIndexSize * NumBoneInfluencesLocal);
 
 	for (int InfluenceIndex = 0; InfluenceIndex < NumBoneInfluencesLocal; InfluenceIndex++)
 	{
-		int BoneIndexOffset = IndicesOffset + (DI_LOCAL(InputWeightIndexSize) * InfluenceIndex);
-		int BoneIndex = DI_LOCAL(InputWeightStream)[BoneIndexOffset];
-		if (DI_LOCAL(InputWeightIndexSize) > 1)
+		int BoneIndexOffset = IndicesOffset + ({DataInterfaceName}_InputWeightIndexSize * InfluenceIndex);
+		int BoneIndex = {DataInterfaceName}_InputWeightStream[BoneIndexOffset];
+		if ({DataInterfaceName}_InputWeightIndexSize > 1)
 		{
-			BoneIndex = DI_LOCAL(InputWeightStream)[BoneIndexOffset + 1] << 8 | BoneIndex;
+			BoneIndex = {DataInterfaceName}_InputWeightStream[BoneIndexOffset + 1] << 8 | BoneIndex;
 		}
-		float BoneWeight = float(DI_LOCAL(InputWeightStream)[WeightsOffset + InfluenceIndex]) / 255.0f;
-		BlendMatrix += BoneWeight * DI_LOCAL(GetBoneMatrix)(BoneIndex);
+		float BoneWeight = float({DataInterfaceName}_InputWeightStream[WeightsOffset + InfluenceIndex]) / 255.0f;
+		BlendMatrix += BoneWeight * GetBoneMatrix_{DataInterfaceName}(BoneIndex);
 	}
 
 	return BlendMatrix;
@@ -169,29 +161,29 @@ DI_IMPL_READ(ReadWeightedBoneMatrix, float3x4, uint VertexIndex)
 
 	float3x4 BlendMatrix = (float3x4)0;
 
-	uint StreamOffset = DI_LOCAL(InputWeightStart) + (VertexIndex * (DI_LOCAL(InputWeightStride) / 4));
+	uint StreamOffset = {DataInterfaceName}_InputWeightStart + (VertexIndex * ({DataInterfaceName}_InputWeightStride / 4));
 
-	if (DI_LOCAL(NumBoneInfluences) > 0)
+	if ({DataInterfaceName}_NumBoneInfluences > 0)
 	{
-		int4 BlendIndices = DI_LOCAL(GetBlendIndices)(StreamOffset, 0);
-		float4 BlendWeights = DI_LOCAL(GetBlendWeights)(StreamOffset, 0);
+		int4 BlendIndices = GetBlendIndices_{DataInterfaceName}(StreamOffset, 0);
+		float4 BlendWeights = GetBlendWeights_{DataInterfaceName}(StreamOffset, 0);
 
-		BlendMatrix += BlendWeights.x * DI_LOCAL(GetBoneMatrix)(BlendIndices.x);
-		BlendMatrix += BlendWeights.y * DI_LOCAL(GetBoneMatrix)(BlendIndices.y);
-		BlendMatrix += BlendWeights.z * DI_LOCAL(GetBoneMatrix)(BlendIndices.z);
-		BlendMatrix += BlendWeights.w * DI_LOCAL(GetBoneMatrix)(BlendIndices.w);
+		BlendMatrix += BlendWeights.x * GetBoneMatrix_{DataInterfaceName}(BlendIndices.x);
+		BlendMatrix += BlendWeights.y * GetBoneMatrix_{DataInterfaceName}(BlendIndices.y);
+		BlendMatrix += BlendWeights.z * GetBoneMatrix_{DataInterfaceName}(BlendIndices.z);
+		BlendMatrix += BlendWeights.w * GetBoneMatrix_{DataInterfaceName}(BlendIndices.w);
 	}
 
 	// todo[CF]: make this a permutation when permutation support is working again?
-	if (DI_LOCAL(NumBoneInfluences) > 4)
+	if ({DataInterfaceName}_NumBoneInfluences > 4)
 	{
-		int4 BlendIndices = DI_LOCAL(GetBlendIndices)(StreamOffset, 1);
-		float4 BlendWeights = DI_LOCAL(GetBlendWeights)(StreamOffset, 1);
+		int4 BlendIndices = GetBlendIndices_{DataInterfaceName}(StreamOffset, 1);
+		float4 BlendWeights = GetBlendWeights_{DataInterfaceName}(StreamOffset, 1);
 
-		BlendMatrix += BlendWeights.x * DI_LOCAL(GetBoneMatrix)(BlendIndices.x);
-		BlendMatrix += BlendWeights.y * DI_LOCAL(GetBoneMatrix)(BlendIndices.y);
-		BlendMatrix += BlendWeights.z * DI_LOCAL(GetBoneMatrix)(BlendIndices.z);
-		BlendMatrix += BlendWeights.w * DI_LOCAL(GetBoneMatrix)(BlendIndices.w);
+		BlendMatrix += BlendWeights.x * GetBoneMatrix_{DataInterfaceName}(BlendIndices.x);
+		BlendMatrix += BlendWeights.y * GetBoneMatrix_{DataInterfaceName}(BlendIndices.y);
+		BlendMatrix += BlendWeights.z * GetBoneMatrix_{DataInterfaceName}(BlendIndices.z);
+		BlendMatrix += BlendWeights.w * GetBoneMatrix_{DataInterfaceName}(BlendIndices.w);
 	}
 
 	return BlendMatrix;
diff --git a/Engine/Plugins/Animation/DeformerGraph/Shaders/Private/DataInterfaceSkinnedMesh.ush b/Engine/Plugins/Animation/DeformerGraph/Shaders/Private/DataInterfaceSkinnedMesh.ush
--- a/Engine/Plugins/Animation/DeformerGraph/Shaders/Private/DataInterfaceSkinnedMesh.ush
+++ b/Engine/Plugins/Animation/DeformerGraph/Shaders/Private/DataInterfaceSkinnedMesh.ush
@@ -1,98 +1,97 @@
 // Copyright Epic Games, Inc. All Rights Reserved.
 
-#include "/Plugin/ComputeFramework/Private/ComputeKernelCommon.ush"
+// Include required for platform TANGENT_RWBUFFER_FORMAT.
 #include "/Engine/Private/GpuSkinCommon.ush"
 
-// todo: These should be set by compute graph permutations when that is working.
+// todo: This should be set by compute graph permutations.
 #define MERGE_DUPLICATED_VERTICES 0
 
-uint DI_LOCAL(NumVertices);
-uint DI_LOCAL(NumTriangles);
-uint DI_LOCAL(NumUVChannels);
-uint DI_LOCAL(IndexBufferStart);
-uint DI_LOCAL(InputStreamStart);
-
-Buffer<uint> DI_LOCAL(IndexBuffer);
-Buffer<float> DI_LOCAL(PositionInputBuffer);
-Buffer<SNORM float4> DI_LOCAL(TangentInputBuffer);
-Buffer<float4> DI_LOCAL(ColorInputBuffer);
-Buffer<float2> DI_LOCAL(UVInputBuffer);
-Buffer<uint> DI_LOCAL(DuplicatedIndicesIndices);
-Buffer<uint> DI_LOCAL(DuplicatedIndices);
-
-DI_IMPL_READ(ReadNumVertices, uint, )
+uint {DataInterfaceName}_NumVertices;
+uint {DataInterfaceName}_NumTriangles;
+uint {DataInterfaceName}_NumUVChannels;
+uint {DataInterfaceName}_IndexBufferStart;
+uint {DataInterfaceName}_InputStreamStart;
+Buffer<uint> {DataInterfaceName}_IndexBuffer;
+Buffer<float> {DataInterfaceName}_PositionInputBuffer;
+Buffer<SNORM float4> {DataInterfaceName}_TangentInputBuffer;
+Buffer<float4> {DataInterfaceName}_ColorInputBuffer;
+Buffer<float2> {DataInterfaceName}_UVInputBuffer;
+Buffer<uint> {DataInterfaceName}_DuplicatedIndicesIndices;
+Buffer<uint> {DataInterfaceName}_DuplicatedIndices;
+
+uint ReadNumVertices_{DataInterfaceName}()
 {
-	return DI_LOCAL(NumVertices);
+	return {DataInterfaceName}_NumVertices;
 }
 
-DI_IMPL_READ(ReadNumTriangles, uint, )
+uint ReadNumTriangles_{DataInterfaceName}()
 {
-	return DI_LOCAL(NumTriangles);
+	return {DataInterfaceName}_NumTriangles;
 }
 
-DI_IMPL_READ(ReadNumUVChannels, uint, )
+uint ReadNumUVChannels_{DataInterfaceName}()
 {
-	return DI_LOCAL(NumUVChannels);
+	return {DataInterfaceName}_NumUVChannels;
 }
 
-DI_IMPL_READ(ReadIndexBuffer, uint, uint Index)
+uint ReadIndexBuffer_{DataInterfaceName}(uint Index)
 {
-	uint BufferIndex = DI_LOCAL(IndexBufferStart) + Index;
-	return DI_LOCAL(IndexBuffer)[BufferIndex] - DI_LOCAL(InputStreamStart);
+	uint BufferIndex = {DataInterfaceName}_IndexBufferStart + Index;
+	return {DataInterfaceName}_IndexBuffer[BufferIndex] - {DataInterfaceName}_InputStreamStart;
 }
 
-DI_IMPL_READ(ReadPosition, float3, uint VertexIndex)
+float3 ReadPosition_{DataInterfaceName}(uint VertexIndex)
 {
-	uint BufferIndex = (DI_LOCAL(InputStreamStart) + VertexIndex) * 3;
-	return float3(DI_LOCAL(PositionInputBuffer)[BufferIndex], DI_LOCAL(PositionInputBuffer)[BufferIndex + 1], DI_LOCAL(PositionInputBuffer)[BufferIndex + 2]);
+	uint BufferIndex = ({DataInterfaceName}_InputStreamStart + VertexIndex) * 3;
+	return float3({DataInterfaceName}_PositionInputBuffer[BufferIndex], {DataInterfaceName}_PositionInputBuffer[BufferIndex + 1], {DataInterfaceName}_PositionInputBuffer[BufferIndex + 2]);
 }
 
-DI_IMPL_READ(ReadTangentX, float4, uint VertexIndex)
+float4 ReadTangentX_{DataInterfaceName}(uint VertexIndex)
 {
-	uint BufferIndex = (DI_LOCAL(InputStreamStart) + VertexIndex) * 2;
-	return TangentBias_SkinCache(DI_LOCAL(TangentInputBuffer)[BufferIndex]);
+	uint BufferIndex = ({DataInterfaceName}_InputStreamStart + VertexIndex) * 2;
+	return TangentBias_SkinCache({DataInterfaceName}_TangentInputBuffer[BufferIndex]);
 }
 
-DI_IMPL_READ(ReadTangentZ, float4, uint VertexIndex)
+float4 ReadTangentZ_{DataInterfaceName}(uint VertexIndex)
 {
-	uint BufferIndex = (DI_LOCAL(InputStreamStart) + VertexIndex) * 2;
-	return TangentBias_SkinCache(DI_LOCAL(TangentInputBuffer)[BufferIndex + 1]);
+	uint BufferIndex = ({DataInterfaceName}_InputStreamStart + VertexIndex) * 2;
+	return TangentBias_SkinCache({DataInterfaceName}_TangentInputBuffer[BufferIndex + 1]);
 }
 
-DI_IMPL_READ(ReadColor, float4, uint VertexIndex)
+float4 ReadColor_{DataInterfaceName}(uint VertexIndex)
 {
-	uint BufferIndex = DI_LOCAL(InputStreamStart) + VertexIndex;
-	return DI_LOCAL(ColorInputBuffer)[BufferIndex] FMANUALFETCH_COLOR_COMPONENT_SWIZZLE;
+	uint BufferIndex = {DataInterfaceName}_InputStreamStart + VertexIndex;
+	return {DataInterfaceName}_ColorInputBuffer[BufferIndex] FMANUALFETCH_COLOR_COMPONENT_SWIZZLE;
 }
 
-DI_IMPL_READ(ReadUV, float2, uint VertexIndex, uint UVChannelIndex)
+float2 ReadUV_{DataInterfaceName}(uint VertexIndex, uint UVChannelIndex)
 {
-	uint BufferIndex = (DI_LOCAL(InputStreamStart) + VertexIndex) * DI_LOCAL(NumUVChannels) + UVChannelIndex;
-	return DI_LOCAL(UVInputBuffer)[BufferIndex];
+	uint BufferIndex = ({DataInterfaceName}_InputStreamStart + VertexIndex) * {DataInterfaceName}_NumUVChannels + UVChannelIndex;
+	return {DataInterfaceName}_UVInputBuffer[BufferIndex];
 }
 
-DI_IMPL_READ(ReadDuplicatedIndicesStart, uint, uint VertexIndex)
+uint ReadDuplicatedIndicesStart_{DataInterfaceName}(uint VertexIndex)
 {
 #if MERGE_DUPLICATED_VERTICES 
-	return DI_LOCAL(DuplicatedIndicesIndices)[VertexIndex * 2 + 1];
+	return {DataInterfaceName}_DuplicatedIndicesIndices[VertexIndex * 2 + 1];
 #else
 	return 0;
 #endif
 }
 
-DI_IMPL_READ(ReadDuplicatedIndicesLength, uint, uint VertexIndex)
+uint ReadDuplicatedIndicesLength_{DataInterfaceName}(uint VertexIndex)
 {
 #if MERGE_DUPLICATED_VERTICES 
-	return DI_LOCAL(DuplicatedIndicesIndices)[VertexIndex * 2];
+	return {DataInterfaceName}_DuplicatedIndicesIndices[VertexIndex * 2];
 #else
 	return 0;
 #endif
 }
 
-DI_IMPL_READ(ReadDuplicatedIndex, uint, uint Index)
+uint ReadDuplicatedIndex_{DataInterfaceName}(uint Index)
 {
 #if MERGE_DUPLICATED_VERTICES 
-	return DI_LOCAL(DuplicatedIndices)[Index];
+	return {DataInterfaceName}_DuplicatedIndices[Index];
 #else
 	return 0;
 #endif
diff --git a/Engine/Plugins/Animation/DeformerGraph/Shaders/Private/DataInterfaceSkinnedMeshExec.ush b/Engine/Plugins/Animation/DeformerGraph/Shaders/Private/DataInterfaceSkinnedMeshExec.ush
--- a/Engine/Plugins/Animation/DeformerGraph/Shaders/Private/DataInterfaceSkinnedMeshExec.ush
+++ b/Engine/Plugins/Animation/DeformerGraph/Shaders/Private/DataInterfaceSkinnedMeshExec.ush
@@ -1,10 +1,8 @@
 // Copyright Epic Games, Inc. All Rights Reserved.
 
-#include "/Plugin/ComputeFramework/Private/ComputeKernelCommon.ush"
+int3 {DataInterfaceName}_NumThreads;
 
-int3 DI_LOCAL(NumThreads);
-
-DI_IMPL_READ(ReadNumThreads, int3, )
+int3 ReadNumThreads_{DataInterfaceName}()
 {
-	return DI_LOCAL(NumThreads);
+	return {DataInterfaceName}_NumThreads;
 }
diff --git a/Engine/Plugins/Animation/DeformerGraph/Shaders/Private/DataInterfaceSkinnedMeshWrite.ush b/Engine/Plugins/Animation/DeformerGraph/Shaders/Private/DataInterfaceSkinnedMeshWrite.ush
--- a/Engine/Plugins/Animation/DeformerGraph/Shaders/Private/DataInterfaceSkinnedMeshWrite.ush
+++ b/Engine/Plugins/Animation/DeformerGraph/Shaders/Private/DataInterfaceSkinnedMeshWrite.ush
@@ -1,42 +1,41 @@
 // Copyright Epic Games, Inc. All Rights Reserved.
 
-#include "/Plugin/ComputeFramework/Private/ComputeKernelCommon.ush"
+// Include required for platform TANGENT_RWBUFFER_FORMAT.
 #include "/Engine/Private/GpuSkinCommon.ush"
 
-uint DI_LOCAL(NumVertices);
-uint DI_LOCAL(OutputStreamStart);
+uint {DataInterfaceName}_NumVertices;
+uint {DataInterfaceName}_OutputStreamStart;
+RWBuffer<float> {DataInterfaceName}_PositionBufferUAV;
+RWBuffer<TANGENT_RWBUFFER_FORMAT> {DataInterfaceName}_TangentBufferUAV;
+RWBuffer<UNORM float4> {DataInterfaceName}_ColorBufferUAV;
 
-RWBuffer<float> DI_LOCAL(PositionBufferUAV);
-RWBuffer<TANGENT_RWBUFFER_FORMAT> DI_LOCAL(TangentBufferUAV);
-RWBuffer<UNORM float4> DI_LOCAL(ColorBufferUAV);
-
-DI_IMPL_READ(ReadNumVertices, uint, )
+uint ReadNumVertices_{DataInterfaceName}()
 {
-	return DI_LOCAL(NumVertices);
+	return {DataInterfaceName}_NumVertices;
 }
 
-DI_IMPL_WRITE(WritePosition, uint VertexIndex, float3 Position)
+void WritePosition_{DataInterfaceName}(uint VertexIndex, float3 Position)
 {
-	uint BufferIndex = (DI_LOCAL(OutputStreamStart) + VertexIndex) * 3;
-	DI_LOCAL(PositionBufferUAV)[BufferIndex + 0] = Position.x;
-	DI_LOCAL(PositionBufferUAV)[BufferIndex + 1] = Position.y;
-	DI_LOCAL(PositionBufferUAV)[BufferIndex + 2] = Position.z;
+	uint BufferIndex = ({DataInterfaceName}_OutputStreamStart + VertexIndex) * 3;
+	{DataInterfaceName}_PositionBufferUAV[BufferIndex + 0] = Position.x;
+	{DataInterfaceName}_PositionBufferUAV[BufferIndex + 1] = Position.y;
+	{DataInterfaceName}_PositionBufferUAV[BufferIndex + 2] = Position.z;
 }
 
-DI_IMPL_WRITE(WriteTangentX, uint VertexIndex, float4 TangentX)
+void WriteTangentX_{DataInterfaceName}(uint VertexIndex, float4 TangentX)
 {
-	uint BufferIndex = (DI_LOCAL(OutputStreamStart) + VertexIndex) * 2;
-	DI_LOCAL(TangentBufferUAV)[BufferIndex + 0] = TangentUnbias_SkinCache(TangentX);
+	uint BufferIndex = ({DataInterfaceName}_OutputStreamStart + VertexIndex) * 2;
+	{DataInterfaceName}_TangentBufferUAV[BufferIndex + 0] = TangentUnbias_SkinCache(TangentX);
 }
 
-DI_IMPL_WRITE(WriteTangentZ, uint VertexIndex, float4 TangentZ)
+void WriteTangentZ_{DataInterfaceName}(uint VertexIndex, float4 TangentZ)
 {
-	uint BufferIndex = (DI_LOCAL(OutputStreamStart) + VertexIndex) * 2;
-	DI_LOCAL(TangentBufferUAV)[BufferIndex + 1] = TangentUnbias_SkinCache(TangentZ);
+	uint BufferIndex = ({DataInterfaceName}_OutputStreamStart + VertexIndex) * 2;
+	{DataInterfaceName}_TangentBufferUAV[BufferIndex + 1] = TangentUnbias_SkinCache(TangentZ);
 }
 
-DI_IMPL_WRITE(WriteColor, uint VertexIndex, float4 Color)
+void WriteColor_{DataInterfaceName}(uint VertexIndex, float4 Color)
 {
-	uint BufferIndex = DI_LOCAL(OutputStreamStart) + VertexIndex;
-	DI_LOCAL(ColorBufferUAV)[BufferIndex] = Color;
+	uint BufferIndex = {DataInterfaceName}_OutputStreamStart + VertexIndex;
+	{DataInterfaceName}_ColorBufferUAV[BufferIndex] = Color;
 }
diff --git a/Engine/Plugins/Animation/DeformerGraph/Source/OptimusCore/Private/DataInterfaces/OptimusDataInterfaceAnimAttribute.cpp b/Engine/Plugins/Animation/DeformerGraph/Source/OptimusCore/Private/DataInterfaces/OptimusDataInterfaceAnimAttribute.cpp
--- a/Engine/Plugins/Animation/DeformerGraph/Source/OptimusCore/Private/DataInterfaces/OptimusDataInterfaceAnimAttribute.cpp
+++ b/Engine/Plugins/Animation/DeformerGraph/Source/OptimusCore/Private/DataInterfaces/OptimusDataInterfaceAnimAttribute.cpp
@@ -218,11 +218,8 @@ void UOptimusAnimAttributeDataInterface::GetShaderParameters(TCHAR const* UID, F
 	InOutBuilder.AddNestedStruct(UID, ShaderParameterMetadata);
 }
 
-void UOptimusAnimAttributeDataInterface::GetHLSL(FString& OutHLSL) const
+void UOptimusAnimAttributeDataInterface::GetHLSL(FString& OutHLSL, FString const& InDataInterfaceName) const
 {
-	// Need include for DI_LOCAL macro expansion.
-	OutHLSL += TEXT("#include \"/Plugin/ComputeFramework/Private/ComputeKernelCommon.ush\"\n");
-	
 	for (int32 Index = 0; Index < AttributeArray.Num(); Index++)
 	{
 		const FOptimusAnimAttributeDescription& Attribute = AttributeArray[Index];
@@ -235,19 +232,24 @@ void UOptimusAnimAttributeDataInterface::GetHLSL(FString& OutHLSL) const
 			if (!bIsStruct)
 			{
 				// Add uniforms.
-				OutHLSL += FString::Printf(TEXT("%s DI_LOCAL(%s);\n"), *TypeName, *Attribute.HlslId);
+				OutHLSL += FString::Printf(TEXT("%s %s_%s);\n"), 
+					*TypeName, 
+					*InDataInterfaceName, 
+					*Attribute.HlslId);
 					
 				// Add function getters.
-				OutHLSL += FString::Printf(TEXT("DI_IMPL_READ(Read%s, %s, )\n{\n\treturn DI_LOCAL(%s);\n}\n"),
-					*Attribute.HlslId,
+				OutHLSL += FString::Printf(TEXT("%s Read%s_%s()\n{\n\treturn %s_%s;\n}\n"), 
 					*TypeName,
+					*Attribute.HlslId,
+					*InDataInterfaceName, 
+					*InDataInterfaceName, 
 					*Attribute.HlslId);
 			}
 			else if (bIsStruct)
 			{
 				struct FStructParser 
 				{
-					void WalkStruct(FShaderValueTypeHandle ShaderValueTypeHandle, const TArray<FString>& InPrefix)
+					void WalkStruct(FString const& InDataInterfaceName, FShaderValueTypeHandle ShaderValueTypeHandle, const TArray<FString>& InPrefix)
 					{
 						const TArray<FShaderValueType::FStructElement>& StructElements =
 							ShaderValueTypeHandle->StructElements;
@@ -264,22 +266,30 @@ void UOptimusAnimAttributeDataInterface::GetHLSL(FString& OutHLSL) const
 							
 							if (StructElement.Type->Type == EShaderFundamentalType::Struct && !StructElement.Type->bIsDynamicArray)
 							{
-								WalkStruct(StructElement.Type, NewPrefix);
+								WalkStruct(InDataInterfaceName, StructElement.Type, NewPrefix);
 							}
 							else
 							{
 								FString ElementType = StructElement.Type.ValueTypePtr->ToString();
 
 								// Add uniforms.
-								MemberHlsl.Add(FString::Printf(TEXT("%s DI_LOCAL(%s);\n"), *ElementType, *ElementHlslId));
+								MemberHlsl.Add(FString::Printf(TEXT("%s %s_%s);\n"),
+									*ElementType, 
+									*InDataInterfaceName,
+									*ElementHlslId));
 						
 								// Add function getters.
-								MemberHlsl.Add(FString::Printf(TEXT("DI_IMPL_READ(Read%s, %s, )\n{\n\treturn DI_LOCAL(%s);\n}\n"),
-									*ElementHlslId,
+								MemberHlsl.Add(FString::Printf(TEXT("%s Read%s_%s()\n{\n\treturn %s_%s;\n}\n"),
 									*ElementType,
+									*ElementHlslId,
+									*InDataInterfaceName,
+									*InDataInterfaceName,
 									*ElementHlslId));
 
-								StructMemberCopyHlsl.Add(FString::Printf(TEXT("%s = DI_LOCAL(%s);\n"), *ElementAccessor, *ElementHlslId));	
+								StructMemberCopyHlsl.Add(FString::Printf(TEXT("%s = %s_%s;\n"),
+									*ElementAccessor, 
+									*InDataInterfaceName,
+									*ElementHlslId));
 							}
 						}
 					}
@@ -292,7 +302,7 @@ void UOptimusAnimAttributeDataInterface::GetHLSL(FString& OutHLSL) const
 			
 				FStructParser Parser;
 				
-				Parser.WalkStruct(Attribute.DataType->ShaderValueType, {Attribute.HlslId});
+				Parser.WalkStruct(InDataInterfaceName, Attribute.DataType->ShaderValueType, {Attribute.HlslId});
 
 				// Top level Hlsl code for flattened list of struct members
 				for (const FString& MemberHlslString : Parser.MemberHlsl)
@@ -301,9 +311,12 @@ void UOptimusAnimAttributeDataInterface::GetHLSL(FString& OutHLSL) const
 				}
 				
 				// Add final user facing struct getters.
-				OutHLSL += FString::Printf(TEXT("DI_IMPL_READ(Read%s, %s, )\n"),
+				OutHLSL += FString::Printf(TEXT("%s Read%s_%s()\n{\n\treturn %s_%s;\n}\n"),
+					*TypeName,
 					*Attribute.HlslId,
-					*TypeName);
+					*InDataInterfaceName,
+					*InDataInterfaceName,
+					*Attribute.HlslId);
 				
 				OutHLSL += FString::Printf(TEXT("{\n"));
 
@@ -360,12 +373,12 @@ void UOptimusAnimAttributeDataInterface::GetStructDeclarations(TSet<FString>& Ou
 
 void UOptimusAnimAttributeDataInterface::GetShaderHash(FString& InOutKey) const
 {
-	FSHA1 HashState;
-	FString HLSL;
-	GetHLSL(HLSL);
-	
-	HashState.UpdateWithString(*HLSL, HLSL.Len());
-	HashState.Finalize().AppendString(InOutKey);
+// 	FSHA1 HashState;
+// 	FString HLSL;
+// 	GetHLSL(HLSL);
+// 	
+// 	HashState.UpdateWithString(*HLSL, HLSL.Len());
+// 	HashState.Finalize().AppendString(InOutKey);
 }
 
 UComputeDataProvider* UOptimusAnimAttributeDataInterface::CreateDataProvider(TObjectPtr<UObject> InBinding, uint64 InInputMask, uint64 InOutputMask) const
diff --git a/Engine/Plugins/Animation/DeformerGraph/Source/OptimusCore/Private/DataInterfaces/OptimusDataInterfaceAnimAttribute.h b/Engine/Plugins/Animation/DeformerGraph/Source/OptimusCore/Private/DataInterfaces/OptimusDataInterfaceAnimAttribute.h
--- a/Engine/Plugins/Animation/DeformerGraph/Source/OptimusCore/Private/DataInterfaces/OptimusDataInterfaceAnimAttribute.h
+++ b/Engine/Plugins/Animation/DeformerGraph/Source/OptimusCore/Private/DataInterfaces/OptimusDataInterfaceAnimAttribute.h
@@ -122,7 +122,7 @@ public:
 	TCHAR const* GetClassName() const override { return TEXT("AnimAttribute"); }
 	virtual void GetSupportedInputs(TArray<FShaderFunctionDefinition>& OutFunctions) const override;
 	void GetShaderParameters(TCHAR const* UID, FShaderParametersMetadataBuilder& InOutBuilder, FShaderParametersMetadataAllocations& InOutAllocations) const override;
-	void GetHLSL(FString& OutHLSL) const override;
+	void GetHLSL(FString& OutHLSL, FString const& InDataInterfaceName) const override;
 	void GetStructDeclarations(TSet<FString>& OutStructsSeen, TArray<FString>& OutStructs) const override;
 	void GetShaderHash(FString& InOutKey) const override;
 	UComputeDataProvider* CreateDataProvider(TObjectPtr<UObject> InBinding, uint64 InInputMask, uint64 InOutputMask) const override;
diff --git a/Engine/Plugins/Animation/DeformerGraph/Source/OptimusCore/Private/DataInterfaces/OptimusDataInterfaceCloth.cpp b/Engine/Plugins/Animation/DeformerGraph/Source/OptimusCore/Private/DataInterfaces/OptimusDataInterfaceCloth.cpp
--- a/Engine/Plugins/Animation/DeformerGraph/Source/OptimusCore/Private/DataInterfaces/OptimusDataInterfaceCloth.cpp
+++ b/Engine/Plugins/Animation/DeformerGraph/Source/OptimusCore/Private/DataInterfaces/OptimusDataInterfaceCloth.cpp
@@ -91,9 +91,16 @@ void UOptimusClothDataInterface::GetShaderHash(FString& InOutKey) const
 	GetShaderFileHash(TEXT("/Plugin/Optimus/Private/DataInterfaceCloth.ush"), EShaderPlatform::SP_PCD3D_SM5).AppendString(InOutKey);
 }
 
-void UOptimusClothDataInterface::GetHLSL(FString& OutHLSL) const
+void UOptimusClothDataInterface::GetHLSL(FString& OutHLSL, FString const& InDataInterfaceName) const
 {
-	OutHLSL += TEXT("#include \"/Plugin/Optimus/Private/DataInterfaceCloth.ush\"\n");
+	TMap<FString, FStringFormatArg> TemplateArgs =
+	{
+		{ TEXT("DataInterfaceName"), InDataInterfaceName },
+	};
+
+	FString TemplateFile;
+	LoadShaderSourceFile(TEXT("/Plugin/Optimus/Private/DataInterfaceCloth.ush"), EShaderPlatform::SP_PCD3D_SM5, &TemplateFile, nullptr);
+	OutHLSL += FString::Format(*TemplateFile, TemplateArgs);
 }
 
 UComputeDataProvider* UOptimusClothDataInterface::CreateDataProvider(TObjectPtr<UObject> InBinding, uint64 InInputMask, uint64 InOutputMask) const
diff --git a/Engine/Plugins/Animation/DeformerGraph/Source/OptimusCore/Private/DataInterfaces/OptimusDataInterfaceCloth.h b/Engine/Plugins/Animation/DeformerGraph/Source/OptimusCore/Private/DataInterfaces/OptimusDataInterfaceCloth.h
--- a/Engine/Plugins/Animation/DeformerGraph/Source/OptimusCore/Private/DataInterfaces/OptimusDataInterfaceCloth.h
+++ b/Engine/Plugins/Animation/DeformerGraph/Source/OptimusCore/Private/DataInterfaces/OptimusDataInterfaceCloth.h
@@ -28,7 +28,7 @@ public:
 	void GetShaderParameters(TCHAR const* UID, FShaderParametersMetadataBuilder& InOutBuilder, FShaderParametersMetadataAllocations& InOutAllocations) const override;
 	void GetPermutations(FComputeKernelPermutationVector& OutPermutationVector) const override;
 	void GetShaderHash(FString& InOutKey) const override;
-	void GetHLSL(FString& OutHLSL) const override;
+	void GetHLSL(FString& OutHLSL, FString const& InDataInterfaceName) const override;
 	UComputeDataProvider* CreateDataProvider(TObjectPtr<UObject> InBinding, uint64 InInputMask, uint64 InOutputMask) const override;
 	//~ End UComputeDataInterface Interface
 };
diff --git a/Engine/Plugins/Animation/DeformerGraph/Source/OptimusCore/Private/DataInterfaces/OptimusDataInterfaceConnectivity.cpp b/Engine/Plugins/Animation/DeformerGraph/Source/OptimusCore/Private/DataInterfaces/OptimusDataInterfaceConnectivity.cpp
--- a/Engine/Plugins/Animation/DeformerGraph/Source/OptimusCore/Private/DataInterfaces/OptimusDataInterfaceConnectivity.cpp
+++ b/Engine/Plugins/Animation/DeformerGraph/Source/OptimusCore/Private/DataInterfaces/OptimusDataInterfaceConnectivity.cpp
@@ -180,9 +180,16 @@ void UOptimusConnectivityDataInterface::GetShaderHash(FString& InOutKey) const
 	GetShaderFileHash(TEXT("/Plugin/Optimus/Private/DataInterfaceConnectivity.ush"), EShaderPlatform::SP_PCD3D_SM5).AppendString(InOutKey);
 }
 
-void UOptimusConnectivityDataInterface::GetHLSL(FString& OutHLSL) const
+void UOptimusConnectivityDataInterface::GetHLSL(FString& OutHLSL, FString const& InDataInterfaceName) const
 {
-	OutHLSL += TEXT("#include \"/Plugin/Optimus/Private/DataInterfaceConnectivity.ush\"\n");
+	TMap<FString, FStringFormatArg> TemplateArgs =
+	{
+		{ TEXT("DataInterfaceName"), InDataInterfaceName },
+	};
+
+	FString TemplateFile;
+	LoadShaderSourceFile(TEXT("/Plugin/Optimus/Private/DataInterfaceConnectivity.ush"), EShaderPlatform::SP_PCD3D_SM5, &TemplateFile, nullptr);
+	OutHLSL += FString::Format(*TemplateFile, TemplateArgs);
 }
 
 UComputeDataProvider* UOptimusConnectivityDataInterface::CreateDataProvider(TObjectPtr<UObject> InBinding, uint64 InInputMask, uint64 InOutputMask) const
diff --git a/Engine/Plugins/Animation/DeformerGraph/Source/OptimusCore/Private/DataInterfaces/OptimusDataInterfaceConnectivity.h b/Engine/Plugins/Animation/DeformerGraph/Source/OptimusCore/Private/DataInterfaces/OptimusDataInterfaceConnectivity.h
--- a/Engine/Plugins/Animation/DeformerGraph/Source/OptimusCore/Private/DataInterfaces/OptimusDataInterfaceConnectivity.h
+++ b/Engine/Plugins/Animation/DeformerGraph/Source/OptimusCore/Private/DataInterfaces/OptimusDataInterfaceConnectivity.h
@@ -30,7 +30,7 @@ public:
 	void GetSupportedInputs(TArray<FShaderFunctionDefinition>& OutFunctions) const override;
 	void GetShaderParameters(TCHAR const* UID, FShaderParametersMetadataBuilder& InOutBuilder, FShaderParametersMetadataAllocations& InOutAllocations) const override;
 	void GetShaderHash(FString& InOutKey) const override;
-	void GetHLSL(FString& OutHLSL) const override;
+	void GetHLSL(FString& OutHLSL, FString const& InDataInterfaceName) const override;
 	UComputeDataProvider* CreateDataProvider(TObjectPtr<UObject> InBinding, uint64 InInputMask, uint64 InOutputMask) const override;
 	//~ End UComputeDataInterface Interface
 
diff --git a/Engine/Plugins/Animation/DeformerGraph/Source/OptimusCore/Private/DataInterfaces/OptimusDataInterfaceDebugDraw.cpp b/Engine/Plugins/Animation/DeformerGraph/Source/OptimusCore/Private/DataInterfaces/OptimusDataInterfaceDebugDraw.cpp
--- a/Engine/Plugins/Animation/DeformerGraph/Source/OptimusCore/Private/DataInterfaces/OptimusDataInterfaceDebugDraw.cpp
+++ b/Engine/Plugins/Animation/DeformerGraph/Source/OptimusCore/Private/DataInterfaces/OptimusDataInterfaceDebugDraw.cpp
@@ -71,9 +71,16 @@ void UOptimusDebugDrawDataInterface::GetShaderHash(FString& InOutKey) const
 	GetShaderFileHash(TEXT("/Plugin/Optimus/Private/DataInterfaceDebugDraw.ush"), EShaderPlatform::SP_PCD3D_SM5).AppendString(InOutKey);
 }
 
-void UOptimusDebugDrawDataInterface::GetHLSL(FString& OutHLSL) const
+void UOptimusDebugDrawDataInterface::GetHLSL(FString& OutHLSL, FString const& InDataInterfaceName) const
 {
-	OutHLSL += TEXT("#include \"/Plugin/Optimus/Private/DataInterfaceDebugDraw.ush\"\n");
+	TMap<FString, FStringFormatArg> TemplateArgs =
+	{
+		{ TEXT("DataInterfaceName"), InDataInterfaceName },
+	};
+
+	FString TemplateFile;
+	LoadShaderSourceFile(TEXT("/Plugin/Optimus/Private/DataInterfaceDebugDraw.ush"), EShaderPlatform::SP_PCD3D_SM5, &TemplateFile, nullptr);
+	OutHLSL += FString::Format(*TemplateFile, TemplateArgs);
 }
 
 UComputeDataProvider* UOptimusDebugDrawDataInterface::CreateDataProvider(TObjectPtr<UObject> InBinding, uint64 InInputMask, uint64 InOutputMask) const
diff --git a/Engine/Plugins/Animation/DeformerGraph/Source/OptimusCore/Private/DataInterfaces/OptimusDataInterfaceDebugDraw.h b/Engine/Plugins/Animation/DeformerGraph/Source/OptimusCore/Private/DataInterfaces/OptimusDataInterfaceDebugDraw.h
--- a/Engine/Plugins/Animation/DeformerGraph/Source/OptimusCore/Private/DataInterfaces/OptimusDataInterfaceDebugDraw.h
+++ b/Engine/Plugins/Animation/DeformerGraph/Source/OptimusCore/Private/DataInterfaces/OptimusDataInterfaceDebugDraw.h
@@ -56,7 +56,7 @@ public:
 	void GetSupportedInputs(TArray<FShaderFunctionDefinition>& OutFunctions) const override;
 	void GetShaderParameters(TCHAR const* UID, FShaderParametersMetadataBuilder& InOutBuilder, FShaderParametersMetadataAllocations& InOutAllocations) const override;
 	void GetShaderHash(FString& InOutKey) const override;
-	void GetHLSL(FString& OutHLSL) const override;
+	void GetHLSL(FString& OutHLSL, FString const& InDataInterfaceName) const override;
 	UComputeDataProvider* CreateDataProvider(TObjectPtr<UObject> InBinding, uint64 InInputMask, uint64 InOutputMask) const override;
 	//~ End UComputeDataInterface Interface
 
diff --git a/Engine/Plugins/Animation/DeformerGraph/Source/OptimusCore/Private/DataInterfaces/OptimusDataInterfaceGraph.cpp b/Engine/Plugins/Animation/DeformerGraph/Source/OptimusCore/Private/DataInterfaces/OptimusDataInterfaceGraph.cpp
--- a/Engine/Plugins/Animation/DeformerGraph/Source/OptimusCore/Private/DataInterfaces/OptimusDataInterfaceGraph.cpp
+++ b/Engine/Plugins/Animation/DeformerGraph/Source/OptimusCore/Private/DataInterfaces/OptimusDataInterfaceGraph.cpp
@@ -72,19 +72,17 @@ void UOptimusGraphDataInterface::GetShaderParameters(TCHAR const* UID, FShaderPa
 	InOutBuilder.AddNestedStruct(UID, ShaderParameterMetadata);
 }
 
-void UOptimusGraphDataInterface::GetHLSL(FString& OutHLSL) const
+void UOptimusGraphDataInterface::GetHLSL(FString& OutHLSL, FString const& InDataInterfaceName) const
 {
-	// Need include for DI_LOCAL macro expansion.
-	OutHLSL += TEXT("#include \"/Plugin/ComputeFramework/Private/ComputeKernelCommon.ush\"\n");
 	// Add uniforms.
 	for (FOptimusGraphVariableDescription const& Variable : Variables)
 	{
-		OutHLSL += FString::Printf(TEXT("%s DI_LOCAL(%s);\n"), *Variable.ValueType->ToString(), *Variable.Name);
+		OutHLSL += FString::Printf(TEXT("%s %s_%s;\n"), *Variable.ValueType->ToString(), *InDataInterfaceName, *Variable.Name);
 	}
 	// Add function getters.
 	for (FOptimusGraphVariableDescription const& Variable : Variables)
 	{
-		OutHLSL += FString::Printf(TEXT("DI_IMPL_READ(Read%s, %s, )\n{\n\treturn DI_LOCAL(%s);\n}\n"), *Variable.Name, *Variable.ValueType->ToString(), *Variable.Name);
+		OutHLSL += FString::Printf(TEXT("%s Read%s_%s()\n{\n\treturn %s_%s;\n}\n"), *Variable.ValueType->ToString(), *Variable.Name, *InDataInterfaceName, *InDataInterfaceName, *Variable.Name);
 	}
 }
 
diff --git a/Engine/Plugins/Animation/DeformerGraph/Source/OptimusCore/Private/DataInterfaces/OptimusDataInterfaceGraph.h b/Engine/Plugins/Animation/DeformerGraph/Source/OptimusCore/Private/DataInterfaces/OptimusDataInterfaceGraph.h
--- a/Engine/Plugins/Animation/DeformerGraph/Source/OptimusCore/Private/DataInterfaces/OptimusDataInterfaceGraph.h
+++ b/Engine/Plugins/Animation/DeformerGraph/Source/OptimusCore/Private/DataInterfaces/OptimusDataInterfaceGraph.h
@@ -42,7 +42,7 @@ public:
 	TCHAR const* GetClassName() const override { return TEXT("Graph"); }
 	void GetSupportedInputs(TArray<FShaderFunctionDefinition>& OutFunctions) const override;
 	void GetShaderParameters(TCHAR const* UID, FShaderParametersMetadataBuilder& InOutBuilder, FShaderParametersMetadataAllocations& InOutAllocations) const override;
-	void GetHLSL(FString& OutHLSL) const override;
+	void GetHLSL(FString& OutHLSL, FString const& InDataInterfaceName) const override;
 	UComputeDataProvider* CreateDataProvider(TObjectPtr<UObject> InBinding, uint64 InInputMask, uint64 InOutputMask) const override;
 	//~ End UComputeDataInterface Interface
 
diff --git a/Engine/Plugins/Animation/DeformerGraph/Source/OptimusCore/Private/DataInterfaces/OptimusDataInterfaceJiggleSpring.cpp b/Engine/Plugins/Animation/DeformerGraph/Source/OptimusCore/Private/DataInterfaces/OptimusDataInterfaceJiggleSpring.cpp
--- a/Engine/Plugins/Animation/DeformerGraph/Source/OptimusCore/Private/DataInterfaces/OptimusDataInterfaceJiggleSpring.cpp
+++ b/Engine/Plugins/Animation/DeformerGraph/Source/OptimusCore/Private/DataInterfaces/OptimusDataInterfaceJiggleSpring.cpp
@@ -341,9 +341,16 @@ UOptimusJiggleSpringDataInterface::GetShaderHash(FString& InOutKey) const
 }
 
 void 
-UOptimusJiggleSpringDataInterface::GetHLSL(FString& OutHLSL) const
+UOptimusJiggleSpringDataInterface::GetHLSL(FString& OutHLSL, FString const& InDataInterfaceName) const
 {
-	OutHLSL += TEXT("#include \"/Plugin/Optimus/Private/DataInterfaceJiggleSpring.ush\"\n");
+	TMap<FString, FStringFormatArg> TemplateArgs =
+	{
+		{ TEXT("DataInterfaceName"), InDataInterfaceName },
+	};
+
+	FString TemplateFile;
+	LoadShaderSourceFile(TEXT("/Plugin/Optimus/Private/DataInterfaceJiggleSpring.ush"), EShaderPlatform::SP_PCD3D_SM5, &TemplateFile, nullptr);
+	OutHLSL += FString::Format(*TemplateFile, TemplateArgs);
 }
 
 template <class T>
diff --git a/Engine/Plugins/Animation/DeformerGraph/Source/OptimusCore/Private/DataInterfaces/OptimusDataInterfaceJiggleSpring.h b/Engine/Plugins/Animation/DeformerGraph/Source/OptimusCore/Private/DataInterfaces/OptimusDataInterfaceJiggleSpring.h
--- a/Engine/Plugins/Animation/DeformerGraph/Source/OptimusCore/Private/DataInterfaces/OptimusDataInterfaceJiggleSpring.h
+++ b/Engine/Plugins/Animation/DeformerGraph/Source/OptimusCore/Private/DataInterfaces/OptimusDataInterfaceJiggleSpring.h
@@ -106,7 +106,7 @@ public:
 	void GetShaderParameters(TCHAR const* UID, FShaderParametersMetadataBuilder& InOutBuilder, FShaderParametersMetadataAllocations& InOutAllocations) const override;
 	void GetPermutations(FComputeKernelPermutationVector& OutPermutationVector) const override;
 	void GetShaderHash(FString& InOutKey) const override;
-	void GetHLSL(FString& OutHLSL) const override;
+	void GetHLSL(FString& OutHLSL, FString const& InDataInterfaceName) const override;
 	UComputeDataProvider* CreateDataProvider(TObjectPtr<UObject> InBinding, uint64 InInputMask, uint64 InOutputMask) const override;
 	//~ End UComputeDataInterface Interface
 
diff --git a/Engine/Plugins/Animation/DeformerGraph/Source/OptimusCore/Private/DataInterfaces/OptimusDataInterfaceMorphTarget.cpp b/Engine/Plugins/Animation/DeformerGraph/Source/OptimusCore/Private/DataInterfaces/OptimusDataInterfaceMorphTarget.cpp
--- a/Engine/Plugins/Animation/DeformerGraph/Source/OptimusCore/Private/DataInterfaces/OptimusDataInterfaceMorphTarget.cpp
+++ b/Engine/Plugins/Animation/DeformerGraph/Source/OptimusCore/Private/DataInterfaces/OptimusDataInterfaceMorphTarget.cpp
@@ -70,9 +70,16 @@ void UOptimusMorphTargetDataInterface::GetShaderHash(FString& InOutKey) const
 	GetShaderFileHash(TEXT("/Plugin/Optimus/Private/DataInterfaceMorphTarget.ush"), EShaderPlatform::SP_PCD3D_SM5).AppendString(InOutKey);
 }
 
-void UOptimusMorphTargetDataInterface::GetHLSL(FString& OutHLSL) const
+void UOptimusMorphTargetDataInterface::GetHLSL(FString& OutHLSL, FString const& InDataInterfaceName) const
 {
-	OutHLSL += TEXT("#include \"/Plugin/Optimus/Private/DataInterfaceMorphTarget.ush\"\n");
+	TMap<FString, FStringFormatArg> TemplateArgs =
+	{
+		{ TEXT("DataInterfaceName"), InDataInterfaceName },
+	};
+
+	FString TemplateFile;
+	LoadShaderSourceFile(TEXT("/Plugin/Optimus/Private/DataInterfaceMorphTarget.ush"), EShaderPlatform::SP_PCD3D_SM5, &TemplateFile, nullptr);
+	OutHLSL += FString::Format(*TemplateFile, TemplateArgs);
 }
 
 UComputeDataProvider* UOptimusMorphTargetDataInterface::CreateDataProvider(TObjectPtr<UObject> InBinding, uint64 InInputMask, uint64 InOutputMask) const
diff --git a/Engine/Plugins/Animation/DeformerGraph/Source/OptimusCore/Private/DataInterfaces/OptimusDataInterfaceMorphTarget.h b/Engine/Plugins/Animation/DeformerGraph/Source/OptimusCore/Private/DataInterfaces/OptimusDataInterfaceMorphTarget.h
--- a/Engine/Plugins/Animation/DeformerGraph/Source/OptimusCore/Private/DataInterfaces/OptimusDataInterfaceMorphTarget.h
+++ b/Engine/Plugins/Animation/DeformerGraph/Source/OptimusCore/Private/DataInterfaces/OptimusDataInterfaceMorphTarget.h
@@ -29,7 +29,7 @@ public:
 	void GetShaderParameters(TCHAR const* UID, FShaderParametersMetadataBuilder& InOutBuilder, FShaderParametersMetadataAllocations& InOutAllocations) const override;
 	void GetPermutations(FComputeKernelPermutationVector& OutPermutationVector) const override;
 	void GetShaderHash(FString& InOutKey) const override;
-	void GetHLSL(FString& OutHLSL) const override;
+	void GetHLSL(FString& OutHLSL, FString const& InDataInterfaceName) const override;
 	UComputeDataProvider* CreateDataProvider(TObjectPtr<UObject> InBinding, uint64 InInputMask, uint64 InOutputMask) const override;
 	//~ End UComputeDataInterface Interface
 };
diff --git a/Engine/Plugins/Animation/DeformerGraph/Source/OptimusCore/Private/DataInterfaces/OptimusDataInterfaceRawBuffer.cpp b/Engine/Plugins/Animation/DeformerGraph/Source/OptimusCore/Private/DataInterfaces/OptimusDataInterfaceRawBuffer.cpp
--- a/Engine/Plugins/Animation/DeformerGraph/Source/OptimusCore/Private/DataInterfaces/OptimusDataInterfaceRawBuffer.cpp
+++ b/Engine/Plugins/Animation/DeformerGraph/Source/OptimusCore/Private/DataInterfaces/OptimusDataInterfaceRawBuffer.cpp
@@ -137,29 +137,25 @@ void UOptimusRawBufferDataInterface::GetShaderHash(FString& InOutKey) const
 	GetShaderFileHash(TEXT("/Plugin/Optimus/Private/DataInterfaceRawBuffer.ush"), EShaderPlatform::SP_PCD3D_SM5).AppendString(InOutKey);
 }
 
-void UOptimusRawBufferDataInterface::GetHLSL(FString& OutHLSL) const
+void UOptimusRawBufferDataInterface::GetHLSL(FString& OutHLSL, FString const& InDataInterfaceName) const
 {
-	TStringBuilder<512> StringBuilder;
-
 	const FString PublicType = ValueType->ToString();
 	const FString RawType = GetRawType();
 	const bool bUseRawType = !RawType.IsEmpty();
 
-	StringBuilder.Appendf(TEXT("#define PUBLIC_TYPE %s\n"), *PublicType);
-	StringBuilder.Appendf(TEXT("#define BUFFER_TYPE %s\n"), bUseRawType ? *RawType : *PublicType);
-	StringBuilder.Appendf(TEXT("#define BUFFER_TYPE_RAW %d\n"), bUseRawType ? 1 : 0);
-	StringBuilder.Appendf(TEXT("#define BUFFER_TYPE_SUPPORTS_ATOMIC %d\n"), SupportsAtomics() ? 1 : 0);
-	StringBuilder.Appendf(TEXT("#define BUFFER_SPLIT_READ_WRITE %d\n"), UseSplitBuffers() ? 1 : 0);
-
-	StringBuilder.Append(TEXT("#include \"/Plugin/Optimus/Private/DataInterfaceRawBuffer.ush\"\n"));
-	
-	StringBuilder.Append(TEXT("#undef PUBLIC_TYPE\n"));
-	StringBuilder.Append(TEXT("#undef BUFFER_TYPE\n"));
-	StringBuilder.Append(TEXT("#undef BUFFER_TYPE_RAW\n"));
-	StringBuilder.Append(TEXT("#undef BUFFER_TYPE_SUPPORTS_ATOMIC\n"));
-	StringBuilder.Append(TEXT("#undef BUFFER_SPLIT_READ_WRITE\n"));
-
-	OutHLSL += StringBuilder.ToString();
+	TMap<FString, FStringFormatArg> TemplateArgs =
+	{
+		{ TEXT("DataInterfaceName"), InDataInterfaceName },
+		{ TEXT("PublicType"), PublicType },
+		{ TEXT("BufferType"), bUseRawType ? RawType : PublicType },
+		{ TEXT("BufferTypeRaw"), bUseRawType },
+		{ TEXT("SupportAtomic"), SupportsAtomics() ? 1 : 0 },
+		{ TEXT("SplitReadWrite"), UseSplitBuffers() ? 1 : 0 },
+	};
+
+	FString TemplateFile;
+	LoadShaderSourceFile(TEXT("/Plugin/Optimus/Private/DataInterfaceRawBuffer.ush"), EShaderPlatform::SP_PCD3D_SM5, &TemplateFile, nullptr);
+	OutHLSL += FString::Format(*TemplateFile, TemplateArgs);
 }
 
 FString UOptimusTransientBufferDataInterface::GetDisplayName() const
diff --git a/Engine/Plugins/Animation/DeformerGraph/Source/OptimusCore/Private/DataInterfaces/OptimusDataInterfaceRawBuffer.h b/Engine/Plugins/Animation/DeformerGraph/Source/OptimusCore/Private/DataInterfaces/OptimusDataInterfaceRawBuffer.h
--- a/Engine/Plugins/Animation/DeformerGraph/Source/OptimusCore/Private/DataInterfaces/OptimusDataInterfaceRawBuffer.h
+++ b/Engine/Plugins/Animation/DeformerGraph/Source/OptimusCore/Private/DataInterfaces/OptimusDataInterfaceRawBuffer.h
@@ -47,7 +47,7 @@ public:
 	void GetSupportedInputs(TArray<FShaderFunctionDefinition>& OutFunctions) const override;
 	void GetSupportedOutputs(TArray<FShaderFunctionDefinition>& OutFunctions) const override;
 	void GetShaderHash(FString& InOutKey) const override;
-	void GetHLSL(FString& OutHLSL) const override;
+	void GetHLSL(FString& OutHLSL, FString const& InDataInterfaceName) const override;
 	//~ End UComputeDataInterface Interface
 
 	/** The value type we should be allocating elements for */
diff --git a/Engine/Plugins/Animation/DeformerGraph/Source/OptimusCore/Private/DataInterfaces/OptimusDataInterfaceScene.cpp b/Engine/Plugins/Animation/DeformerGraph/Source/OptimusCore/Private/DataInterfaces/OptimusDataInterfaceScene.cpp
--- a/Engine/Plugins/Animation/DeformerGraph/Source/OptimusCore/Private/DataInterfaces/OptimusDataInterfaceScene.cpp
+++ b/Engine/Plugins/Animation/DeformerGraph/Source/OptimusCore/Private/DataInterfaces/OptimusDataInterfaceScene.cpp
@@ -8,7 +8,7 @@
 #include "CoreGlobals.h"
 #include "Engine/World.h"
 #include "SceneInterface.h"
-#include "ShaderCore.h"
+#include "ShaderCompilerCore.h"
 #include "ShaderParameterMetadataBuilder.h"
 
 FString UOptimusSceneDataInterface::GetDisplayName() const
@@ -63,9 +63,16 @@ void UOptimusSceneDataInterface::GetShaderHash(FString& InOutKey) const
 	GetShaderFileHash(TEXT("/Plugin/Optimus/Private/DataInterfaceScene.ush"), EShaderPlatform::SP_PCD3D_SM5).AppendString(InOutKey);
 }
 
-void UOptimusSceneDataInterface::GetHLSL(FString& OutHLSL) const
+void UOptimusSceneDataInterface::GetHLSL(FString& OutHLSL, FString const& InDataInterfaceName) const
 {
-	OutHLSL += TEXT("#include \"/Plugin/Optimus/Private/DataInterfaceScene.ush\"\n");
+ 	TMap<FString, FStringFormatArg> TemplateArgs =
+	{
+		{ TEXT("DataInterfaceName"), InDataInterfaceName },
+	};
+
+	FString TemplateFile;
+	LoadShaderSourceFile(TEXT("/Plugin/Optimus/Private/DataInterfaceScene.ush"), EShaderPlatform::SP_PCD3D_SM5, &TemplateFile, nullptr);
+	OutHLSL += FString::Format(*TemplateFile, TemplateArgs);
 }
 
 UComputeDataProvider* UOptimusSceneDataInterface::CreateDataProvider(TObjectPtr<UObject> InBinding, uint64 InInputMask, uint64 InOutputMask) const
diff --git a/Engine/Plugins/Animation/DeformerGraph/Source/OptimusCore/Private/DataInterfaces/OptimusDataInterfaceScene.h b/Engine/Plugins/Animation/DeformerGraph/Source/OptimusCore/Private/DataInterfaces/OptimusDataInterfaceScene.h
--- a/Engine/Plugins/Animation/DeformerGraph/Source/OptimusCore/Private/DataInterfaces/OptimusDataInterfaceScene.h
+++ b/Engine/Plugins/Animation/DeformerGraph/Source/OptimusCore/Private/DataInterfaces/OptimusDataInterfaceScene.h
@@ -26,7 +26,7 @@ public:
 	void GetSupportedInputs(TArray<FShaderFunctionDefinition>& OutFunctions) const override;
 	void GetShaderParameters(TCHAR const* UID, FShaderParametersMetadataBuilder& InOutBuilder, FShaderParametersMetadataAllocations& InOutAllocations) const override;
 	void GetShaderHash(FString& InOutKey) const override;
-	void GetHLSL(FString& OutHLSL) const override;
+	void GetHLSL(FString& OutHLSL, FString const& InDataInterfaceName) const override;
 	UComputeDataProvider* CreateDataProvider(TObjectPtr<UObject> InBinding, uint64 InInputMask, uint64 InOutputMask) const override;
 	//~ End UComputeDataInterface Interface
 };
diff --git a/Engine/Plugins/Animation/DeformerGraph/Source/OptimusCore/Private/DataInterfaces/OptimusDataInterfaceSkeleton.cpp b/Engine/Plugins/Animation/DeformerGraph/Source/OptimusCore/Private/DataInterfaces/OptimusDataInterfaceSkeleton.cpp
--- a/Engine/Plugins/Animation/DeformerGraph/Source/OptimusCore/Private/DataInterfaces/OptimusDataInterfaceSkeleton.cpp
+++ b/Engine/Plugins/Animation/DeformerGraph/Source/OptimusCore/Private/DataInterfaces/OptimusDataInterfaceSkeleton.cpp
@@ -94,9 +94,16 @@ void UOptimusSkeletonDataInterface::GetShaderHash(FString& InOutKey) const
 	GetShaderFileHash(TEXT("/Plugin/Optimus/Private/DataInterfaceSkeleton.ush"), EShaderPlatform::SP_PCD3D_SM5).AppendString(InOutKey);
 }
 
-void UOptimusSkeletonDataInterface::GetHLSL(FString& OutHLSL) const
+void UOptimusSkeletonDataInterface::GetHLSL(FString& OutHLSL, FString const& InDataInterfaceName) const
 {
-	OutHLSL += TEXT("#include \"/Plugin/Optimus/Private/DataInterfaceSkeleton.ush\"\n");
+	TMap<FString, FStringFormatArg> TemplateArgs =
+	{
+		{ TEXT("DataInterfaceName"), InDataInterfaceName },
+	};
+
+	FString TemplateFile;
+	LoadShaderSourceFile(TEXT("/Plugin/Optimus/Private/DataInterfaceSkeleton.ush"), EShaderPlatform::SP_PCD3D_SM5, &TemplateFile, nullptr);
+	OutHLSL += FString::Format(*TemplateFile, TemplateArgs);
 }
 
 UComputeDataProvider* UOptimusSkeletonDataInterface::CreateDataProvider(TObjectPtr<UObject> InBinding, uint64 InInputMask, uint64 InOutputMask) const
diff --git a/Engine/Plugins/Animation/DeformerGraph/Source/OptimusCore/Private/DataInterfaces/OptimusDataInterfaceSkeleton.h b/Engine/Plugins/Animation/DeformerGraph/Source/OptimusCore/Private/DataInterfaces/OptimusDataInterfaceSkeleton.h
--- a/Engine/Plugins/Animation/DeformerGraph/Source/OptimusCore/Private/DataInterfaces/OptimusDataInterfaceSkeleton.h
+++ b/Engine/Plugins/Animation/DeformerGraph/Source/OptimusCore/Private/DataInterfaces/OptimusDataInterfaceSkeleton.h
@@ -29,7 +29,7 @@ public:
 	void GetShaderParameters(TCHAR const* UID, FShaderParametersMetadataBuilder& InOutBuilder, FShaderParametersMetadataAllocations& InOutAllocations) const override;
 	void GetPermutations(FComputeKernelPermutationVector& OutPermutationVector) const override;
 	void GetShaderHash(FString& InOutKey) const override;
-	void GetHLSL(FString& OutHLSL) const override;
+	void GetHLSL(FString& OutHLSL, FString const& InDataInterfaceName) const override;
 	UComputeDataProvider* CreateDataProvider(TObjectPtr<UObject> InBinding, uint64 InInputMask, uint64 InOutputMask) const override;
 	//~ End UComputeDataInterface Interface
 };
diff --git a/Engine/Plugins/Animation/DeformerGraph/Source/OptimusCore/Private/DataInterfaces/OptimusDataInterfaceSkinnedMesh.cpp b/Engine/Plugins/Animation/DeformerGraph/Source/OptimusCore/Private/DataInterfaces/OptimusDataInterfaceSkinnedMesh.cpp
--- a/Engine/Plugins/Animation/DeformerGraph/Source/OptimusCore/Private/DataInterfaces/OptimusDataInterfaceSkinnedMesh.cpp
+++ b/Engine/Plugins/Animation/DeformerGraph/Source/OptimusCore/Private/DataInterfaces/OptimusDataInterfaceSkinnedMesh.cpp
@@ -128,9 +128,16 @@ void UOptimusSkinnedMeshDataInterface::GetShaderHash(FString& InOutKey) const
 	GetShaderFileHash(TEXT("/Plugin/Optimus/Private/DataInterfaceSkinnedMesh.ush"), EShaderPlatform::SP_PCD3D_SM5).AppendString(InOutKey);
 }
 
-void UOptimusSkinnedMeshDataInterface::GetHLSL(FString& OutHLSL) const
+void UOptimusSkinnedMeshDataInterface::GetHLSL(FString& OutHLSL, FString const& InDataInterfaceName) const
 {
-	OutHLSL += TEXT("#include \"/Plugin/Optimus/Private/DataInterfaceSkinnedMesh.ush\"\n");
+	TMap<FString, FStringFormatArg> TemplateArgs =
+	{
+		{ TEXT("DataInterfaceName"), InDataInterfaceName },
+	};
+
+	FString TemplateFile;
+	LoadShaderSourceFile(TEXT("/Plugin/Optimus/Private/DataInterfaceSkinnedMesh.ush"), EShaderPlatform::SP_PCD3D_SM5, &TemplateFile, nullptr);
+	OutHLSL += FString::Format(*TemplateFile, TemplateArgs);
 }
 
 UComputeDataProvider* UOptimusSkinnedMeshDataInterface::CreateDataProvider(TObjectPtr<UObject> InBinding, uint64 InInputMask, uint64 InOutputMask) const
diff --git a/Engine/Plugins/Animation/DeformerGraph/Source/OptimusCore/Private/DataInterfaces/OptimusDataInterfaceSkinnedMesh.h b/Engine/Plugins/Animation/DeformerGraph/Source/OptimusCore/Private/DataInterfaces/OptimusDataInterfaceSkinnedMesh.h
--- a/Engine/Plugins/Animation/DeformerGraph/Source/OptimusCore/Private/DataInterfaces/OptimusDataInterfaceSkinnedMesh.h
+++ b/Engine/Plugins/Animation/DeformerGraph/Source/OptimusCore/Private/DataInterfaces/OptimusDataInterfaceSkinnedMesh.h
@@ -28,7 +28,7 @@ public:
 	void GetSupportedInputs(TArray<FShaderFunctionDefinition>& OutFunctions) const override;
 	void GetShaderParameters(TCHAR const* UID, FShaderParametersMetadataBuilder& InOutBuilder, FShaderParametersMetadataAllocations& InOutAllocations) const override;
 	void GetShaderHash(FString& InOutKey) const override;
-	void GetHLSL(FString& OutHLSL) const override;
+	void GetHLSL(FString& OutHLSL, FString const& InDataInterfaceName) const override;
 	UComputeDataProvider* CreateDataProvider(TObjectPtr<UObject> InBinding, uint64 InInputMask, uint64 InOutputMask) const override;
 	//~ End UComputeDataInterface Interface
 };
diff --git a/Engine/Plugins/Animation/DeformerGraph/Source/OptimusCore/Private/DataInterfaces/OptimusDataInterfaceSkinnedMeshExec.cpp b/Engine/Plugins/Animation/DeformerGraph/Source/OptimusCore/Private/DataInterfaces/OptimusDataInterfaceSkinnedMeshExec.cpp
--- a/Engine/Plugins/Animation/DeformerGraph/Source/OptimusCore/Private/DataInterfaces/OptimusDataInterfaceSkinnedMeshExec.cpp
+++ b/Engine/Plugins/Animation/DeformerGraph/Source/OptimusCore/Private/DataInterfaces/OptimusDataInterfaceSkinnedMeshExec.cpp
@@ -55,9 +55,16 @@ void UOptimusSkinnedMeshExecDataInterface::GetShaderHash(FString& InOutKey) cons
 	GetShaderFileHash(TEXT("/Plugin/Optimus/Private/DataInterfaceSkinnedMeshExec.ush"), EShaderPlatform::SP_PCD3D_SM5).AppendString(InOutKey);
 }
 
-void UOptimusSkinnedMeshExecDataInterface::GetHLSL(FString& OutHLSL) const
+void UOptimusSkinnedMeshExecDataInterface::GetHLSL(FString& OutHLSL, FString const& InDataInterfaceName) const
 {
-	OutHLSL += TEXT("#include \"/Plugin/Optimus/Private/DataInterfaceSkinnedMeshExec.ush\"\n");
+	TMap<FString, FStringFormatArg> TemplateArgs =
+	{
+		{ TEXT("DataInterfaceName"), InDataInterfaceName },
+	};
+
+	FString TemplateFile;
+	LoadShaderSourceFile(TEXT("/Plugin/Optimus/Private/DataInterfaceSkinnedMeshExec.ush"), EShaderPlatform::SP_PCD3D_SM5, &TemplateFile, nullptr);
+	OutHLSL += FString::Format(*TemplateFile, TemplateArgs);
 }
 
 UComputeDataProvider* UOptimusSkinnedMeshExecDataInterface::CreateDataProvider(TObjectPtr<UObject> InBinding, uint64 InInputMask, uint64 InOutputMask) const
diff --git a/Engine/Plugins/Animation/DeformerGraph/Source/OptimusCore/Private/DataInterfaces/OptimusDataInterfaceSkinnedMeshExec.h b/Engine/Plugins/Animation/DeformerGraph/Source/OptimusCore/Private/DataInterfaces/OptimusDataInterfaceSkinnedMeshExec.h
--- a/Engine/Plugins/Animation/DeformerGraph/Source/OptimusCore/Private/DataInterfaces/OptimusDataInterfaceSkinnedMeshExec.h
+++ b/Engine/Plugins/Animation/DeformerGraph/Source/OptimusCore/Private/DataInterfaces/OptimusDataInterfaceSkinnedMeshExec.h
@@ -43,7 +43,7 @@ public:
 	void GetSupportedInputs(TArray<FShaderFunctionDefinition>& OutFunctions) const override;
 	void GetShaderParameters(TCHAR const* UID, FShaderParametersMetadataBuilder& InOutBuilder, FShaderParametersMetadataAllocations& InOutAllocations) const override;
 	void GetShaderHash(FString& InOutKey) const override;
-	void GetHLSL(FString& OutHLSL) const override;
+	void GetHLSL(FString& OutHLSL, FString const& InDataInterfaceName) const override;
 	UComputeDataProvider* CreateDataProvider(TObjectPtr<UObject> InBinding, uint64 InInputMask, uint64 InOutputMask) const override;
 	//~ End UComputeDataInterface Interface
 
diff --git a/Engine/Plugins/Animation/DeformerGraph/Source/OptimusCore/Private/DataInterfaces/OptimusDataInterfaceSkinnedMeshWrite.cpp b/Engine/Plugins/Animation/DeformerGraph/Source/OptimusCore/Private/DataInterfaces/OptimusDataInterfaceSkinnedMeshWrite.cpp
--- a/Engine/Plugins/Animation/DeformerGraph/Source/OptimusCore/Private/DataInterfaces/OptimusDataInterfaceSkinnedMeshWrite.cpp
+++ b/Engine/Plugins/Animation/DeformerGraph/Source/OptimusCore/Private/DataInterfaces/OptimusDataInterfaceSkinnedMeshWrite.cpp
@@ -88,9 +88,16 @@ void UOptimusSkinnedMeshWriteDataInterface::GetShaderHash(FString& InOutKey) con
 	GetShaderFileHash(TEXT("/Plugin/Optimus/Private/DataInterfaceSkinnedMeshWrite.ush"), EShaderPlatform::SP_PCD3D_SM5).AppendString(InOutKey);
 }
 
-void UOptimusSkinnedMeshWriteDataInterface::GetHLSL(FString& OutHLSL) const
+void UOptimusSkinnedMeshWriteDataInterface::GetHLSL(FString& OutHLSL, FString const& InDataInterfaceName) const
 {
-	OutHLSL += TEXT("#include \"/Plugin/Optimus/Private/DataInterfaceSkinnedMeshWrite.ush\"\n");
+	TMap<FString, FStringFormatArg> TemplateArgs =
+	{
+		{ TEXT("DataInterfaceName"), InDataInterfaceName },
+	};
+
+	FString TemplateFile;
+	LoadShaderSourceFile(TEXT("/Plugin/Optimus/Private/DataInterfaceSkinnedMeshWrite.ush"), EShaderPlatform::SP_PCD3D_SM5, &TemplateFile, nullptr);
+	OutHLSL += FString::Format(*TemplateFile, TemplateArgs);
 }
 
 UComputeDataProvider* UOptimusSkinnedMeshWriteDataInterface::CreateDataProvider(TObjectPtr<UObject> InBinding, uint64 InInputMask, uint64 InOutputMask) const
diff --git a/Engine/Plugins/Animation/DeformerGraph/Source/OptimusCore/Private/DataInterfaces/OptimusDataInterfaceSkinnedMeshWrite.h b/Engine/Plugins/Animation/DeformerGraph/Source/OptimusCore/Private/DataInterfaces/OptimusDataInterfaceSkinnedMeshWrite.h
--- a/Engine/Plugins/Animation/DeformerGraph/Source/OptimusCore/Private/DataInterfaces/OptimusDataInterfaceSkinnedMeshWrite.h
+++ b/Engine/Plugins/Animation/DeformerGraph/Source/OptimusCore/Private/DataInterfaces/OptimusDataInterfaceSkinnedMeshWrite.h
@@ -32,7 +32,7 @@ public:
 	void GetSupportedOutputs(TArray<FShaderFunctionDefinition>& OutFunctions) const override;
 	void GetShaderParameters(TCHAR const* UID, FShaderParametersMetadataBuilder& InOutBuilder, FShaderParametersMetadataAllocations& InOutAllocations) const override;
 	void GetShaderHash(FString& InOutKey) const override;
-	void GetHLSL(FString& OutHLSL) const override;
+	void GetHLSL(FString& OutHLSL, FString const& InDataInterfaceName) const override;
 	UComputeDataProvider* CreateDataProvider(TObjectPtr<UObject> InBinding, uint64 InInputMask, uint64 InOutputMask) const override;
 	//~ End UComputeDataInterface Interface
 };
diff --git a/Engine/Plugins/Animation/MLDeformer/MLDeformerFramework/Shaders/Private/MLDeformerGraphDataInterface.ush b/Engine/Plugins/Animation/MLDeformer/MLDeformerFramework/Shaders/Private/MLDeformerGraphDataInterface.ush
--- a/Engine/Plugins/Animation/MLDeformer/MLDeformerFramework/Shaders/Private/MLDeformerGraphDataInterface.ush
+++ b/Engine/Plugins/Animation/MLDeformer/MLDeformerFramework/Shaders/Private/MLDeformerGraphDataInterface.ush
@@ -1,22 +1,19 @@
 // Copyright Epic Games, Inc. All Rights Reserved.
 
-#include "/Plugin/ComputeFramework/Private/ComputeKernelCommon.ush"
+uint {DataInterfaceName}_NumVertices;
+uint {DataInterfaceName}_InputStreamStart;
+float {DataInterfaceName}_Weight;
+Buffer<float> {DataInterfaceName}_PositionDeltaBuffer;
+Buffer<uint> {DataInterfaceName}_VertexMapBuffer;
 
-uint DI_LOCAL(NumVertices);
-uint DI_LOCAL(InputStreamStart);
-float DI_LOCAL(Weight);
-
-Buffer<float> DI_LOCAL(PositionDeltaBuffer);
-Buffer<uint> DI_LOCAL(VertexMapBuffer);
-
-DI_IMPL_READ(ReadNumVertices, uint, )
+uint ReadNumVertices_{DataInterfaceName}()
 {
-	return DI_LOCAL(NumVertices);
+	return {DataInterfaceName}_NumVertices;
 }
 
-DI_IMPL_READ(ReadPositionDelta, float3, uint VertexIndex)
+float3 ReadPositionDelta_{DataInterfaceName}(uint VertexIndex)
 {
-	uint BufferIndex = DI_LOCAL(VertexMapBuffer)[VertexIndex + DI_LOCAL(InputStreamStart)] * 3;
-	float3 Delta = float3(DI_LOCAL(PositionDeltaBuffer)[BufferIndex], DI_LOCAL(PositionDeltaBuffer)[BufferIndex + 1], DI_LOCAL(PositionDeltaBuffer)[BufferIndex + 2]);
-	return Delta * DI_LOCAL(Weight);
+	uint BufferIndex = {DataInterfaceName}_VertexMapBuffer[VertexIndex + {DataInterfaceName}_InputStreamStart] * 3;
+	float3 Delta = float3({DataInterfaceName}_PositionDeltaBuffer[BufferIndex], {DataInterfaceName}_PositionDeltaBuffer[BufferIndex + 1], {DataInterfaceName}_PositionDeltaBuffer[BufferIndex + 2]);
+	return Delta * {DataInterfaceName}_Weight;
 }
diff --git a/Engine/Plugins/Animation/MLDeformer/MLDeformerFramework/Shaders/Private/MLDeformerGraphHeatMapDataInterface.ush b/Engine/Plugins/Animation/MLDeformer/MLDeformerFramework/Shaders/Private/MLDeformerGraphHeatMapDataInterface.ush
--- a/Engine/Plugins/Animation/MLDeformer/MLDeformerFramework/Shaders/Private/MLDeformerGraphHeatMapDataInterface.ush
+++ b/Engine/Plugins/Animation/MLDeformer/MLDeformerFramework/Shaders/Private/MLDeformerGraphHeatMapDataInterface.ush
@@ -1,43 +1,40 @@
 // Copyright Epic Games, Inc. All Rights Reserved.
 
-#include "/Plugin/ComputeFramework/Private/ComputeKernelCommon.ush"
-
-uint DI_LOCAL(NumVertices);
-uint DI_LOCAL(InputStreamStart);
-uint DI_LOCAL(HeatMapMode);
-float DI_LOCAL(HeatMapMax);
-float DI_LOCAL(GroundTruthLerp);
-uint DI_LOCAL(GroundTruthBufferSize);
-
-StructuredBuffer<float3> DI_LOCAL(PositionGroundTruthBuffer);
-Buffer<uint> DI_LOCAL(VertexMapBuffer);
-
-DI_IMPL_READ(ReadNumVertices, uint, )
+uint {DataInterfaceName}_NumVertices;
+uint {DataInterfaceName}_InputStreamStart;
+uint {DataInterfaceName}_HeatMapMode;
+float {DataInterfaceName}_HeatMapMax;
+float {DataInterfaceName}_GroundTruthLerp;
+uint {DataInterfaceName}_GroundTruthBufferSize;
+StructuredBuffer<float3> {DataInterfaceName}_PositionGroundTruthBuffer;
+Buffer<uint> {DataInterfaceName}_VertexMapBuffer;
+
+uint ReadNumVertices_{DataInterfaceName}()
 {
-	return DI_LOCAL(NumVertices);
+	return {DataInterfaceName}_NumVertices;
 }
 
-DI_IMPL_READ(ReadPositionGroundTruth, float3, uint VertexIndex)
+float3 ReadPositionGroundTruth_{DataInterfaceName}(uint VertexIndex)
 {
-	uint BufferIndex = DI_LOCAL(VertexMapBuffer)[VertexIndex + DI_LOCAL(InputStreamStart)];
-	if (BufferIndex >= DI_LOCAL(GroundTruthBufferSize))
+	uint BufferIndex = {DataInterfaceName}_VertexMapBuffer[VertexIndex + {DataInterfaceName}_InputStreamStart];
+	if (BufferIndex >= {DataInterfaceName}_GroundTruthBufferSize)
 	{
 		return float3(0, 0, 0);
 	}
-	return DI_LOCAL(PositionGroundTruthBuffer)[BufferIndex];
+	return {DataInterfaceName}_PositionGroundTruthBuffer[BufferIndex];
 }
 
-DI_IMPL_READ(ReadHeatMapMode, uint, )
+uint ReadHeatMapMode_{DataInterfaceName}()
 {
-	return DI_LOCAL(HeatMapMode);
+	return {DataInterfaceName}_HeatMapMode;
 }
 
-DI_IMPL_READ(ReadHeatMapMax, float, )
+float ReadHeatMapMax_{DataInterfaceName}()
 {
-	return DI_LOCAL(HeatMapMax);
+	return {DataInterfaceName}_HeatMapMax;
 }
 
-DI_IMPL_READ(ReadGroundTruthLerp, float, )
+float ReadGroundTruthLerp_{DataInterfaceName}()
 {
-	return DI_LOCAL(GroundTruthLerp);
+	return {DataInterfaceName}_GroundTruthLerp;
 }
diff --git a/Engine/Plugins/Animation/MLDeformer/MLDeformerFramework/Source/MLDeformerFramework/Private/MLDeformerGraphDataInterface.cpp b/Engine/Plugins/Animation/MLDeformer/MLDeformerFramework/Source/MLDeformerFramework/Private/MLDeformerGraphDataInterface.cpp
--- a/Engine/Plugins/Animation/MLDeformer/MLDeformerFramework/Source/MLDeformerFramework/Private/MLDeformerGraphDataInterface.cpp
+++ b/Engine/Plugins/Animation/MLDeformer/MLDeformerFramework/Source/MLDeformerFramework/Private/MLDeformerGraphDataInterface.cpp
@@ -54,9 +54,16 @@ void UMLDeformerGraphDataInterface::GetShaderParameters(TCHAR const* UID, FShade
 	InOutBuilder.AddNestedStruct<FMLDeformerGraphDataInterfaceParameters>(UID);
 }
 
-void UMLDeformerGraphDataInterface::GetHLSL(FString& OutHLSL) const
+void UMLDeformerGraphDataInterface::GetHLSL(FString& OutHLSL, FString const& InDataInterfaceName) const
 {
-	OutHLSL += TEXT("#include \"/Plugin/MLDeformerFramework/Private/MLDeformerGraphDataInterface.ush\"\n");
+	TMap<FString, FStringFormatArg> TemplateArgs =
+	{
+		{ TEXT("DataInterfaceName"), InDataInterfaceName },
+	};
+
+	FString TemplateFile;
+	LoadShaderSourceFile(TEXT("/Plugin/MLDeformerFramework/Private/MLDeformerGraphDataInterface.ush"), EShaderPlatform::SP_PCD3D_SM5, &TemplateFile, nullptr);
+	OutHLSL += FString::Format(*TemplateFile, TemplateArgs);
 }
 
 UComputeDataProvider* UMLDeformerGraphDataInterface::CreateDataProvider(TObjectPtr<UObject> InBinding, uint64 InInputMask, uint64 InOutputMask) const
diff --git a/Engine/Plugins/Animation/MLDeformer/MLDeformerFramework/Source/MLDeformerFramework/Private/MLDeformerGraphDebugDataInterface.cpp b/Engine/Plugins/Animation/MLDeformer/MLDeformerFramework/Source/MLDeformerFramework/Private/MLDeformerGraphDebugDataInterface.cpp
--- a/Engine/Plugins/Animation/MLDeformer/MLDeformerFramework/Source/MLDeformerFramework/Private/MLDeformerGraphDebugDataInterface.cpp
+++ b/Engine/Plugins/Animation/MLDeformer/MLDeformerFramework/Source/MLDeformerFramework/Private/MLDeformerGraphDebugDataInterface.cpp
@@ -71,9 +71,16 @@ void UMLDeformerGraphDebugDataInterface::GetShaderParameters(TCHAR const* UID, F
 	InOutBuilder.AddNestedStruct<FMLDeformerGraphDebugDataInterfaceParameters>(UID);
 }
 
-void UMLDeformerGraphDebugDataInterface::GetHLSL(FString& OutHLSL) const
+void UMLDeformerGraphDebugDataInterface::GetHLSL(FString& OutHLSL, FString const& InDataInterfaceName) const
 {
-	OutHLSL += TEXT("#include \"/Plugin/MLDeformerFramework/Private/MLDeformerGraphHeatMapDataInterface.ush\"\n");
+	TMap<FString, FStringFormatArg> TemplateArgs =
+	{
+		{ TEXT("DataInterfaceName"), InDataInterfaceName },
+	};
+
+	FString TemplateFile;
+	LoadShaderSourceFile(TEXT("/Plugin/MLDeformerFramework/Private/MLDeformerGraphHeatMapDataInterface.ush"), EShaderPlatform::SP_PCD3D_SM5, &TemplateFile, nullptr);
+	OutHLSL += FString::Format(*TemplateFile, TemplateArgs);
 }
 
 UComputeDataProvider* UMLDeformerGraphDebugDataInterface::CreateDataProvider(TObjectPtr<UObject> InBinding, uint64 InInputMask, uint64 InOutputMask) const
diff --git a/Engine/Plugins/Animation/MLDeformer/MLDeformerFramework/Source/MLDeformerFramework/Public/MLDeformerGraphDataInterface.h b/Engine/Plugins/Animation/MLDeformer/MLDeformerFramework/Source/MLDeformerFramework/Public/MLDeformerGraphDataInterface.h
--- a/Engine/Plugins/Animation/MLDeformer/MLDeformerFramework/Source/MLDeformerFramework/Public/MLDeformerGraphDataInterface.h
+++ b/Engine/Plugins/Animation/MLDeformer/MLDeformerFramework/Source/MLDeformerFramework/Public/MLDeformerGraphDataInterface.h
@@ -72,7 +72,7 @@ public:
 	TCHAR const* GetClassName() const override { return TEXT("MLDeformer"); }
 	virtual void GetSupportedInputs(TArray<FShaderFunctionDefinition>& OutFunctions) const override;
 	virtual void GetShaderParameters(TCHAR const* UID, FShaderParametersMetadataBuilder& InOutBuilder, FShaderParametersMetadataAllocations& InOutAllocations) const override;
-	virtual void GetHLSL(FString& OutHLSL) const override;
+	virtual void GetHLSL(FString& OutHLSL, FString const& InDataInterfaceName) const override;
 	virtual UComputeDataProvider* CreateDataProvider(TObjectPtr<UObject> InBinding, uint64 InInputMask, uint64 InOutputMask) const override;
 	// ~END UComputeDataInterface overrides.
 };
diff --git a/Engine/Plugins/Animation/MLDeformer/MLDeformerFramework/Source/MLDeformerFramework/Public/MLDeformerGraphDebugDataInterface.h b/Engine/Plugins/Animation/MLDeformer/MLDeformerFramework/Source/MLDeformerFramework/Public/MLDeformerGraphDebugDataInterface.h
--- a/Engine/Plugins/Animation/MLDeformer/MLDeformerFramework/Source/MLDeformerFramework/Public/MLDeformerGraphDebugDataInterface.h
+++ b/Engine/Plugins/Animation/MLDeformer/MLDeformerFramework/Source/MLDeformerFramework/Public/MLDeformerGraphDebugDataInterface.h
@@ -57,7 +57,7 @@ public:
 	TCHAR const* GetClassName() const override { return TEXT("MLDeformerDebug"); }
 	virtual void GetSupportedInputs(TArray<FShaderFunctionDefinition>& OutFunctions) const override;
 	virtual void GetShaderParameters(TCHAR const* UID, FShaderParametersMetadataBuilder& OutBuilder, FShaderParametersMetadataAllocations& InOutAllocations) const override;
-	virtual void GetHLSL(FString& OutHLSL) const override;
+	virtual void GetHLSL(FString& OutHLSL, FString const& InDataInterfaceName) const override;
 	virtual UComputeDataProvider* CreateDataProvider(TObjectPtr<UObject> InBinding, uint64 InInputMask, uint64 InOutputMask) const override;
 	// ~END UComputeDataInterface overrides.
 };
diff --git a/Engine/Plugins/Runtime/ComputeFramework/Shaders/Private/BufferAlias.ush b/Engine/Plugins/Runtime/ComputeFramework/Shaders/Private/BufferAlias.ush
new file mode 100644
--- /dev/null
+++ b/Engine/Plugins/Runtime/ComputeFramework/Shaders/Private/BufferAlias.ush
@@ -0,0 +1,51 @@
+// Copyright Epic Games, Inc. All Rights Reserved.
+
+// Implement ReadRawBuffer/WriteRawBuffer functions for plain old types that need it.
+// Due to restrictions on some platforms we can't use StructureBuffer with unaligned types. That includes all vector<T, 3> types.
+// Similarly user structure types that don't follow alignment restrictions might have to implement similar functions to read from raw buffers.
+
+#pragma once
+
+void ReadRawBuffer(in StructuredBuffer<float> InBuffer, in uint InBufferIndex, out float3 OutValue)
+{
+	for (uint Index = 0; Index < 3; ++Index)
+	{
+		OutValue[Index] = InBuffer[InBufferIndex * 3 + Index];
+	}
+}
+void ReadRawBuffer(in RWStructuredBuffer<float> InBuffer, in uint InBufferIndex, out float3 OutValue)
+{
+	for (uint Index = 0; Index < 3; ++Index)
+	{
+		OutValue[Index] = InBuffer[InBufferIndex * 3 + Index];
+	}
+}
+void WriteRawBuffer(in RWStructuredBuffer<float> InBuffer, in uint InBufferIndex, in float3 InValue)
+{
+	for (uint Index = 0; Index < 3; ++Index)
+	{
+		InBuffer[InBufferIndex * 3 + Index] = InValue[Index];
+	}
+}
+
+void ReadRawBuffer(in StructuredBuffer<int> InBuffer, in uint InBufferIndex, out int3 OutValue)
+{
+	for (uint Index = 0; Index < 3; ++Index)
+	{
+		OutValue[Index] = InBuffer[InBufferIndex * 3 + Index];
+	}
+}
+void ReadRawBuffer(in RWStructuredBuffer<int> InBuffer, in uint InBufferIndex, out int3 OutValue)
+{
+	for (uint Index = 0; Index < 3; ++Index)
+	{
+		OutValue[Index] = InBuffer[InBufferIndex * 3 + Index];
+	}
+}
+void WriteRawBuffer(in RWStructuredBuffer<int> InBuffer, in uint InBufferIndex, in int3 InValue)
+{
+	for (uint Index = 0; Index < 3; ++Index)
+	{
+		InBuffer[InBufferIndex * 3 + Index] = InValue[Index];
+	}
+}
diff --git a/Engine/Plugins/Runtime/ComputeFramework/Shaders/Private/ComputeKernelCommon.ush b/Engine/Plugins/Runtime/ComputeFramework/Shaders/Private/ComputeKernelCommon.ush
--- a/Engine/Plugins/Runtime/ComputeFramework/Shaders/Private/ComputeKernelCommon.ush
+++ b/Engine/Plugins/Runtime/ComputeFramework/Shaders/Private/ComputeKernelCommon.ush
@@ -2,8 +2,6 @@
 
 #pragma once
 
-#include "/Engine/Public/Platform.ush"
-
 // Macro to mark a kernel entry point.
 // Only one should exist per kernel.
 
@@ -26,59 +24,12 @@
 
 #define _DI_CAT_IMPL_(A, B) A##B
 #define _DI_CAT_(A, B) _DI_CAT_IMPL_(A, B)
-#define DI_LOCAL(Name) _DI_CAT_(DI_UID, Name)
+#define DI_UNIFORM_LOCAL(Name) _DI_CAT_(DI_UID, Name)
+#define DI_FUNCTION_LOCAL(Name) _DI_CAT_(Name, DI_UID)
 
 // Macros to expose functions implemented by a data provider. 
 // This makes a local unqiue function name and expands the function prototype.
 // We could just use DI_LOCAL, but these macro also gives the possiblity to extract functions by parsing the file.
 
-#define DI_IMPL_READ(Name, ReturnType, ...) ReturnType DI_LOCAL(Name)(__VA_ARGS__)
-#define DI_IMPL_WRITE(Name, ...) void DI_LOCAL(Name)(__VA_ARGS__)
-
-// Implement ReadRawBuffer/WriteRawBuffer functions for plain old types that need it.
-// Due to restrictions on some platforms we can't use StructureBuffer with unaligned types. That includes all vector<T, 3> types.
-// Similarly user structure types that don't follow alignment restrictions might have to implement similar functions to read from raw buffers.
-
-void ReadRawBuffer(in StructuredBuffer<float> InBuffer, in uint InBufferIndex, out float3 OutValue)
-{
-	for (uint Index = 0; Index < 3; ++Index)
-	{
-		OutValue[Index] = InBuffer[InBufferIndex * 3 + Index];
-	}
-}
-void ReadRawBuffer(in RWStructuredBuffer<float> InBuffer, in uint InBufferIndex, out float3 OutValue)
-{
-	for (uint Index = 0; Index < 3; ++Index)
-	{
-		OutValue[Index] = InBuffer[InBufferIndex * 3 + Index];
-	}
-}
-void WriteRawBuffer(in RWStructuredBuffer<float> InBuffer, in uint InBufferIndex, in float3 InValue)
-{
-	for (uint Index = 0; Index < 3; ++Index)
-	{
-		InBuffer[InBufferIndex * 3 + Index] = InValue[Index];
-	}
-}
-
-void ReadRawBuffer(in StructuredBuffer<int> InBuffer, in uint InBufferIndex, out int3 OutValue)
-{
-	for (uint Index = 0; Index < 3; ++Index)
-	{
-		OutValue[Index] = InBuffer[InBufferIndex * 3 + Index];
-	}
-}
-void ReadRawBuffer(in RWStructuredBuffer<int> InBuffer, in uint InBufferIndex, out int3 OutValue)
-{
-	for (uint Index = 0; Index < 3; ++Index)
-	{
-		OutValue[Index] = InBuffer[InBufferIndex * 3 + Index];
-	}
-}
-void WriteRawBuffer(in RWStructuredBuffer<int> InBuffer, in uint InBufferIndex, in int3 InValue)
-{
-	for (uint Index = 0; Index < 3; ++Index)
-	{
-		InBuffer[InBufferIndex * 3 + Index] = InValue[Index];
-	}
-}
+#define DI_IMPL_READ(Name, ReturnType, ...) ReturnType DI_FUNCTION_LOCAL(Name)(__VA_ARGS__)
+#define DI_IMPL_WRITE(Name, ...) void DI_FUNCTION_LOCAL(Name)(__VA_ARGS__)
diff --git a/Engine/Plugins/Runtime/ComputeFramework/Source/ComputeFramework/Private/ComputeGraph.cpp b/Engine/Plugins/Runtime/ComputeFramework/Source/ComputeFramework/Private/ComputeGraph.cpp
--- a/Engine/Plugins/Runtime/ComputeFramework/Source/ComputeFramework/Private/ComputeGraph.cpp
+++ b/Engine/Plugins/Runtime/ComputeFramework/Source/ComputeFramework/Private/ComputeGraph.cpp
@@ -394,7 +394,7 @@ namespace
 
 		StringBuilder.Append(TEXT(") { "));
 		StringBuilder.Append(bHasReturn ? TEXT("return ") : TEXT(""));
-		StringBuilder.Append(UID).Append(TEXT("_")).Append(*FnImpl.Name);
+		StringBuilder.Append(*FnImpl.Name).Append(TEXT("_")).Append(UID);
 		StringBuilder.Append(TEXT("("));
 
 		for (int32 ParameterIndex = bHasReturn ? 1 : 0; ParameterIndex < NumParams; ++ParameterIndex)
@@ -462,6 +462,9 @@ FString UComputeGraph::BuildKernelSource(
 	TSet<FString> StructsSeen;
 	TArray<FString> StructDeclarations;
 
+	// Add include for any platform specific types.
+	HLSL += TEXT("#include \"/Engine/Public/Platform.ush\"\n");
+
 	// Add virtual source includes from the additional sources.
 	for (TPair<FString, FString> const& AdditionalSource : InAdditionalSources)
 	{
@@ -495,9 +498,7 @@ FString UComputeGraph::BuildKernelSource(
 		{
 			// Add a unique prefix to generate unique names in the data interface shader code.
 			FString NamePrefix = GetUniqueDataInterfaceName(DataInterface, DataProviderIndex);
-			HLSL += FString::Printf(TEXT("#define DI_UID %s_\n"), *NamePrefix);
-			DataInterface->GetHLSL(HLSL);
-			HLSL += TEXT("#undef DI_UID\n");
+			DataInterface->GetHLSL(HLSL, NamePrefix);
 
 			DataInterface->GetStructDeclarations(StructsSeen, StructDeclarations);
 			
diff --git a/Engine/Plugins/Runtime/ComputeFramework/Source/ComputeFramework/Public/ComputeFramework/ComputeDataInterface.h b/Engine/Plugins/Runtime/ComputeFramework/Source/ComputeFramework/Public/ComputeFramework/ComputeDataInterface.h
--- a/Engine/Plugins/Runtime/ComputeFramework/Source/ComputeFramework/Public/ComputeFramework/ComputeDataInterface.h
+++ b/Engine/Plugins/Runtime/ComputeFramework/Source/ComputeFramework/Public/ComputeFramework/ComputeDataInterface.h
@@ -7,7 +7,6 @@
 
 struct FComputeKernelDefinitionSet;
 struct FComputeKernelPermutationVector;
-struct FShaderCompilerEnvironment;
 struct FShaderFunctionDefinition;
 struct FShaderParametersMetadataAllocations;
 class FShaderParametersMetadataBuilder;
@@ -43,7 +42,7 @@ public:
 	/** Gather any extra struct types that this data provider relies on. */
 	virtual void GetStructDeclarations(TSet<FString>& OutStructsSeen, TArray<FString>& OutStructs) const {}
 	/** Gather the shader code for this data provider. */
-	virtual void GetHLSL(FString& OutHLSL) const {}
+	virtual void GetHLSL(FString& OutHLSL, FString const& InDataInterfaceName) const {}
 	/** Get class of UObject required to instantiate a UComputeDataProvider from this interface. */
 	virtual UClass* GetBindingType() const { return nullptr; }
 	/** Instantiate an associated UComputeDataProvider. */
