Optimus: Allow plugins to register their data interfaces, specifically so they can register any custom types before they get used.

#rnx

diff --git a/Engine/Plugins/Animation/DeformerGraph/Source/OptimusCore/Private/OptimusCoreModule.cpp b/Engine/Plugins/Animation/DeformerGraph/Source/OptimusCore/Private/OptimusCoreModule.cpp
--- a/Engine/Plugins/Animation/DeformerGraph/Source/OptimusCore/Private/OptimusCoreModule.cpp
+++ b/Engine/Plugins/Animation/DeformerGraph/Source/OptimusCore/Private/OptimusCoreModule.cpp
@@ -33,6 +33,22 @@ void FOptimusCoreModule::ShutdownModule()
 	FOptimusDataTypeRegistry::UnregisterAllTypes();
 }
 
+
+bool FOptimusCoreModule::RegisterDataInterfaceClass(TSubclassOf<UOptimusComputeDataInterface> InDataInterfaceClass)
+{
+	if (InDataInterfaceClass)
+	{
+		UOptimusComputeDataInterface* DataInterface = Cast<UOptimusComputeDataInterface>(InDataInterfaceClass->GetDefaultObject());
+		if (ensure(DataInterface))
+		{
+			DataInterface->RegisterTypes();
+			return true;
+		}
+	}
+	return false;
+}
+
+
 IMPLEMENT_MODULE(FOptimusCoreModule, OptimusCore)
 
 DEFINE_LOG_CATEGORY(LogOptimusCore);
diff --git a/Engine/Plugins/Animation/DeformerGraph/Source/OptimusCore/Private/OptimusCoreModule.h b/Engine/Plugins/Animation/DeformerGraph/Source/OptimusCore/Private/OptimusCoreModule.h
--- a/Engine/Plugins/Animation/DeformerGraph/Source/OptimusCore/Private/OptimusCoreModule.h
+++ b/Engine/Plugins/Animation/DeformerGraph/Source/OptimusCore/Private/OptimusCoreModule.h
@@ -4,12 +4,16 @@
 
 #include "IOptimusCoreModule.h"
 
+
 class FOptimusCoreModule : public IOptimusCoreModule
 {
 public:
 	/** IModuleInterface implementation */
 	virtual void StartupModule() override;
 	virtual void ShutdownModule() override;
+
+	/** IOptimusCoreModule implementation */
+	bool RegisterDataInterfaceClass(TSubclassOf<UOptimusComputeDataInterface> InDataInterfaceClass) override;
 };
 
 DECLARE_LOG_CATEGORY_EXTERN(LogOptimusCore, Log, All);
diff --git a/Engine/Plugins/Animation/DeformerGraph/Source/OptimusCore/Public/IOptimusCoreModule.h b/Engine/Plugins/Animation/DeformerGraph/Source/OptimusCore/Public/IOptimusCoreModule.h
--- a/Engine/Plugins/Animation/DeformerGraph/Source/OptimusCore/Public/IOptimusCoreModule.h
+++ b/Engine/Plugins/Animation/DeformerGraph/Source/OptimusCore/Public/IOptimusCoreModule.h
@@ -6,7 +6,20 @@
 #include "Modules/ModuleInterface.h"
 #include "Templates/SubclassOf.h"
 
+
+class UOptimusComputeDataInterface;
+
+
 class IOptimusCoreModule : public IModuleInterface
 {
 public:
+	/** Register a data interface after the OptimusCore module has loaded */
+	template<typename T>
+	bool RegisterDataInterfaceClass()
+	{
+		return RegisterDataInterfaceClass(T::StaticClass());
+	}
+	
+protected:
+	virtual bool RegisterDataInterfaceClass(TSubclassOf<UOptimusComputeDataInterface> InDataInterfaceClass) = 0;
 };
