Notify with PinValueChanged after constant value property is edited.
Use PostEditChangeProperty() instead of PostEditChangeChainProperty() to catch undo/redo changes.
#rnx

diff --git a/Engine/Plugins/Animation/DeformerGraph/Source/OptimusCore/Private/Nodes/OptimusNode_ConstantValue.cpp b/Engine/Plugins/Animation/DeformerGraph/Source/OptimusCore/Private/Nodes/OptimusNode_ConstantValue.cpp
--- a/Engine/Plugins/Animation/DeformerGraph/Source/OptimusCore/Private/Nodes/OptimusNode_ConstantValue.cpp
+++ b/Engine/Plugins/Animation/DeformerGraph/Source/OptimusCore/Private/Nodes/OptimusNode_ConstantValue.cpp
@@ -82,15 +82,16 @@ void UOptimusNode_ConstantValue::PostLoad()
 
 #if WITH_EDITOR
 
-void UOptimusNode_ConstantValue::PostEditChangeChainProperty(FPropertyChangedChainEvent& PropertyChangedEvent)
+void UOptimusNode_ConstantValue::PostEditChangeProperty(FPropertyChangedEvent& PropertyChangedEvent)
 {
-	Super::PostEditChangeChainProperty(PropertyChangedEvent);
+	Super::PostEditChangeProperty(PropertyChangedEvent);
 
-	if (PropertyChangedEvent.ChangeType == EPropertyChangeType::ValueSet)
+	UOptimusNodeGraph* Graph = GetOwningGraph();
+	if (UOptimusNodePin* ValuePin = FindPinFromPath({ TEXT("Value") }))
 	{
-		UOptimusNodeGraph* Graph = GetOwningGraph();
-		Graph->GlobalNotify(EOptimusGlobalNotifyType::ConstantValueChanged, this);
+		Graph->Notify(EOptimusGraphNotifyType::PinValueChanged, ValuePin);
 	}
+	Graph->GlobalNotify(EOptimusGlobalNotifyType::ConstantValueChanged, this);
 }
 
 #endif // WITH_EDITOR
diff --git a/Engine/Plugins/Animation/DeformerGraph/Source/OptimusCore/Private/Nodes/OptimusNode_ConstantValue.h b/Engine/Plugins/Animation/DeformerGraph/Source/OptimusCore/Private/Nodes/OptimusNode_ConstantValue.h
--- a/Engine/Plugins/Animation/DeformerGraph/Source/OptimusCore/Private/Nodes/OptimusNode_ConstantValue.h
+++ b/Engine/Plugins/Animation/DeformerGraph/Source/OptimusCore/Private/Nodes/OptimusNode_ConstantValue.h
@@ -57,7 +57,7 @@ public:
 	void PostLoad() override;
 
 #if WITH_EDITOR
-	void PostEditChangeChainProperty(FPropertyChangedChainEvent& PropertyChangedEvent) override;
+	void PostEditChangeProperty(FPropertyChangedEvent& PropertyChangedEvent) override;
 #endif
 
 	// IOptimusValueProvider overrides 
