Optimus: fixed a typo

diff --git a/Engine/Plugins/Animation/DeformerGraph/Source/OptimusCore/Private/DataInterfaces/OptimusDataInterfaceAnimAttribute.cpp b/Engine/Plugins/Animation/DeformerGraph/Source/OptimusCore/Private/DataInterfaces/OptimusDataInterfaceAnimAttribute.cpp
--- a/Engine/Plugins/Animation/DeformerGraph/Source/OptimusCore/Private/DataInterfaces/OptimusDataInterfaceAnimAttribute.cpp
+++ b/Engine/Plugins/Animation/DeformerGraph/Source/OptimusCore/Private/DataInterfaces/OptimusDataInterfaceAnimAttribute.cpp
@@ -273,7 +273,7 @@ void UOptimusAnimAttributeDataInterface::GetHLSL(FString& OutHLSL, FString const
 								FString ElementType = StructElement.Type.ValueTypePtr->ToString();
 
 								// Add uniforms.
-								MemberHlsl.Add(FString::Printf(TEXT("%s %s_%s);\n"),
+								MemberHlsl.Add(FString::Printf(TEXT("%s %s_%s;\n"),
 									*ElementType, 
 									*InDataInterfaceName,
 									*ElementHlslId));
@@ -311,12 +311,10 @@ void UOptimusAnimAttributeDataInterface::GetHLSL(FString& OutHLSL, FString const
 				}
 				
 				// Add final user facing struct getters.
-				OutHLSL += FString::Printf(TEXT("%s Read%s_%s()\n{\n\treturn %s_%s;\n}\n"),
+				OutHLSL += FString::Printf(TEXT("%s Read%s_%s()\n"),
 					*TypeName,
 					*Attribute.HlslId,
-					*InDataInterfaceName,
-					*InDataInterfaceName,
-					*Attribute.HlslId);
+					*InDataInterfaceName);
 				
 				OutHLSL += FString::Printf(TEXT("{\n"));
 
