Optimus: ValueContainerClass is no longer standalone so that we don't need to manually GC it

diff --git a/Engine/Plugins/Animation/DeformerGraph/Source/OptimusCore/Private/OptimusValueContainer.cpp b/Engine/Plugins/Animation/DeformerGraph/Source/OptimusCore/Private/OptimusValueContainer.cpp
--- a/Engine/Plugins/Animation/DeformerGraph/Source/OptimusCore/Private/OptimusValueContainer.cpp
+++ b/Engine/Plugins/Animation/DeformerGraph/Source/OptimusCore/Private/OptimusValueContainer.cpp
@@ -25,7 +25,7 @@ UClass* UOptimusValueContainerGeneratorClass::GetClassForType(UPackage* InPackag
 	{
 		UClass *ParentClass = UOptimusValueContainer::StaticClass();
 		// Construct a value node class for this data type
-		TypeClass = NewObject<UOptimusValueContainerGeneratorClass>(InPackage, *ClassName, RF_Standalone);
+		TypeClass = NewObject<UOptimusValueContainerGeneratorClass>(InPackage, *ClassName);
 		TypeClass->SetSuperStruct(ParentClass);
 		TypeClass->PropertyLink = ParentClass->PropertyLink;
 
