Optimus: Fix a typo in anim attribute data interface for hlsl generation

diff --git a/Engine/Plugins/Animation/DeformerGraph/Source/OptimusCore/Private/DataInterfaces/OptimusDataInterfaceAnimAttribute.cpp b/Engine/Plugins/Animation/DeformerGraph/Source/OptimusCore/Private/DataInterfaces/OptimusDataInterfaceAnimAttribute.cpp
--- a/Engine/Plugins/Animation/DeformerGraph/Source/OptimusCore/Private/DataInterfaces/OptimusDataInterfaceAnimAttribute.cpp
+++ b/Engine/Plugins/Animation/DeformerGraph/Source/OptimusCore/Private/DataInterfaces/OptimusDataInterfaceAnimAttribute.cpp
@@ -232,7 +232,7 @@ void UOptimusAnimAttributeDataInterface::GetHLSL(FString& OutHLSL, FString const
 			if (!bIsStruct)
 			{
 				// Add uniforms.
-				OutHLSL += FString::Printf(TEXT("%s %s_%s);\n"), 
+				OutHLSL += FString::Printf(TEXT("%s %s_%s;\n"), 
 					*TypeName, 
 					*InDataInterfaceName, 
 					*Attribute.HlslId);
