Optimuis: when pasting nodes from a source asset to a target asset, make sure that we replace all references to objects in the source asset.

diff --git a/Engine/Plugins/Animation/DeformerGraph/Source/OptimusCore/Private/Actions/OptimusNodeGraphActions.cpp b/Engine/Plugins/Animation/DeformerGraph/Source/OptimusCore/Private/Actions/OptimusNodeGraphActions.cpp
--- a/Engine/Plugins/Animation/DeformerGraph/Source/OptimusCore/Private/Actions/OptimusNodeGraphActions.cpp
+++ b/Engine/Plugins/Animation/DeformerGraph/Source/OptimusCore/Private/Actions/OptimusNodeGraphActions.cpp
@@ -5,12 +5,14 @@
 #include "IOptimusPathResolver.h"
 #include "Nodes/OptimusNode_ComputeKernelFunction.h"
 #include "Nodes/OptimusNode_CustomComputeKernel.h"
+#include "Nodes/OptimusNode_ConstantValue.h"
 #include "OptimusHelpers.h"
 #include "OptimusNode.h"
 #include "OptimusNodeGraph.h"
 #include "OptimusNodeLink.h"
 #include "OptimusNodePin.h"
 
+
 #include "Serialization/MemoryReader.h"
 #include "Serialization/MemoryWriter.h"
 #include "Serialization/ObjectAndNameAsStringProxyArchive.h"
@@ -364,16 +366,22 @@ bool FOptimusNodeGraphAction_DuplicateNode::Do(
 		return ConfigureNodeFunc(InNode);
 	};
 
+	UOptimusNodeGraph* Graph = InRoot->ResolveGraphPath(GraphPath);
+	if (!Graph)
+	{
+		return false;
+	}
+
 	UClass* NodeClass = Optimus::FindObjectInPackageOrGlobal<UClass>(NodeClassPath);
 	if (!NodeClass)
 	{
 		return false;
 	}
-	
-	UOptimusNodeGraph* Graph = InRoot->ResolveGraphPath(GraphPath);
-	if (!Graph)
+
+	if (UOptimusNode_ConstantValueGeneratorClass* GeneratorClass = Cast<UOptimusNode_ConstantValueGeneratorClass>(NodeClass))
 	{
-		return false;
+		// Make sure a node class from the current package is used for constant nodes
+		NodeClass = UOptimusNode_ConstantValueGeneratorClass::GetClassForType(Graph->GetPackage(),GeneratorClass->DataType);
 	}
 
 	UOptimusNode* Node = Graph->CreateNodeDirect(NodeClass, NodeName, BootstrapNodeFunc);
diff --git a/Engine/Plugins/Animation/DeformerGraph/Source/OptimusCore/Private/Nodes/OptimusNode_AnimAttributeDataInterface.cpp b/Engine/Plugins/Animation/DeformerGraph/Source/OptimusCore/Private/Nodes/OptimusNode_AnimAttributeDataInterface.cpp
--- a/Engine/Plugins/Animation/DeformerGraph/Source/OptimusCore/Private/Nodes/OptimusNode_AnimAttributeDataInterface.cpp
+++ b/Engine/Plugins/Animation/DeformerGraph/Source/OptimusCore/Private/Nodes/OptimusNode_AnimAttributeDataInterface.cpp
@@ -104,6 +104,16 @@ void UOptimusNode_AnimAttributeDataInterface::OnDataTypeChanged(FName InTypeName
 	Interface->OnDataTypeChanged(InTypeName);
 }
 
+void UOptimusNode_AnimAttributeDataInterface::PostDuplicate(EDuplicateMode::Type DuplicateMode)
+{
+	Super::PostDuplicate(DuplicateMode);
+
+	// Recreate the value containers in case the node was pasted from a different asset
+	// Otherwise the value container can still reference the value container generator class in the asset it was copied from
+	RecreateValueContainers();
+}
+
+
 void UOptimusNode_AnimAttributeDataInterface::UpdatePinTypes()
 {
 	UOptimusAnimAttributeDataInterface* Interface = Cast<UOptimusAnimAttributeDataInterface>(DataInterfaceData);
diff --git a/Engine/Plugins/Animation/DeformerGraph/Source/OptimusCore/Private/Nodes/OptimusNode_AnimAttributeDataInterface.h b/Engine/Plugins/Animation/DeformerGraph/Source/OptimusCore/Private/Nodes/OptimusNode_AnimAttributeDataInterface.h
--- a/Engine/Plugins/Animation/DeformerGraph/Source/OptimusCore/Private/Nodes/OptimusNode_AnimAttributeDataInterface.h
+++ b/Engine/Plugins/Animation/DeformerGraph/Source/OptimusCore/Private/Nodes/OptimusNode_AnimAttributeDataInterface.h
@@ -25,6 +25,10 @@ public:
 	void RecreateValueContainers();
 
 	virtual void OnDataTypeChanged(FName InTypeName) override;
+
+protected:
+	// UObject overrides
+	virtual void PostDuplicate(EDuplicateMode::Type DuplicateMode) override;
 	
 private:
 	void UpdatePinTypes();
diff --git a/Engine/Plugins/Animation/DeformerGraph/Source/OptimusCore/Private/Nodes/OptimusNode_ComponentSource.cpp b/Engine/Plugins/Animation/DeformerGraph/Source/OptimusCore/Private/Nodes/OptimusNode_ComponentSource.cpp
--- a/Engine/Plugins/Animation/DeformerGraph/Source/OptimusCore/Private/Nodes/OptimusNode_ComponentSource.cpp
+++ b/Engine/Plugins/Animation/DeformerGraph/Source/OptimusCore/Private/Nodes/OptimusNode_ComponentSource.cpp
@@ -5,8 +5,41 @@
 #include "OptimusNodePin.h"
 
 #include "OptimusDataTypeRegistry.h"
+#include "OptimusDeformer.h"
 
 
+void UOptimusNode_ComponentSource::PostDuplicate(EDuplicateMode::Type DuplicateMode)
+{
+	Super::PostDuplicate(DuplicateMode);
+
+	if (!Binding)
+	{
+		return;
+	}
+	
+	if (!GetOwningGraph())
+	{
+		return;
+	}
+	
+	const UOptimusDeformer* OldBindingOwner = Binding->GetOwningDeformer();
+	const UOptimusDeformer* NewBindingOwner = Cast<UOptimusDeformer>(GetOwningGraph()->GetCollectionRoot());
+
+	if (!NewBindingOwner)
+	{
+		return;
+	}
+	
+	// No action needed if we are copying/pasting within the same deformer asset 
+	if (OldBindingOwner == NewBindingOwner)
+	{
+		return;
+	}
+
+	// Refresh the binding so that we don't hold a reference to a binding in another deformer asset
+	Binding = NewBindingOwner->ResolveComponentBinding(Binding->GetFName());
+}
+
 void UOptimusNode_ComponentSource::SetComponentSourceBinding(
 	UOptimusComponentSourceBinding* InBinding
 	)
diff --git a/Engine/Plugins/Animation/DeformerGraph/Source/OptimusCore/Private/Nodes/OptimusNode_ComponentSource.h b/Engine/Plugins/Animation/DeformerGraph/Source/OptimusCore/Private/Nodes/OptimusNode_ComponentSource.h
--- a/Engine/Plugins/Animation/DeformerGraph/Source/OptimusCore/Private/Nodes/OptimusNode_ComponentSource.h
+++ b/Engine/Plugins/Animation/DeformerGraph/Source/OptimusCore/Private/Nodes/OptimusNode_ComponentSource.h
@@ -40,6 +40,9 @@ protected:
 	 *  for backcomp.
 	 */
 	UOptimusNodePin* GetComponentPin() const;
+
+	// UObject overrides
+	virtual void PostDuplicate(EDuplicateMode::Type DuplicateMode) override;
 	
 private:
 	UPROPERTY()
diff --git a/Engine/Plugins/Animation/DeformerGraph/Source/OptimusCore/Private/Nodes/OptimusNode_GetVariable.cpp b/Engine/Plugins/Animation/DeformerGraph/Source/OptimusCore/Private/Nodes/OptimusNode_GetVariable.cpp
--- a/Engine/Plugins/Animation/DeformerGraph/Source/OptimusCore/Private/Nodes/OptimusNode_GetVariable.cpp
+++ b/Engine/Plugins/Animation/DeformerGraph/Source/OptimusCore/Private/Nodes/OptimusNode_GetVariable.cpp
@@ -3,9 +3,44 @@
 #include "Nodes/OptimusNode_GetVariable.h"
 
 #include "OptimusCoreModule.h"
+#include "OptimusDeformer.h"
 #include "OptimusNodePin.h"
 #include "OptimusVariableDescription.h"
 
+#define LOCTEXT_NAMESPACE "OptimusGetVariable"
+
+void UOptimusNode_GetVariable::PostDuplicate(EDuplicateMode::Type DuplicateMode)
+{
+	Super::PostDuplicate(DuplicateMode);
+
+	if (!VariableDesc.IsValid())
+	{
+		return;
+	}
+	
+	if (!GetOwningGraph())
+	{
+		return;
+	}
+	
+	const UOptimusDeformer* OldDescOwner = VariableDesc->GetOwningDeformer();
+	const UOptimusDeformer* NewDescOwner = Cast<UOptimusDeformer>(GetOwningGraph()->GetCollectionRoot());
+
+	if (!NewDescOwner)
+	{
+		return;
+	}
+	
+	// No action needed if we are copying/pasting within the same deformer asset 
+	if (OldDescOwner == NewDescOwner)
+	{
+		return;
+	}
+
+	// Refresh the ResourceDesc so that we don't hold a reference to a VariableDesc in another deformer asset
+	VariableDesc = NewDescOwner->ResolveVariable(VariableDesc->GetFName());
+}
+
 void UOptimusNode_GetVariable::SetVariableDescription(UOptimusVariableDescription* InVariableDesc)
 {
 	if (!ensure(InVariableDesc))
@@ -30,6 +65,17 @@ UOptimusVariableDescription* UOptimusNode_GetVariable::GetVariableDescription()
 }
 
 
+TOptional<FText> UOptimusNode_GetVariable::ValidateForCompile() const
+{
+	const UOptimusVariableDescription* VariableDescription = GetVariableDescription();
+	if (!VariableDescription)
+	{
+		return LOCTEXT("NoDescriptor", "No variable descriptor set on this node");
+	}
+
+	return {};
+}
+
 FString UOptimusNode_GetVariable::GetValueName() const
 {
 	if (const UOptimusVariableDescription* Var = VariableDesc.Get())
@@ -80,3 +126,5 @@ void UOptimusNode_GetVariable::ConstructNode()
 			Var->DataType);
 	}
 }
+
+#undef LOCTEXT_NAMESPACE
\ No newline at end of file
diff --git a/Engine/Plugins/Animation/DeformerGraph/Source/OptimusCore/Private/Nodes/OptimusNode_GetVariable.h b/Engine/Plugins/Animation/DeformerGraph/Source/OptimusCore/Private/Nodes/OptimusNode_GetVariable.h
--- a/Engine/Plugins/Animation/DeformerGraph/Source/OptimusCore/Private/Nodes/OptimusNode_GetVariable.h
+++ b/Engine/Plugins/Animation/DeformerGraph/Source/OptimusCore/Private/Nodes/OptimusNode_GetVariable.h
@@ -29,6 +29,8 @@ public:
 		return CategoryName::Variables;
 	}
 
+	TOptional<FText> ValidateForCompile() const override;
+	
 	// IOptimusValueProvider overrides 
 	FString GetValueName() const override;
 	FOptimusDataTypeRef GetValueType() const override;
@@ -37,6 +39,9 @@ public:
 protected:
 	void ConstructNode() override;
 
+	// UObject overrides
+	virtual void PostDuplicate(EDuplicateMode::Type DuplicateMode) override;
+	
 private:
 	UPROPERTY()
 	TWeakObjectPtr<UOptimusVariableDescription> VariableDesc;
diff --git a/Engine/Plugins/Animation/DeformerGraph/Source/OptimusCore/Private/Nodes/OptimusNode_ResourceAccessorBase.cpp b/Engine/Plugins/Animation/DeformerGraph/Source/OptimusCore/Private/Nodes/OptimusNode_ResourceAccessorBase.cpp
--- a/Engine/Plugins/Animation/DeformerGraph/Source/OptimusCore/Private/Nodes/OptimusNode_ResourceAccessorBase.cpp
+++ b/Engine/Plugins/Animation/DeformerGraph/Source/OptimusCore/Private/Nodes/OptimusNode_ResourceAccessorBase.cpp
@@ -3,12 +3,45 @@
 #include "OptimusNode_ResourceAccessorBase.h"
 
 #include "OptimusCoreModule.h"
+#include "OptimusDeformer.h"
 #include "OptimusResourceDescription.h"
 #include "DataInterfaces/OptimusDataInterfaceRawBuffer.h"
 
 #define LOCTEXT_NAMESPACE "OptimusResourceAccessorBase"
 
 
+void UOptimusNode_ResourceAccessorBase::PostDuplicate(EDuplicateMode::Type DuplicateMode)
+{
+	Super::PostDuplicate(DuplicateMode);
+
+	if (!ResourceDesc.IsValid())
+	{
+		return;
+	}
+	
+	if (!GetOwningGraph())
+	{
+		return;
+	}
+	
+	const UOptimusDeformer* OldDescOwner = ResourceDesc->GetOwningDeformer();
+	const UOptimusDeformer* NewDescOwner = Cast<UOptimusDeformer>(GetOwningGraph()->GetCollectionRoot());
+
+	if (!NewDescOwner)
+	{
+		return;
+	}
+	
+	// No action needed if we are copying/pasting within the same deformer asset 
+	if (OldDescOwner == NewDescOwner)
+	{
+		return;
+	}
+
+	// Refresh the ResourceDesc so that we don't hold a reference to a ResourceDesc in another deformer asset
+	ResourceDesc = NewDescOwner->ResolveResource(ResourceDesc->GetFName());
+}
+
 void UOptimusNode_ResourceAccessorBase::SetResourceDescription(UOptimusResourceDescription* InResourceDesc)
 {
 	if (!ensure(InResourceDesc))
diff --git a/Engine/Plugins/Animation/DeformerGraph/Source/OptimusCore/Private/Nodes/OptimusNode_ResourceAccessorBase.h b/Engine/Plugins/Animation/DeformerGraph/Source/OptimusCore/Private/Nodes/OptimusNode_ResourceAccessorBase.h
--- a/Engine/Plugins/Animation/DeformerGraph/Source/OptimusCore/Private/Nodes/OptimusNode_ResourceAccessorBase.h
+++ b/Engine/Plugins/Animation/DeformerGraph/Source/OptimusCore/Private/Nodes/OptimusNode_ResourceAccessorBase.h
@@ -45,6 +45,10 @@ public:
 	UOptimusComponentSourceBinding* GetComponentBinding() const override;
 
 	
+protected:
+	// UObject overrides
+	virtual void PostDuplicate(EDuplicateMode::Type DuplicateMode) override;
+	
 protected:
 	UPROPERTY()
 	TWeakObjectPtr<UOptimusResourceDescription> ResourceDesc;
diff --git a/Engine/Plugins/Animation/DeformerGraph/Source/OptimusCore/Private/OptimusNodeGraph.cpp b/Engine/Plugins/Animation/DeformerGraph/Source/OptimusCore/Private/OptimusNodeGraph.cpp
--- a/Engine/Plugins/Animation/DeformerGraph/Source/OptimusCore/Private/OptimusNodeGraph.cpp
+++ b/Engine/Plugins/Animation/DeformerGraph/Source/OptimusCore/Private/OptimusNodeGraph.cpp
@@ -1658,7 +1658,12 @@ IOptimusNodeGraphCollectionOwner* UOptimusNodeGraph::GetCollectionOwner() const
 
 IOptimusNodeGraphCollectionOwner* UOptimusNodeGraph::GetCollectionRoot() const
 {
-	return GetCollectionOwner()->GetCollectionRoot();
+	if (GetCollectionOwner())
+	{
+		return GetCollectionOwner()->GetCollectionRoot();
+	}
+	
+	return nullptr;
 }
 
 
