Optimus: Don't list both float/double in the type inspector, only the float type should be visible.

#rnx

diff --git a/Engine/Plugins/Animation/DeformerGraph/Source/OptimusCore/Private/OptimusDataTypeRegistry.cpp b/Engine/Plugins/Animation/DeformerGraph/Source/OptimusCore/Private/OptimusDataTypeRegistry.cpp
--- a/Engine/Plugins/Animation/DeformerGraph/Source/OptimusCore/Private/OptimusDataTypeRegistry.cpp
+++ b/Engine/Plugins/Animation/DeformerGraph/Source/OptimusCore/Private/OptimusDataTypeRegistry.cpp
@@ -339,12 +339,12 @@ void FOptimusDataTypeRegistry::RegisterBuiltinTypes()
 	    },
 		ConvertPropertyValuePOD<float, float>,
 		FName(TEXT("real")), {}, 
-	    EOptimusDataTypeUsageFlags::Resource | EOptimusDataTypeUsageFlags::Variable | EOptimusDataTypeUsageFlags::AnimAttributes);
+	    EOptimusDataTypeUsageFlags::Resource | EOptimusDataTypeUsageFlags::AnimAttributes);
 
 	// double -> float 
 	Registry.RegisterType(
 	    *FDoubleProperty::StaticClass(),
-	    FText::FromString(TEXT("Double")),
+	    FText::FromString(TEXT("Float")),
 	    FShaderValueType::Get(EShaderFundamentalType::Float),
 	    [](UStruct* InScope, FName InName) {
 		    FDoubleProperty* Property = new FDoubleProperty(InScope, InName, RF_Public);
@@ -359,7 +359,7 @@ void FOptimusDataTypeRegistry::RegisterBuiltinTypes()
 	    },
 		ConvertPropertyValuePOD<double, float>,
 		FName(TEXT("real")), {}, 
-	    EOptimusDataTypeUsageFlags::Resource | EOptimusDataTypeUsageFlags::Variable);
+	    EOptimusDataTypeUsageFlags::Variable);
 
 	// FVector2D -> float2
 	Registry.RegisterType(
