Optimus: Enable selection of blueprint components in component binding UI.
Added "Auto" option to binding choice which signals that we want to use automatic binding (by component class and tag).
Changed binding logic so that we store binding setup once at deformer init, and then use that plus binding logic to set the bindings on refresh.
This is different from before where we would restore binding setup after each refresh.

diff --git a/Engine/Plugins/Animation/DeformerGraph/Source/OptimusCore/Private/OptimusDeformerInstance.cpp b/Engine/Plugins/Animation/DeformerGraph/Source/OptimusCore/Private/OptimusDeformerInstance.cpp
--- a/Engine/Plugins/Animation/DeformerGraph/Source/OptimusCore/Private/OptimusDeformerInstance.cpp
+++ b/Engine/Plugins/Animation/DeformerGraph/Source/OptimusCore/Private/OptimusDeformerInstance.cpp
@@ -100,29 +100,81 @@ FOptimusDeformerInstanceExecInfo::FOptimusDeformerInstanceExecInfo()
 }
 
 
-void UOptimusDeformerInstanceSettings::RefreshComponentBindings(
-	UOptimusDeformer* InDeformer,
-	UMeshComponent* InMeshComponent
-	)
+bool FOptimusDeformerInstanceComponentBinding::GetSanitizedComponentName(FString& InOutName)
 {
-	
-	// Try to retain existing component binding as much as possible if called after a compilation or component rebuild
+	// Remove suffix for blueprint spawned components.
+	return InOutName.RemoveFromEnd(TEXT("_GEN_VARIABLE"));
+}
+
+FName FOptimusDeformerInstanceComponentBinding::GetSanitizedComponentName(FName InName)
+{
+	FString Name = InName.ToString();
+	if (GetSanitizedComponentName(Name))
+	{
+		return FName(Name);
+	}
+	// No change.
+	return InName;
+}
+
+FName FOptimusDeformerInstanceComponentBinding::GetSanitizedComponentName(UActorComponent const* InComponent)
+{
+	return InComponent ? GetSanitizedComponentName(InComponent->GetFName()) : FName();
+}
+
+TSoftObjectPtr<UActorComponent> FOptimusDeformerInstanceComponentBinding::GetActorComponent(AActor const* InActor, FString const& InName)
+{
+	if (InActor != nullptr && !InName.IsEmpty())
+	{
+		FString Path = InActor->GetPathName() + TEXT(".") + InName;
+		return TSoftObjectPtr<UActorComponent>(FSoftObjectPath(Path));
+	}
+	return {};
+}
+
+TSoftObjectPtr<UActorComponent> FOptimusDeformerInstanceComponentBinding::GetActorComponent(AActor const* InActor) const
+{
+	return GetActorComponent(InActor, ComponentName.ToString());
+}
+
+
+void UOptimusDeformerInstanceSettings::InitializeSettings(UOptimusDeformer* InDeformer, UMeshComponent* InPrimaryComponent)
+{
+	Deformer = InDeformer;
+
+	Bindings.SetNum(InDeformer->GetComponentBindings().Num());
+	for (int32 BindingIndex = 0; BindingIndex < Bindings.Num(); ++BindingIndex)
+	{
+		Bindings[BindingIndex].ProviderName = InDeformer->GetComponentBindings()[BindingIndex]->BindingName;
+		if (BindingIndex == 0)
+		{
+			Bindings[BindingIndex].ComponentName = FOptimusDeformerInstanceComponentBinding::GetSanitizedComponentName(InPrimaryComponent);
+		}
+	}
+}
+
+void UOptimusDeformerInstanceSettings::GetComponentBindings(
+	UOptimusDeformer* InDeformer, 
+	UMeshComponent* InPrimaryComponent, 
+	TArray<UActorComponent*>& OutComponents) const
+{
+	AActor const* Actor = InPrimaryComponent != nullptr ? InPrimaryComponent->GetOwner() : nullptr;
+
+	// Try to map onto the configured component bindings as much as possible.
 	TMap<FName, UActorComponent*> ExistingBindings;
 
-	for (FOptimusDeformerInstanceComponentBinding& Binding: Bindings)
+	for (FOptimusDeformerInstanceComponentBinding const& Binding : Bindings)
 	{
-		// It's likely that a component we referred to has been nixed, via AActor::RerunConstructionScripts, so
-		// we need to re-fetch the object from memory.
-		Binding.ActorComponent.ResetWeakPtr();
-		ExistingBindings.Add(Binding.ProviderName, Binding.ActorComponent.Get());
+		TSoftObjectPtr<UActorComponent> ActorComponent = Binding.GetActorComponent(Actor);
+		UActorComponent* Component = ActorComponent.Get();
+		ExistingBindings.Add(Binding.ProviderName, Component);
 	}
 	
-	AActor* Actor = InMeshComponent != nullptr ? InMeshComponent->GetOwner() : nullptr;
+	// Iterate component bindings and try to find a match.
 	TSet<UActorComponent*> ComponentsUsed;
-
 	const TArray<UOptimusComponentSourceBinding*>& ComponentBindings = InDeformer->GetComponentBindings();
-	Bindings.Reset(ComponentBindings.Num());
-	for (const UOptimusComponentSourceBinding* Binding: ComponentBindings)
+	OutComponents.Reset(ComponentBindings.Num());
+	for (const UOptimusComponentSourceBinding* Binding : ComponentBindings)
 	{
 		FName BindingName = Binding->BindingName;
 		UActorComponent* BoundComponent = nullptr;
@@ -130,7 +182,7 @@ void UOptimusDeformerInstanceSettings::RefreshComponentBindings(
 		// Primary binding always binds to the mesh component we're applied to.
 		if (Binding->IsPrimaryBinding())
 		{
-			BoundComponent = InMeshComponent;
+			BoundComponent = InPrimaryComponent;
 		}
 		else
 		{
@@ -197,33 +249,18 @@ void UOptimusDeformerInstanceSettings::RefreshComponentBindings(
 			}
 		}
 
-		Bindings.Add({BindingName, BoundComponent});
+		OutComponents.Add(BoundComponent);
 		ComponentsUsed.Add(BoundComponent);
 	}
 }
 
-TArray<UActorComponent*> UOptimusDeformerInstanceSettings::GetBoundComponents() const
-{
-	TArray<UActorComponent*> Result;
-	for (const FOptimusDeformerInstanceComponentBinding& Binding: Bindings)
-	{
-		UActorComponent* ActorComponent = Binding.ActorComponent.Get();
-		if (ActorComponent)
-		{
-			Result.Add(ActorComponent);
-		}
-	}
-	return Result;
-}
-
-
 AActor* UOptimusDeformerInstanceSettings::GetActor() const
 {
 	// We should be owned by an actor at some point.
 	return GetTypedOuter<AActor>();
 }
 
-UOptimusComponentSourceBinding* UOptimusDeformerInstanceSettings::GetComponentBindingByName(FName InBindingName) const
+UOptimusComponentSourceBinding const* UOptimusDeformerInstanceSettings::GetComponentBindingByName(FName InBindingName) const
 {
 	if (const UOptimusDeformer* DeformerResolved = Deformer.Get())
 	{
@@ -238,25 +275,6 @@ UOptimusComponentSourceBinding* UOptimusDeformerInstanceSettings::GetComponentBi
 	return nullptr;
 }
 
-#if WITH_EDITOR
-void UOptimusDeformerInstanceSettings::PostEditChangeProperty(FPropertyChangedEvent& PropertyChangedEvent)
-{
-	if (PropertyChangedEvent.GetPropertyName() == GET_MEMBER_NAME_CHECKED(FOptimusDeformerInstanceComponentBinding, ActorComponent))
-	{
-		// TODO: Update dependencies.
-	}
-}
-#endif
-
-void UOptimusDeformerInstanceSettings::InitializeSettings(
-	UOptimusDeformer* InDeformer,
-	UMeshComponent* InMeshComponent
-	)
-{
-	Deformer = InDeformer;
-	RefreshComponentBindings(InDeformer, InMeshComponent);
-}
-
 
 void UOptimusDeformerInstance::SetMeshComponent(UMeshComponent* InMeshComponent)
 { 
@@ -276,10 +294,11 @@ void UOptimusDeformerInstance::SetupFromDeformer(UOptimusDeformer* InDeformer)
 
 	// Update the component bindings before creating data providers. 
 	// The bindings are in the same order as the component bindings in the deformer.
+	TArray<UActorComponent*> BoundComponents;
 	UOptimusDeformerInstanceSettings* InstanceSettingsPtr = InstanceSettings.Get(); 
 	if (InstanceSettingsPtr)
 	{
-		InstanceSettingsPtr->RefreshComponentBindings(InDeformer, MeshComponent.Get());
+		InstanceSettingsPtr->GetComponentBindings(InDeformer, MeshComponent.Get(), BoundComponents);
 	}
 
 	// Create the persistent buffer pool
@@ -296,12 +315,13 @@ void UOptimusDeformerInstance::SetupFromDeformer(UOptimusDeformer* InDeformer)
 		Info.GraphType = ComputeGraphInfo.GraphType;
 		Info.ComputeGraph = ComputeGraphInfo.ComputeGraph;
 
-		if (InstanceSettingsPtr)
+		if (BoundComponents.Num())
 		{
-			// Fall back on everything being the given component.
-			for (int32 Index = 0; Index < InstanceSettingsPtr->Bindings.Num(); Index++)
+			UMeshComponent* ActorComponent = MeshComponent.Get();
+			AActor const* Actor = ActorComponent ? ActorComponent->GetOwner() : nullptr;
+			for (int32 Index = 0; Index < BoundComponents.Num(); Index++)
 			{
-				Info.ComputeGraphInstance.CreateDataProviders(Info.ComputeGraph, Index, InstanceSettingsPtr->Bindings[Index].ActorComponent.Get());
+				Info.ComputeGraphInstance.CreateDataProviders(Info.ComputeGraph, Index, BoundComponents[Index]);
 			}
 		}
 		else
diff --git a/Engine/Plugins/Animation/DeformerGraph/Source/OptimusCore/Public/OptimusDeformerInstance.h b/Engine/Plugins/Animation/DeformerGraph/Source/OptimusCore/Public/OptimusDeformerInstance.h
--- a/Engine/Plugins/Animation/DeformerGraph/Source/OptimusCore/Public/OptimusDeformerInstance.h
+++ b/Engine/Plugins/Animation/DeformerGraph/Source/OptimusCore/Public/OptimusDeformerInstance.h
@@ -4,6 +4,7 @@
 
 #include "Animation/MeshDeformerInstance.h"
 #include "ComputeFramework/ComputeGraphInstance.h"
+#include "Engine/EngineTypes.h"
 
 #include "OptimusDeformerInstance.generated.h"
 
@@ -78,12 +79,23 @@ USTRUCT(BlueprintType)
 struct FOptimusDeformerInstanceComponentBinding
 {
 	GENERATED_BODY()
-	
-	UPROPERTY(VisibleAnywhere, Category="Binding")
+
+	/** Binding name on deformer graph. */
+	UPROPERTY(VisibleAnywhere, Category="Binding", meta = (DisplayName = "Binding"))
 	FName ProviderName;
-	
-	UPROPERTY(EditAnywhere, Category="Binding")
-	TSoftObjectPtr<UActorComponent> ActorComponent;
+
+	/** Component name to bind. This should be sanitized before storage. */
+	UPROPERTY(EditAnywhere, Category = "Binding", meta = (DisplayName = "Component"))
+ 	FName ComponentName;
+
+	/** Get the component on an actor that matches the stored component name. */
+	TSoftObjectPtr<UActorComponent> GetActorComponent(AActor const* InActor) const;
+
+	/** Helpers to create ComponentName. */
+	OPTIMUSCORE_API static bool GetSanitizedComponentName(FString& InOutName);
+	OPTIMUSCORE_API static FName GetSanitizedComponentName(FName InName);
+	OPTIMUSCORE_API static FName GetSanitizedComponentName(UActorComponent const* InComponent);
+	OPTIMUSCORE_API static TSoftObjectPtr<UActorComponent> GetActorComponent(AActor const* InActor, FString const& InName);
 };
 
 
@@ -93,35 +105,29 @@ class OPTIMUSCORE_API UOptimusDeformerInstanceSettings :
 {
 	GENERATED_BODY()
 
-public:
-	/** */
-	UPROPERTY(EditAnywhere, Category="Deformer|Settings", EditFixedSize, meta=(NoResetToDefault, EditFixedOrder))
+	/** Stored weak pointer to a deformer. This is only required by the details customization for resolving binding class types. */
+	UPROPERTY()
+	TWeakObjectPtr<UOptimusDeformer> Deformer;
+
+	/** Array of binding descriptions. This is fixed and used by GetComponentBindings() to resolve final bindings for a given context. */
+	UPROPERTY(EditAnywhere, Category = "Deformer|Settings", EditFixedSize, meta = (EditFixedOrder))
 	TArray<FOptimusDeformerInstanceComponentBinding> Bindings;
 
-	void RefreshComponentBindings(
-		UOptimusDeformer* InDeformer,
-		UMeshComponent* InMeshComponent
-		);
+public:
+	/** Setup the object. This initializes the binding names and the primary binding component. */
+	void InitializeSettings(UOptimusDeformer* InDeformer, UMeshComponent* InPrimaryComponent);
 
-	TArray<UActorComponent*> GetBoundComponents() const;
-	AActor* GetActor() const;
-	UOptimusComponentSourceBinding* GetComponentBindingByName(FName InBindingName) const;
+	/** Get an array of recommended component bindings, based on the stored settings. */
+	void GetComponentBindings(UOptimusDeformer* InDeformer, UMeshComponent* InPrimaryComponent, TArray<UActorComponent*>& OutComponentBindings) const;
 
-#if WITH_EDITOR
-	// -- UObject
-	void PostEditChangeProperty(FPropertyChangedEvent& PropertyChangedEvent) override;
-#endif
-	
 protected:
-	friend class UOptimusDeformer;
+	friend class FOptimusDeformerInstanceComponentBindingCustomization;
 
-	void InitializeSettings(
-		UOptimusDeformer* InDeformer,
-		UMeshComponent* InMeshComponent);
+	/** Get the actor associated with this object. Used only by details customization. */
+	AActor* GetActor() const;
 
-private:
-	UPROPERTY()
-	TWeakObjectPtr<UOptimusDeformer> Deformer; 
+	/** Get a full component source binding object by binding name. Used only by details customization. */
+	UOptimusComponentSourceBinding const* GetComponentBindingByName(FName InBindingName) const;
 };
 
 
@@ -205,7 +211,7 @@ private:
 	UPROPERTY()
 	TWeakObjectPtr<UMeshComponent> MeshComponent;
 
-	/** The Mesh Component that owns this Mesh Deformer Instance. */
+	/** The settings for this Mesh Deformer Instance. */
 	UPROPERTY()
 	TWeakObjectPtr<UOptimusDeformerInstanceSettings> InstanceSettings;
 	
diff --git a/Engine/Plugins/Animation/DeformerGraph/Source/OptimusEditor/Private/OptimusDetailsCustomization.cpp b/Engine/Plugins/Animation/DeformerGraph/Source/OptimusEditor/Private/OptimusDetailsCustomization.cpp
--- a/Engine/Plugins/Animation/DeformerGraph/Source/OptimusEditor/Private/OptimusDetailsCustomization.cpp
+++ b/Engine/Plugins/Animation/DeformerGraph/Source/OptimusEditor/Private/OptimusDetailsCustomization.cpp
@@ -5,6 +5,7 @@
 #include "ComputeFramework/ShaderParamTypeDefinition.h"
 #include "DetailLayoutBuilder.h"
 #include "DetailWidgetRow.h"
+#include "Engine/SCS_Node.h"
 #include "IDetailChildrenBuilder.h"
 #include "IOptimusExecutionDomainProvider.h"
 #include "IOptimusParameterBindingProvider.h"
@@ -47,7 +48,6 @@
 
 #define LOCTEXT_NAMESPACE "OptimusDetailCustomization"
 
-
 TSharedRef<IPropertyTypeCustomization> FOptimusDataTypeRefCustomization::MakeInstance()
 {
 	return MakeShared<FOptimusDataTypeRefCustomization>();
@@ -1557,43 +1557,67 @@ TSharedRef<IPropertyTypeCustomization> FOptimusDeformerInstanceComponentBindingC
 }
 
 
-
 void FOptimusDeformerInstanceComponentBindingCustomization::CustomizeHeader(
 	TSharedRef<IPropertyHandle> InPropertyHandle,
 	FDetailWidgetRow& InHeaderRow,
 	IPropertyTypeCustomizationUtils& InCustomizationUtils
 	)
 {
-	const TSharedPtr<IPropertyHandle> NameProperty = InPropertyHandle->GetChildHandle(GET_MEMBER_NAME_CHECKED(FOptimusDeformerInstanceComponentBinding, ProviderName));
-	const TSharedPtr<IPropertyHandle> ComponentProperty = InPropertyHandle->GetChildHandle(GET_MEMBER_NAME_CHECKED(FOptimusDeformerInstanceComponentBinding, ActorComponent));
+	const TSharedPtr<IPropertyHandle> BindingPropertyHandle = InPropertyHandle->GetChildHandle(GET_MEMBER_NAME_CHECKED(FOptimusDeformerInstanceComponentBinding, ProviderName));
+	const TSharedPtr<IPropertyHandle> ComponentPropertyHandle = InPropertyHandle->GetChildHandle(GET_MEMBER_NAME_CHECKED(FOptimusDeformerInstanceComponentBinding, ComponentName));
 	
 	FName BindingName;
-	NameProperty->GetValue(BindingName);
+	BindingPropertyHandle->GetValue(BindingName);
 
-	UObject* SelectedComponent;
-	ComponentProperty->GetValue(SelectedComponent);
+	FName ComponentName;
+	ComponentPropertyHandle->GetValue(ComponentName);
 	
 	TArray<UObject*> OuterObjects;
 	InPropertyHandle->GetOuterObjects(OuterObjects);
 
-	const UOptimusDeformerInstanceSettings* BindingProvider = Cast<UOptimusDeformerInstanceSettings>(OuterObjects[0]);
-	const UOptimusComponentSourceBinding* Binding = nullptr;
+	UOptimusDeformerInstanceSettings const* BindingProvider = Cast<UOptimusDeformerInstanceSettings>(OuterObjects[0]);
+	AActor const* OwningActor = BindingProvider ? BindingProvider->GetActor() : nullptr;
 
+	const UOptimusComponentSourceBinding* Binding = nullptr;
 	FComponentHandle SelectedComponentHandle;
 
 	if (BindingProvider)
 	{
-		const AActor* OwningActor = BindingProvider->GetActor();
 		Binding = BindingProvider->GetComponentBindingByName(BindingName);
 		
 		if (OwningActor && Binding)
 		{
 			TArray<UActorComponent*> FilteredComponents;
+
+			// Add actor components
 			OwningActor->GetComponents(Binding->GetComponentSource()->GetComponentClass(), FilteredComponents);
-			for (const UActorComponent* Component: FilteredComponents)
+			// Add unset "auto" setting.
+			FilteredComponents.Add(nullptr);
+
+			// Add blueprint components
+			if (OwningActor->HasAnyFlags(RF_ArchetypeObject))
 			{
-				ComponentHandles.Add(MakeShared<FSoftObjectPath>(FSoftObjectPath::GetOrCreateIDForObject(Component)));
-				if (Component == SelectedComponent)
+				if (UBlueprintGeneratedClass* BlueprintGeneratedClass = Cast<UBlueprintGeneratedClass>(OwningActor->GetClass()))
+				{
+					const TArray<USCS_Node*>& ActorBlueprintNodes = BlueprintGeneratedClass->SimpleConstructionScript->GetAllNodes();
+					for (USCS_Node* Node : ActorBlueprintNodes)
+					{
+						if (Node->ComponentTemplate->IsA(Binding->GetComponentSource()->GetComponentClass()))
+						{
+							FilteredComponents.Add(Node->ComponentTemplate);
+						}
+					}
+				}
+			}
+
+			// Add filtered components to drop down
+			for (const UActorComponent* Component : FilteredComponents)
+			{
+				FSoftObjectPath Path = Component ? FSoftObjectPath::GetOrCreateIDForObject(Component) : FSoftObjectPath();
+				ComponentHandles.Add(MakeShared<FSoftObjectPath>(Path));
+				
+				FName Name = Component ? FOptimusDeformerInstanceComponentBinding::GetSanitizedComponentName(Component) : FName();
+				if (Name.ToString() == ComponentName.ToString())
 				{
 					SelectedComponentHandle = ComponentHandles.Last();
 				}
@@ -1604,7 +1628,7 @@ void FOptimusDeformerInstanceComponentBindingCustomization::CustomizeHeader(
 	InHeaderRow
 	.NameContent()
 	[
-		NameProperty->CreatePropertyNameWidget(FText::FromName(BindingName))
+		ComponentPropertyHandle->CreatePropertyNameWidget(FText::FromName(BindingName))
 	]
 	.ValueContent()
 	[
@@ -1615,14 +1639,14 @@ void FOptimusDeformerInstanceComponentBindingCustomization::CustomizeHeader(
 		.OnGenerateWidget_Lambda([](const FComponentHandle InComponentHandle)
 		{
 			const UActorComponent* Component = Cast<UActorComponent>(InComponentHandle->ResolveObject());
-			
+
 			return SNew(SHorizontalBox)
 				+SHorizontalBox::Slot()
 				.AutoWidth()
 				.VAlign(VAlign_Center)
 				[
 					SNew(SImage)
-					.Image(FSlateIconFinder::FindIconBrushForClass(Component ? Component->GetClass() : nullptr, TEXT("SCS.Component")))
+					.Image(Component ? FSlateIconFinder::FindIconBrushForClass(Component->GetClass(), TEXT("SCS.Component")) : nullptr)
 					.ColorAndOpacity(FSlateColor::UseForeground())
 				]
 				+SHorizontalBox::Slot()
@@ -1631,15 +1655,15 @@ void FOptimusDeformerInstanceComponentBindingCustomization::CustomizeHeader(
 				.Padding(2.0f, 0, 0, 0)
 				[
 					SNew(STextBlock)
-					.Text(FText::FromName(Component ? Component->GetFName() : FName("<Invalid>")))
+					.Text(Component ? FText::FromName(FOptimusDeformerInstanceComponentBinding::GetSanitizedComponentName(Component)) : LOCTEXT("AutoName", "Auto"))
 				];
 		})
-		.OnSelectionChanged_Lambda([ComponentProperty](const FComponentHandle InComponentHandle, ESelectInfo::Type InInfo)
+		.OnSelectionChanged_Lambda([ComponentPropertyHandle](const FComponentHandle InComponentHandle, ESelectInfo::Type InInfo)
 		{
 			if (InInfo != ESelectInfo::Direct)
 			{
 				const UActorComponent* Component = Cast<UActorComponent>(InComponentHandle->ResolveObject());
-				ComponentProperty->SetValue(Component);
+				ComponentPropertyHandle->SetValue(FOptimusDeformerInstanceComponentBinding::GetSanitizedComponentName(Component));
 			}
 		})
 		[
@@ -1649,13 +1673,12 @@ void FOptimusDeformerInstanceComponentBindingCustomization::CustomizeHeader(
 			.AutoWidth()
 			[
 				SNew(SImage)
-				.Image_Lambda([ComponentProperty]()-> const FSlateBrush*
+				.Image_Lambda([ComponentPropertyHandle, OwningActor]()-> const FSlateBrush*
 				{
-					if (UObject* ComponentObject = nullptr; ComponentProperty->GetValue(ComponentObject) == FPropertyAccess::Success && ComponentObject)
-					{
-						return FSlateIconFinder::FindIconBrushForClass(ComponentObject->GetClass(), TEXT("SCS.Component"));
-					}
-					return nullptr;
+					FName Name;
+					ComponentPropertyHandle->GetValue(Name);
+					UActorComponent* ComponentObject = (OwningActor == nullptr || Name.IsNone()) ? nullptr : FOptimusDeformerInstanceComponentBinding::GetActorComponent(OwningActor, Name.ToString()).Get();
+					return ComponentObject ? FSlateIconFinder::FindIconBrushForClass(ComponentObject->GetClass(), TEXT("SCS.Component")) : nullptr;
 				})
 				.ColorAndOpacity(FSlateColor::UseForeground())
 			]
@@ -1666,17 +1689,27 @@ void FOptimusDeformerInstanceComponentBindingCustomization::CustomizeHeader(
 			[
 				SNew(STextBlock)
 				.Font(IPropertyTypeCustomizationUtils::GetRegularFont())
-				.Text_Lambda([ComponentProperty]()
+				.Text_Lambda([ComponentPropertyHandle]()
 				{
-					if (UObject* ComponentObject = nullptr; ComponentProperty->GetValue(ComponentObject) == FPropertyAccess::Success && ComponentObject)
-					{
-						return FText::FromName(ComponentObject->GetFName());
-					}
-					return FText::GetEmpty();
+					FName Name;
+					ComponentPropertyHandle->GetValue(Name);
+					return Name.IsNone() ? LOCTEXT("AutoName", "Auto") : FText::FromName(Name);
 				})
 			]
 		]
-	];
+	]
+ 	.OverrideResetToDefault(FResetToDefaultOverride::Create(
+ 		FIsResetToDefaultVisible::CreateLambda([](TSharedPtr<IPropertyHandle> PropertyHandle) 
+		{
+			const TSharedPtr<IPropertyHandle> ComponentPropertyHandle = PropertyHandle->GetChildHandle(GET_MEMBER_NAME_CHECKED(FOptimusDeformerInstanceComponentBinding, ComponentName));
+			return PropertyHandle->GetIndexInArray() != 0 && ComponentPropertyHandle->CanResetToDefault();
+		}),
+ 		FResetToDefaultHandler::CreateLambda([](TSharedPtr<IPropertyHandle> PropertyHandle) 
+		{
+			TSharedPtr<IPropertyHandle> ComponentPropertyHandle = PropertyHandle->GetChildHandle(GET_MEMBER_NAME_CHECKED(FOptimusDeformerInstanceComponentBinding, ComponentName));
+			ComponentPropertyHandle->ResetToDefault();
+		})
+	));
 }
 
 #undef LOCTEXT_NAMESPACE
