Optimus: Ensure pins with incompatible data domains cannot be connected.

#rnx

diff --git a/Engine/Plugins/Animation/DeformerGraph/Source/OptimusCore/Private/OptimusNodePin.cpp b/Engine/Plugins/Animation/DeformerGraph/Source/OptimusCore/Private/OptimusNodePin.cpp
--- a/Engine/Plugins/Animation/DeformerGraph/Source/OptimusCore/Private/OptimusNodePin.cpp
+++ b/Engine/Plugins/Animation/DeformerGraph/Source/OptimusCore/Private/OptimusNodePin.cpp
@@ -483,7 +483,8 @@ bool UOptimusNodePin::CanCannect(const UOptimusNodePin* InOtherPin, FString* Out
 			{
 				*OutReason = FString::Printf(TEXT("Can't connect resources with different data domain types (%s vs %s)."),
 					*FormatDataDomain(OutputPin->DataDomain), *FormatDataDomain(InputPin->DataDomain));
-			}			
+			}
+			return false;
 		}
 	}
 
