Optimus: Prevent a crash when trying to traverse a pin from a deleted node.

#rnx

diff --git a/Engine/Plugins/Animation/DeformerGraph/Source/OptimusEditor/Private/Widgets/SOptimusEditorGraphNode.cpp b/Engine/Plugins/Animation/DeformerGraph/Source/OptimusEditor/Private/Widgets/SOptimusEditorGraphNode.cpp
--- a/Engine/Plugins/Animation/DeformerGraph/Source/OptimusEditor/Private/Widgets/SOptimusEditorGraphNode.cpp
+++ b/Engine/Plugins/Animation/DeformerGraph/Source/OptimusEditor/Private/Widgets/SOptimusEditorGraphNode.cpp
@@ -265,7 +265,14 @@ private:
 	{
 		if (TObjectPtr<UOptimusNodePin> ModelPin = ModelPinPtr.Get())
 		{
-			TSet<UOptimusComponentSourceBinding*> ComponentSourceBindings = ModelPin->GetComponentSourceBindings();
+			TSet<UOptimusComponentSourceBinding*> ComponentSourceBindings;
+
+			// Make sure it doesn't belong to a node that's just been deleted, since Slate updates are usually a frame
+			// behind.
+			if (ModelPin->GetPackage() != GetTransientPackage())
+			{
+				ComponentSourceBindings = ModelPin->GetComponentSourceBindings();
+			}
 
 			if (ComponentSourceBindings.IsEmpty())
 			{
