[Plugins] fix for using deprecated TObjectPtr.IsNull

diff --git a/Engine/Plugins/Animation/DeformerGraph/Source/OptimusCore/Private/Nodes/OptimusNode_DataInterface.cpp b/Engine/Plugins/Animation/DeformerGraph/Source/OptimusCore/Private/Nodes/OptimusNode_DataInterface.cpp
--- a/Engine/Plugins/Animation/DeformerGraph/Source/OptimusCore/Private/Nodes/OptimusNode_DataInterface.cpp
+++ b/Engine/Plugins/Animation/DeformerGraph/Source/OptimusCore/Private/Nodes/OptimusNode_DataInterface.cpp
@@ -184,7 +184,7 @@ void UOptimusNode_DataInterface::PostLoad()
 	Super::PostLoad();
 
 	// Previously DataInterfaceData wasn't always created.
-	if (!DataInterfaceClass.IsNull() && DataInterfaceData.IsNull())
+	if (DataInterfaceClass && !DataInterfaceData)
 	{
 		DataInterfaceData = NewObject<UOptimusComputeDataInterface>(this, DataInterfaceClass);
 	}
@@ -200,7 +200,7 @@ void UOptimusNode_DataInterface::PostLoad()
 void UOptimusNode_DataInterface::ConstructNode()
 {
 	// Create the component pin.
-	if (ensure(!DataInterfaceClass.IsNull()) &&
+	if (ensure(DataInterfaceClass) &&
 		ensure(DataInterfaceData))
 	{
 		SetDisplayName(FText::FromString(DataInterfaceData->GetDisplayName()));
diff --git a/Engine/Plugins/Experimental/Animation/PoseSearch/Source/Runtime/Private/PoseSearch.cpp b/Engine/Plugins/Experimental/Animation/PoseSearch/Source/Runtime/Private/PoseSearch.cpp
--- a/Engine/Plugins/Experimental/Animation/PoseSearch/Source/Runtime/Private/PoseSearch.cpp
+++ b/Engine/Plugins/Experimental/Animation/PoseSearch/Source/Runtime/Private/PoseSearch.cpp
@@ -432,7 +432,7 @@ bool UPoseSearchSchema::IsValid() const
 
 	for (const TObjectPtr<UPoseSearchFeatureChannel>& Channel: Channels)
 	{
-		bValid &= !Channel.IsNull();
+		bValid &= Channel != nullptr;
 	}
 
 	bValid &= (BoneReferences.Num() == BoneIndices.Num());
@@ -2730,7 +2730,7 @@ const FPoseSearchIndexAsset* FSearchResult::GetSearchIndexAsset(bool bMandatory)
 
 static void FillCompactPoseAndComponentRefRotations(FAssetSamplingContext& Context)
 {
-	if (!Context.MirrorDataTable.IsNull())
+	if (Context.MirrorDataTable)
 	{
 		const UMirrorDataTable* MirrorDataTablePtr = Context.MirrorDataTable.Get();
 		MirrorDataTablePtr->FillCompactPoseAndComponentRefRotations(
diff --git a/Engine/Plugins/Experimental/Web/WebAPI/Source/WebAPIEditor/Private/WebAPIDefinition.cpp b/Engine/Plugins/Experimental/Web/WebAPI/Source/WebAPIEditor/Private/WebAPIDefinition.cpp
--- a/Engine/Plugins/Experimental/Web/WebAPI/Source/WebAPIEditor/Private/WebAPIDefinition.cpp
+++ b/Engine/Plugins/Experimental/Web/WebAPI/Source/WebAPIEditor/Private/WebAPIDefinition.cpp
@@ -1,4 +1,4 @@
-﻿// Copyright Epic Games, Inc. All Rights Reserved.
+// Copyright Epic Games, Inc. All Rights Reserved.
 
 #include "WebAPIDefinition.h"
 
@@ -72,7 +72,7 @@ UWebAPIDefinition::~UWebAPIDefinition()
 UObject* UWebAPIDefinition::AddOrGetImportedDataCache(FName InKey, const TSubclassOf<UObject>& InDataCacheClass)
 {
 	TObjectPtr<UObject>* Result = ImportedDataCache.Find(InKey);
-	if(Result && !Result->IsNull())
+	if(Result && *Result)
 	{
 		return *Result;
 	}
diff --git a/Engine/Plugins/Experimental/Web/WebAPI/Source/WebAPIOpenAPI/Private/Tests/WebAPISwagger.spec.cpp b/Engine/Plugins/Experimental/Web/WebAPI/Source/WebAPIOpenAPI/Private/Tests/WebAPISwagger.spec.cpp
--- a/Engine/Plugins/Experimental/Web/WebAPI/Source/WebAPIOpenAPI/Private/Tests/WebAPISwagger.spec.cpp
+++ b/Engine/Plugins/Experimental/Web/WebAPI/Source/WebAPIOpenAPI/Private/Tests/WebAPISwagger.spec.cpp
@@ -1,4 +1,4 @@
-﻿// Copyright Epic Games, Inc. All Rights Reserved.
+// Copyright Epic Games, Inc. All Rights Reserved.
 
 #if WITH_EDITOR
 
diff --git a/Engine/Plugins/Experimental/Web/WebAPI/Source/WebAPIOpenAPI/Private/V2/WebAPISwaggerConverter.cpp b/Engine/Plugins/Experimental/Web/WebAPI/Source/WebAPIOpenAPI/Private/V2/WebAPISwaggerConverter.cpp
--- a/Engine/Plugins/Experimental/Web/WebAPI/Source/WebAPIOpenAPI/Private/V2/WebAPISwaggerConverter.cpp
+++ b/Engine/Plugins/Experimental/Web/WebAPI/Source/WebAPIOpenAPI/Private/V2/WebAPISwaggerConverter.cpp
@@ -1,4 +1,4 @@
-﻿// Copyright Epic Games, Inc. All Rights Reserved.
+// Copyright Epic Games, Inc. All Rights Reserved.
 
 #include "WebAPISwaggerConverter.h"
 #include "WebAPISwaggerConverter.inl"
@@ -307,7 +307,7 @@ namespace UE::WebAPI::Swagger
 			}
 		}
 
-		check(!Result.IsNull());
+		check(Result);
 
 		return Result;
 	}
diff --git a/Engine/Plugins/Experimental/Web/WebAPI/Source/WebAPIOpenAPI/Private/V3/WebAPIOpenAPIConverter.inl b/Engine/Plugins/Experimental/Web/WebAPI/Source/WebAPIOpenAPI/Private/V3/WebAPIOpenAPIConverter.inl
--- a/Engine/Plugins/Experimental/Web/WebAPI/Source/WebAPIOpenAPI/Private/V3/WebAPIOpenAPIConverter.inl
+++ b/Engine/Plugins/Experimental/Web/WebAPI/Source/WebAPIOpenAPI/Private/V3/WebAPIOpenAPIConverter.inl
@@ -1,4 +1,4 @@
-﻿// Copyright Epic Games, Inc. All Rights Reserved.
+// Copyright Epic Games, Inc. All Rights Reserved.
 
 #pragma once
 
@@ -117,7 +117,7 @@ namespace UE::WebAPI::OpenAPI
 			}
 		}
 
-		check(!Result.IsNull());
+		check(Result);
 		
 		return Result;
 	}
