Deformer graph: added an icon to shader text editor

diff --git a/Engine/Plugins/Animation/DeformerGraph/Source/OptimusEditor/Private/OptimusEditorTabSummoners.cpp b/Engine/Plugins/Animation/DeformerGraph/Source/OptimusEditor/Private/OptimusEditorTabSummoners.cpp
--- a/Engine/Plugins/Animation/DeformerGraph/Source/OptimusEditor/Private/OptimusEditorTabSummoners.cpp
+++ b/Engine/Plugins/Animation/DeformerGraph/Source/OptimusEditor/Private/OptimusEditorTabSummoners.cpp
@@ -124,6 +124,7 @@ FOptimusEditorShaderTextEditorTabSummoner::FOptimusEditorShaderTextEditorTabSumm
 	EditorPtr(InEditorApp)
 {
 	TabLabel = LOCTEXT("ShaderTextEditorTab_TabLabel", "Shader Text Editor");
+	TabIcon = FSlateIcon(FAppStyle::GetAppStyleSetName(), "Icons.TextEditor");
 	
 	ViewMenuDescription = LOCTEXT("ShaderTextEditorTab_MenuLabel", "Shader Text Editor");
 	ViewMenuTooltip = LOCTEXT("ShaderTextEditorTab_MenuLabel_Tooltip", "Show the Shader Text Editor Tab");	
diff --git a/Engine/Source/Editor/EditorStyle/Private/StarshipStyle.cpp b/Engine/Source/Editor/EditorStyle/Private/StarshipStyle.cpp
--- a/Engine/Source/Editor/EditorStyle/Private/StarshipStyle.cpp
+++ b/Engine/Source/Editor/EditorStyle/Private/StarshipStyle.cpp
@@ -2576,6 +2576,8 @@ void FStarshipEditorStyle::FStyle::SetupGeneralIcons()
 	Set("Debug", new IMAGE_BRUSH_SVG( "Starship/Common/Debug", Icon16x16 ) );
 	Set("Modules", new IMAGE_BRUSH_SVG( "Starship/Common/Modules", Icon16x16 ) );
 	Set("Versions", new IMAGE_BRUSH_SVG("Starship/Common/Versions", Icon20x20));
+
+	Set("Icons.TextEditor", new IMAGE_BRUSH_SVG("Starship/Common/TextEditor_16", Icon16x16));
 }
 
 void FStarshipEditorStyle::FStyle::SetupWindowStyles()
