Optimus: Use the correct float type for variables that are created with no type specified.

#rnx

diff --git a/Engine/Plugins/Animation/DeformerGraph/Source/OptimusCore/Private/OptimusDeformer.cpp b/Engine/Plugins/Animation/DeformerGraph/Source/OptimusCore/Private/OptimusDeformer.cpp
--- a/Engine/Plugins/Animation/DeformerGraph/Source/OptimusCore/Private/OptimusDeformer.cpp
+++ b/Engine/Plugins/Animation/DeformerGraph/Source/OptimusCore/Private/OptimusDeformer.cpp
@@ -146,7 +146,7 @@ UOptimusVariableDescription* UOptimusDeformer::AddVariable(
 	if (!InDataTypeRef.IsValid())
 	{
 		// Default to float.
-		InDataTypeRef.Set(FOptimusDataTypeRegistry::Get().FindType(*FFloatProperty::StaticClass()));
+		InDataTypeRef.Set(FOptimusDataTypeRegistry::Get().FindType(*FDoubleProperty::StaticClass()));
 	}
 
 	// Is this data type compatible with resources?
diff --git a/Engine/Plugins/Animation/DeformerGraph/Source/OptimusCore/Private/OptimusDeformerInstance.cpp b/Engine/Plugins/Animation/DeformerGraph/Source/OptimusCore/Private/OptimusDeformerInstance.cpp
--- a/Engine/Plugins/Animation/DeformerGraph/Source/OptimusCore/Private/OptimusDeformerInstance.cpp
+++ b/Engine/Plugins/Animation/DeformerGraph/Source/OptimusCore/Private/OptimusDeformerInstance.cpp
@@ -492,9 +492,15 @@ bool UOptimusDeformerInstance::SetIntVariable(FName InVariableName, int32 InValu
 	return SetVariableValue<int32>(Variables, InVariableName, FIntProperty::StaticClass()->GetFName(), InValue);
 }
 
-bool UOptimusDeformerInstance::SetFloatVariable(FName InVariableName, float InValue)
+bool UOptimusDeformerInstance::SetFloatVariable(FName InVariableName, double InValue)
 {
-	return SetVariableValue<float>(Variables, InVariableName, FFloatProperty::StaticClass()->GetFName(), InValue);
+	if (SetVariableValue<double>(Variables, InVariableName, FDoubleProperty::StaticClass()->GetFName(), InValue))
+	{
+		return true;
+	}
+	
+	// Fall back on float
+	return SetVariableValue<float>(Variables, InVariableName, FFloatProperty::StaticClass()->GetFName(), static_cast<float>(InValue));
 }
 
 bool UOptimusDeformerInstance::SetVectorVariable(FName InVariableName, const FVector& InValue)
diff --git a/Engine/Plugins/Animation/DeformerGraph/Source/OptimusCore/Private/OptimusNodePin.cpp b/Engine/Plugins/Animation/DeformerGraph/Source/OptimusCore/Private/OptimusNodePin.cpp
--- a/Engine/Plugins/Animation/DeformerGraph/Source/OptimusCore/Private/OptimusNodePin.cpp
+++ b/Engine/Plugins/Animation/DeformerGraph/Source/OptimusCore/Private/OptimusNodePin.cpp
@@ -454,7 +454,9 @@ bool UOptimusNodePin::CanCannect(const UOptimusNodePin* InOtherPin, FString* Out
 	}
 
 	// Check for incompatible types.
-	if (DataType != InOtherPin->DataType)
+	if (!((DataType->ShaderValueType.IsValid() && InOtherPin->DataType->ShaderValueType.IsValid() &&
+		  DataType->ShaderValueType == InOtherPin->DataType->ShaderValueType) ||
+		 DataType == InOtherPin->DataType))
 	{
 		// TBD: Automatic conversion.
 		if (OutReason)
diff --git a/Engine/Plugins/Animation/DeformerGraph/Source/OptimusCore/Private/OptimusResourceDescription.cpp b/Engine/Plugins/Animation/DeformerGraph/Source/OptimusCore/Private/OptimusResourceDescription.cpp
--- a/Engine/Plugins/Animation/DeformerGraph/Source/OptimusCore/Private/OptimusResourceDescription.cpp
+++ b/Engine/Plugins/Animation/DeformerGraph/Source/OptimusCore/Private/OptimusResourceDescription.cpp
@@ -2,6 +2,7 @@
 
 #include "OptimusResourceDescription.h"
 
+#include "OptimusDataTypeRegistry.h"
 #include "OptimusDeformer.h"
 #include "OptimusHelpers.h"
 #include "DataInterfaces/OptimusDataInterfaceRawBuffer.h"
@@ -24,6 +25,16 @@ void UOptimusResourceDescription::PostLoad()
 		DataInterface->DataDomain = DataDomain;
 		DataInterface->ComponentSourceBinding = ComponentBinding;
 	}
+	
+	// 64-bit float data type is not supported for resources although they were allowed before. Do an in-place upgrade here. 
+	const FOptimusDataTypeHandle FloatDataType = FOptimusDataTypeRegistry::Get().FindType(*FFloatProperty::StaticClass());
+	const FOptimusDataTypeHandle DoubleDataType = FOptimusDataTypeRegistry::Get().FindType(*FDoubleProperty::StaticClass());
+	
+	if (DataType == DoubleDataType)
+	{
+		DataType = FloatDataType;
+		(void)MarkPackageDirty();
+	}
 }
 
 
diff --git a/Engine/Plugins/Animation/DeformerGraph/Source/OptimusCore/Private/OptimusVariableDescription.cpp b/Engine/Plugins/Animation/DeformerGraph/Source/OptimusCore/Private/OptimusVariableDescription.cpp
--- a/Engine/Plugins/Animation/DeformerGraph/Source/OptimusCore/Private/OptimusVariableDescription.cpp
+++ b/Engine/Plugins/Animation/DeformerGraph/Source/OptimusCore/Private/OptimusVariableDescription.cpp
@@ -2,6 +2,7 @@
 
 #include "OptimusVariableDescription.h"
 
+#include "OptimusDataTypeRegistry.h"
 #include "OptimusDeformer.h"
 #include "OptimusHelpers.h"
 #include "OptimusValueContainer.h"
@@ -31,8 +32,25 @@ UOptimusDeformer* UOptimusVariableDescription::GetOwningDeformer() const
 	return Container ? CastChecked<UOptimusDeformer>(Container->GetOuter()) : nullptr;
 }
 
+void UOptimusVariableDescription::PostLoad()
+{
+	Super::PostLoad();
+
+	// 32-bit float data type is not supported for variables although they were allowed before. Do an in-place upgrade here. 
+	const FOptimusDataTypeHandle FloatDataType = FOptimusDataTypeRegistry::Get().FindType(*FFloatProperty::StaticClass());
+	const FOptimusDataTypeHandle DoubleDataType = FOptimusDataTypeRegistry::Get().FindType(*FDoubleProperty::StaticClass());
+	
+	if (DataType == FloatDataType)
+	{
+		DataType = DoubleDataType;
+		ResetValueDataSize();
+		(void)MarkPackageDirty();
+	}
+}
+
 
 #if WITH_EDITOR
+
 void UOptimusVariableDescription::PostEditChangeProperty(FPropertyChangedEvent& PropertyChangedEvent)
 {
 	const FName PropertyName = PropertyChangedEvent.GetPropertyName();
diff --git a/Engine/Plugins/Animation/DeformerGraph/Source/OptimusCore/Public/OptimusDeformerInstance.h b/Engine/Plugins/Animation/DeformerGraph/Source/OptimusCore/Public/OptimusDeformerInstance.h
--- a/Engine/Plugins/Animation/DeformerGraph/Source/OptimusCore/Public/OptimusDeformerInstance.h
+++ b/Engine/Plugins/Animation/DeformerGraph/Source/OptimusCore/Public/OptimusDeformerInstance.h
@@ -163,22 +163,22 @@ public:
 	void SetupFromDeformer(UOptimusDeformer* InDeformer);
 
 	/** Set the value of a boolean variable. */
-	UFUNCTION(BlueprintPure, Category="Deformer", meta=(DisplayName="Set Variable (bool)"))
+	UFUNCTION(BlueprintPure, Category="Deformer", meta=(DisplayName="Set Variable (Boolean)"))
 	bool SetBoolVariable(FName InVariableName, bool InValue);
 
-	/** Set the value of a boolean variable. */
-	UFUNCTION(BlueprintCallable, Category="Deformer", meta=(DisplayName="Set Variable (int)"))
+	/** Set the value of an integer variable. */
+	UFUNCTION(BlueprintCallable, Category="Deformer", meta=(DisplayName="Set Variable (Integer)"))
 	bool SetIntVariable(FName InVariableName, int32 InValue);
 
-	/** Set the value of a boolean variable. */
-	UFUNCTION(BlueprintCallable, Category="Deformer", meta=(DisplayName="Set Variable (float)"))
-	bool SetFloatVariable(FName InVariableName, float InValue);
+	/** Set the value of a float variable. */
+	UFUNCTION(BlueprintCallable, Category="Deformer", meta=(DisplayName="Set Variable (Float)"))
+	bool SetFloatVariable(FName InVariableName, double InValue);
 
-	/** Set the value of a boolean variable. */
+	/** Set the value of a 3-vector variable. */
 	UFUNCTION(BlueprintCallable, Category="Deformer", meta=(DisplayName="Set Variable (Vector)"))
 	bool SetVectorVariable(FName InVariableName, const FVector& InValue);
 
-	/** Set the value of a boolean variable. */
+	/** Set the value of a 4-vector variable. */
 	UFUNCTION(BlueprintCallable, Category="Deformer", meta=(DisplayName="Set Variable (Vector4)"))
 	bool SetVector4Variable(FName InVariableName, const FVector4& InValue);
 
diff --git a/Engine/Plugins/Animation/DeformerGraph/Source/OptimusCore/Public/OptimusVariableDescription.h b/Engine/Plugins/Animation/DeformerGraph/Source/OptimusCore/Public/OptimusVariableDescription.h
--- a/Engine/Plugins/Animation/DeformerGraph/Source/OptimusCore/Public/OptimusVariableDescription.h
+++ b/Engine/Plugins/Animation/DeformerGraph/Source/OptimusCore/Public/OptimusVariableDescription.h
@@ -42,8 +42,6 @@ class OPTIMUSCORE_API UOptimusVariableDescription :
 {
 	GENERATED_BODY()
 public:
-	UOptimusVariableDescription() = default;
-
 	/** 
 	 * Set the value data storage to match the size required by the DataType. 
 	 * If a reallocation is required then the value data will be zeroed.
@@ -74,6 +72,7 @@ public:
 	UPROPERTY()
 	TArray<uint8> ValueData;
 
+	void PostLoad();
 #if WITH_EDITOR
 	void PostEditChangeProperty(struct FPropertyChangedEvent& PropertyChangedEvent) override;
 	void PreEditUndo() override;
