Optimus: Display an error for unconnected input binding that is a struct to avoid crash

diff --git a/Engine/Plugins/Animation/DeformerGraph/Source/OptimusCore/Private/Nodes/OptimusNode_ComputeKernelBase.cpp b/Engine/Plugins/Animation/DeformerGraph/Source/OptimusCore/Private/Nodes/OptimusNode_ComputeKernelBase.cpp
--- a/Engine/Plugins/Animation/DeformerGraph/Source/OptimusCore/Private/Nodes/OptimusNode_ComputeKernelBase.cpp
+++ b/Engine/Plugins/Animation/DeformerGraph/Source/OptimusCore/Private/Nodes/OptimusNode_ComputeKernelBase.cpp
@@ -144,22 +144,31 @@ FOptimus_ComputeKernelResult UOptimusNode_ComputeKernelBase::CreateComputeKernel
 		{
 			for (const UOptimusNodePin* SubPin: Pin->GetSubPins())
 			{
-				ProcessInputPinForComputeKernel(
+				TOptional<FText> Result = ProcessInputPinForComputeKernel(
 					InTraversalContext, SubPin, Pin->GetName(),
 					InNodeDataInterfaceMap, InLinkDataInterfaceMap, InValueNodes,
 					InGraphDataInterface, InGraphDataComponentBinding,
 					KernelSource, GeneratedFunctions, OutInputDataBindings
 					);
+				if (Result.IsSet())
+				{
+					return ReturnError(Result.GetValue());
+				}
 			}
 		}
 		else if (Pin->GetDirection() == EOptimusNodePinDirection::Input)
 		{
-			ProcessInputPinForComputeKernel(
+			TOptional<FText> Result = ProcessInputPinForComputeKernel(
 				InTraversalContext, Pin, FString(),
 				InNodeDataInterfaceMap, InLinkDataInterfaceMap, InValueNodes,
 				InGraphDataInterface, InGraphDataComponentBinding,
 				KernelSource, GeneratedFunctions, OutInputDataBindings
 				);
+			
+			if (Result.IsSet())
+			{
+				return ReturnError(Result.GetValue());
+			}
 		}
 		else if (Pin->GetDirection() == EOptimusNodePinDirection::Output)
 		{
@@ -353,7 +362,7 @@ FString UOptimusNode_ComputeKernelBase::GetCookedKernelSource(
 }
 
 
-void UOptimusNode_ComputeKernelBase::ProcessInputPinForComputeKernel(
+TOptional<FText> UOptimusNode_ComputeKernelBase::ProcessInputPinForComputeKernel(
 	const FOptimusPinTraversalContext& InTraversalContext,
 	const UOptimusNodePin* InInputPin,
 	const FString& InGroupName,
@@ -478,6 +487,12 @@ void UOptimusNode_ComputeKernelBase::ProcessInputPinForComputeKernel(
 		FString OptionalParamStr;
 		if (InInputPin->GetDataDomain().IsSingleton())
 		{
+			// it is not easy to provide a default/zero value for a struct type variable, so we error out here
+			if (ValueType->Type == EShaderFundamentalType::Struct)
+			{
+				return FText::Format(LOCTEXT("DisconnectedStructPinUnsupported", "Disconnected struct type pin '{0}' is not supported"), FText::FromName(InInputPin->GetUniqueName()));
+			}
+			
 			ValueStr = GetShaderParamPinValueString(InInputPin);
 		}
 		else
@@ -509,6 +524,8 @@ void UOptimusNode_ComputeKernelBase::ProcessInputPinForComputeKernel(
 			FString::Printf(TEXT("%s%s Read%s(%s) { return %s; }%s"),
 				*NamespacePrefix, *ValueType->ToString(), *InInputPin->GetName(), *OptionalParamStr, *ValueStr, *NamespaceSuffix));
 	}
+
+	return {};
 }
 
 
diff --git a/Engine/Plugins/Animation/DeformerGraph/Source/OptimusCore/Private/Nodes/OptimusNode_ComputeKernelBase.h b/Engine/Plugins/Animation/DeformerGraph/Source/OptimusCore/Private/Nodes/OptimusNode_ComputeKernelBase.h
--- a/Engine/Plugins/Animation/DeformerGraph/Source/OptimusCore/Private/Nodes/OptimusNode_ComputeKernelBase.h
+++ b/Engine/Plugins/Animation/DeformerGraph/Source/OptimusCore/Private/Nodes/OptimusNode_ComputeKernelBase.h
@@ -74,7 +74,7 @@ protected:
 	
 	
 private:
-	void ProcessInputPinForComputeKernel(
+	TOptional<FText> ProcessInputPinForComputeKernel(
 		const FOptimusPinTraversalContext& InTraversalContext,
 		const UOptimusNodePin* InInputPin,
 		const FString& InGroupName,
