Optimus: Fix multiple debug draw contexts.

diff --git a/Engine/Plugins/Animation/DeformerGraph/Shaders/Private/DataInterfaceDebugDraw.ush b/Engine/Plugins/Animation/DeformerGraph/Shaders/Private/DataInterfaceDebugDraw.ush
--- a/Engine/Plugins/Animation/DeformerGraph/Shaders/Private/DataInterfaceDebugDraw.ush
+++ b/Engine/Plugins/Animation/DeformerGraph/Shaders/Private/DataInterfaceDebugDraw.ush
@@ -1,22 +1,13 @@
 // Copyright Epic Games, Inc. All Rights Reserved.
 
-//todo[CF]:
-// Only included for [FLATTEN] used by MiniFontCommon.ush
-// We need to factor MiniFontCommon.ush out of ShaderPrint.ush
-#include "/Engine/Private/Common.ush"
-
-//todo[CF]:
-// This is a temporary hack to insert our buffer name into ShaderPrint.ush
-// We need ShaderPrint to be flexible enough to take uniforms and buffers as parameters instead of using the hard coded uniform buffer.
-// Until that is done we can't support multiple of this data interface per kernel (although that would be an unusual requirement).
-#define ShaderPrint_StateBuffer		{DataInterfaceName}_ShaderPrint_StateBuffer
-#define ShaderPrint_RWEntryBuffer	{DataInterfaceName}_ShaderPrint_RWEntryBuffer
-#include "/Engine/Private/ShaderPrint.ush"
+#include "/Engine/Private/ShaderPrintCommon.ush"
 
-float4x4 {DataInterfaceName}_LocalToWorld;
+#ifndef DATA_INTERFACE_DEBUGDRAW_ONCE
+#define DATA_INTERFACE_DEBUGDRAW_ONCE 1
 
 struct FDebugDraw
 {
+	FShaderPrintContext Context;
 	float4x4 LocalToWorld;
 
 	float3 TransformLocalToWorld(float3 P)
@@ -26,17 +17,18 @@ struct FDebugDraw
 
 	void AddLine(float3 Pos0, float3 Pos1, float4 Color)
 	{
-		AddLineTWS(TransformLocalToWorld(Pos0), TransformLocalToWorld(Pos1), Color);
+		AddLineTWS(Context, TransformLocalToWorld(Pos0), TransformLocalToWorld(Pos1), Color);
 	}
 
 	void AddCross(float3 Pos, float Size, float4 Color)
 	{
-		AddCrossTWS(TransformLocalToWorld(Pos), Size, Color);
+		AddCrossTWS(Context, TransformLocalToWorld(Pos), Size, Color);
 	}
 
 	void AddQuad(float3 Pos0, float3 Pos1, float3 Pos2, float3 Pos3, float4 Color)
 	{
 		AddQuadTWS(
+			Context,
 			TransformLocalToWorld(Pos0),
 			TransformLocalToWorld(Pos1),
 			TransformLocalToWorld(Pos2),
@@ -46,13 +38,40 @@ struct FDebugDraw
 
 	void AddAxis(float3 Pos, float3x3 InM, float Scale)
 	{
-		AddReferentialWS(TransformLocalToWorld(Pos), mul(InM, (float3x3)LocalToWorld), Scale);
+		AddReferentialWS(Context, TransformLocalToWorld(Pos), mul(InM, (float3x3)LocalToWorld), Scale);
 	}
 };
 
+#endif // DATA_INTERFACE_DEBUGDRAW_ONCE
+
+float4x4 {DataInterfaceName}_LocalToWorld;
+int2 {DataInterfaceName}_Resolution;
+float2 {DataInterfaceName}_FontSize;
+float2 {DataInterfaceName}_FontSpacing;
+uint {DataInterfaceName}_MaxValueCount;
+uint {DataInterfaceName}_MaxSymbolCount;
+uint {DataInterfaceName}_MaxStateCount;
+uint {DataInterfaceName}_MaxLineCount;
+uint {DataInterfaceName}_MaxTriangleCount;
+StructuredBuffer<uint> {DataInterfaceName}_StateBuffer;
+RWStructuredBuffer<uint> {DataInterfaceName}_RWEntryBuffer;
+
 FDebugDraw ReadDebugDraw_{DataInterfaceName}()
 {
 	FDebugDraw DebugDraw;
+	DebugDraw.Context.bIsActive = true;
+	DebugDraw.Context.StartPos = 0;
+	DebugDraw.Context.Pos = 0;
+	DebugDraw.Context.Config.Resolution = {DataInterfaceName}_Resolution;
+	DebugDraw.Context.Config.FontSize = {DataInterfaceName}_FontSize;
+	DebugDraw.Context.Config.FontSpacing = {DataInterfaceName}_FontSpacing;
+	DebugDraw.Context.Config.MaxValueCount = {DataInterfaceName}_MaxValueCount;
+	DebugDraw.Context.Config.MaxSymbolCount = {DataInterfaceName}_MaxSymbolCount;
+	DebugDraw.Context.Config.MaxStateCount = {DataInterfaceName}_MaxStateCount;
+	DebugDraw.Context.Config.MaxLineCount = {DataInterfaceName}_MaxLineCount;
+	DebugDraw.Context.Config.MaxTriangleCount = {DataInterfaceName}_MaxTriangleCount;
+	DebugDraw.Context.Buffers.StateBuffer = {DataInterfaceName}_StateBuffer;
+	DebugDraw.Context.Buffers.RWEntryBuffer = {DataInterfaceName}_RWEntryBuffer;
 	DebugDraw.LocalToWorld = {DataInterfaceName}_LocalToWorld;
 	return DebugDraw;
 }
diff --git a/Engine/Plugins/Animation/DeformerGraph/Source/OptimusCore/Private/DataInterfaces/OptimusDataInterfaceDebugDraw.cpp b/Engine/Plugins/Animation/DeformerGraph/Source/OptimusCore/Private/DataInterfaces/OptimusDataInterfaceDebugDraw.cpp
--- a/Engine/Plugins/Animation/DeformerGraph/Source/OptimusCore/Private/DataInterfaces/OptimusDataInterfaceDebugDraw.cpp
+++ b/Engine/Plugins/Animation/DeformerGraph/Source/OptimusCore/Private/DataInterfaces/OptimusDataInterfaceDebugDraw.cpp
@@ -2,12 +2,12 @@
 
 #include "OptimusDataInterfaceDebugDraw.h"
 
-// #include "Components/PrimitiveComponent.h"
 #include "Components/SkinnedMeshComponent.h"
 #include "ComputeFramework/ShaderParamTypeDefinition.h"
 #include "OptimusDataDomain.h"
 #include "OptimusDataTypeRegistry.h"
 #include "RenderGraphBuilder.h"
+#include "ShaderCompilerCore.h"
 #include "ShaderCore.h"
 #include "ShaderParameterMetadataBuilder.h"
 
@@ -31,8 +31,7 @@ TArray<FOptimusCDIPinDefinition> UOptimusDebugDrawDataInterface::GetPinDefinitio
 
 TSubclassOf<UActorComponent> UOptimusDebugDrawDataInterface::GetRequiredComponentClass() const
 {
-	// return UPrimitiveComponent::StaticClass();
-	return USkinnedMeshComponent::StaticClass();
+	 return USkinnedMeshComponent::StaticClass();
 }
 
 void UOptimusDebugDrawDataInterface::RegisterTypes() 
@@ -58,7 +57,16 @@ void UOptimusDebugDrawDataInterface::GetSupportedInputs(TArray<FShaderFunctionDe
 
 BEGIN_SHADER_PARAMETER_STRUCT(FDebugDrawDataInterfaceParameters, )
 	SHADER_PARAMETER(FMatrix44f, LocalToWorld)
-	SHADER_PARAMETER_STRUCT_INCLUDE(ShaderPrint::FShaderParameters, ShaderPrintParameters)
+	SHADER_PARAMETER(FIntPoint, Resolution)
+	SHADER_PARAMETER(FVector2f, FontSize)
+	SHADER_PARAMETER(FVector2f, FontSpacing)
+	SHADER_PARAMETER(uint32, MaxValueCount)
+	SHADER_PARAMETER(uint32, MaxSymbolCount)
+	SHADER_PARAMETER(uint32, MaxStateCount)
+	SHADER_PARAMETER(uint32, MaxLineCount)
+	SHADER_PARAMETER(uint32, MaxTriangleCount)
+	SHADER_PARAMETER_RDG_BUFFER_SRV(StructuredBuffer<uint2>, StateBuffer)
+	SHADER_PARAMETER_RDG_BUFFER_UAV(RWStructuredBuffer<ShaderPrintItem>, RWEntryBuffer)
 END_SHADER_PARAMETER_STRUCT()
 
 void UOptimusDebugDrawDataInterface::GetShaderParameters(TCHAR const* UID, FShaderParametersMetadataBuilder& InOutBuilder, FShaderParametersMetadataAllocations& InOutAllocations) const
@@ -94,7 +102,7 @@ UComputeDataProvider* UOptimusDebugDrawDataInterface::CreateDataProvider(TObject
 
 bool UOptimusDebugDrawDataProvider::IsValid() const
 {
-	return true;
+	return PrimitiveComponent != nullptr;
 }
 
 FComputeDataProviderRenderProxy* UOptimusDebugDrawDataProvider::GetRenderProxy()
@@ -104,35 +112,38 @@ FComputeDataProviderRenderProxy* UOptimusDebugDrawDataProvider::GetRenderProxy()
 
 
 FOptimusDebugDrawDataProviderProxy::FOptimusDebugDrawDataProviderProxy(UPrimitiveComponent* InPrimitiveComponent, FOptimusDebugDrawParameters const& InDebugDrawParameters)
+	: Setup(FIntRect(0, 0, 1920, 1080))
 {
 	// Split LocalToWorld into a pre-translation and transform for large world coordinate support.
 	FMatrix RenderMatrix = InPrimitiveComponent->GetRenderMatrix();
-	PreViewTranslation = -RenderMatrix.GetOrigin();
+	FVector PreViewTranslation = -RenderMatrix.GetOrigin();
 	LocalToWorld = FMatrix44f(RenderMatrix.ConcatTranslation(PreViewTranslation));
 
-	DebugDrawParameters = InDebugDrawParameters;
-}
+	Setup.FontSize = InDebugDrawParameters.FontSize;
+	Setup.MaxLineCount = Setup.bEnabled ? InDebugDrawParameters.MaxLineCount : 0;
+	Setup.MaxTriangleCount = Setup.bEnabled ? InDebugDrawParameters.MaxTriangleCount : 0;
+	Setup.MaxValueCount = Setup.bEnabled ? InDebugDrawParameters.MaxCharacterCount : 0;
+	Setup.PreViewTranslation = PreViewTranslation;
 
-void FOptimusDebugDrawDataProviderProxy::AllocateResources(FRDGBuilder& GraphBuilder)
-{
-	// Force enable if requested.
-	if (DebugDrawParameters.bForceEnable)
+	ShaderPrint::GetParameters(Setup, ConfigParameters);
+
+	// Force ShaderPrint system enable if requested.
+	if (InDebugDrawParameters.bForceEnable)
 	{
 		ShaderPrint::SetEnabled(true);
 	}
+}
 
+void FOptimusDebugDrawDataProviderProxy::AllocateResources(FRDGBuilder& GraphBuilder)
+{
 	// Allocate ShaderPrint output buffers.
-	ShaderPrint::FShaderPrintSetup Setup(FIntRect(0, 0, 1920, 1080));
-	Setup.FontSize = DebugDrawParameters.FontSize;
-	Setup.MaxLineCount = Setup.bEnabled ? DebugDrawParameters.MaxLineCount : 0;
-	Setup.MaxTriangleCount = Setup.bEnabled ? DebugDrawParameters.MaxTriangleCount : 0;
-	Setup.MaxValueCount = Setup.bEnabled ? DebugDrawParameters.MaxCharacterCount : 0;
-	Setup.PreViewTranslation = PreViewTranslation;
-
 	FShaderPrintData ShaderPrintData = ShaderPrint::CreateShaderPrintData(GraphBuilder, Setup);
 	
-	// Cache parameters for later GatherDispatchData().
-	ShaderPrint::SetParameters(GraphBuilder, ShaderPrintData, ShaderPrintParameters);
+	// Cache parameters for later GatherDispatchData(). 
+	// We only want the buffers that are created here. 
+	// We don't use the ShaderPrint uniform buffer because we'll use the loose ConfigParameters instead.
+	// That's so we can bind multiple instances of this data interface if we want.
+	ShaderPrint::SetParameters(GraphBuilder, ShaderPrintData, CachedParameters);
 
 	if (ShaderPrint::IsEnabled(ShaderPrintData))
 	{
@@ -153,6 +164,15 @@ void FOptimusDebugDrawDataProviderProxy::GatherDispatchData(FDispatchSetup const
 	{
 		FDebugDrawDataInterfaceParameters* Parameters = (FDebugDrawDataInterfaceParameters*)(InOutDispatchData.ParameterBuffer + InDispatchSetup.ParameterBufferOffset + InDispatchSetup.ParameterBufferStride * InvocationIndex);
 		Parameters->LocalToWorld = LocalToWorld;
-		Parameters->ShaderPrintParameters = ShaderPrintParameters;
+		Parameters->Resolution = ConfigParameters.Resolution;
+		Parameters->FontSize = ConfigParameters.FontSize;
+		Parameters->FontSpacing = ConfigParameters.FontSpacing;
+		Parameters->MaxValueCount = ConfigParameters.MaxValueCount;
+		Parameters->MaxSymbolCount = ConfigParameters.MaxSymbolCount;
+		Parameters->MaxStateCount = ConfigParameters.MaxStateCount;
+		Parameters->MaxLineCount = ConfigParameters.MaxLineCount;
+		Parameters->MaxTriangleCount = ConfigParameters.MaxTriangleCount;
+		Parameters->StateBuffer = CachedParameters.ShaderPrint_StateBuffer;
+		Parameters->RWEntryBuffer = CachedParameters.ShaderPrint_RWEntryBuffer;
 	}
 }
diff --git a/Engine/Plugins/Animation/DeformerGraph/Source/OptimusCore/Private/DataInterfaces/OptimusDataInterfaceDebugDraw.h b/Engine/Plugins/Animation/DeformerGraph/Source/OptimusCore/Private/DataInterfaces/OptimusDataInterfaceDebugDraw.h
--- a/Engine/Plugins/Animation/DeformerGraph/Source/OptimusCore/Private/DataInterfaces/OptimusDataInterfaceDebugDraw.h
+++ b/Engine/Plugins/Animation/DeformerGraph/Source/OptimusCore/Private/DataInterfaces/OptimusDataInterfaceDebugDraw.h
@@ -94,8 +94,8 @@ public:
 	//~ End FComputeDataProviderRenderProxy Interface
 
 private:
-	FVector PreViewTranslation;
 	FMatrix44f LocalToWorld;
-	FOptimusDebugDrawParameters DebugDrawParameters;
-	ShaderPrint::FShaderParameters ShaderPrintParameters;
+	ShaderPrint::FShaderPrintSetup Setup;
+	ShaderPrint::FShaderPrintCommonParameters ConfigParameters;
+	ShaderPrint::FShaderParameters CachedParameters;
 };
diff --git a/Engine/Plugins/Runtime/ComputeFramework/Shaders/Private/ComputeKernel.usf b/Engine/Plugins/Runtime/ComputeFramework/Shaders/Private/ComputeKernel.usf
--- a/Engine/Plugins/Runtime/ComputeFramework/Shaders/Private/ComputeKernel.usf
+++ b/Engine/Plugins/Runtime/ComputeFramework/Shaders/Private/ComputeKernel.usf
@@ -1,10 +1,8 @@
 // Copyright Epic Games, Inc. All Rights Reserved.
 
 //todo[CF]: 
-// Remove these includes to keep the compute framework sandboxed as much as possible
-// Any generated uniform buffer includes need to either be removed or somehow pulled through the data interfaces that need them.
+// Remove includes to keep the compute framework sandboxed as much as possible
 #include "/Engine/Private/Common.ush"
-#include "/Engine/Generated/UniformBuffers/ShaderPrintData.ush"
 
 [numthreads(32, 1, 1)]
 void ComputeKernelNull()
