Change mesh deformer to support other component type as primary source.

diff --git a/Engine/Plugins/Animation/DeformerGraph/Source/OptimusCore/Private/DataInterfaces/OptimusDataInterfaceGraph.cpp b/Engine/Plugins/Animation/DeformerGraph/Source/OptimusCore/Private/DataInterfaces/OptimusDataInterfaceGraph.cpp
--- a/Engine/Plugins/Animation/DeformerGraph/Source/OptimusCore/Private/DataInterfaces/OptimusDataInterfaceGraph.cpp
+++ b/Engine/Plugins/Animation/DeformerGraph/Source/OptimusCore/Private/DataInterfaces/OptimusDataInterfaceGraph.cpp
@@ -9,6 +9,7 @@
 #include "OptimusVariableDescription.h"
 #include "ShaderParameterMetadataBuilder.h"
 #include "ComputeFramework/ComputeMetadataBuilder.h"
+#include "MeshDeformerInterface.h"
 
 void UOptimusGraphDataInterface::Init(TArray<FOptimusGraphVariableDescription> const& InVariables)
 {
@@ -89,9 +90,10 @@ void UOptimusGraphDataInterface::GetHLSL(FString& OutHLSL, FString const& InData
 UComputeDataProvider* UOptimusGraphDataInterface::CreateDataProvider(TObjectPtr<UObject> InBinding, uint64 InInputMask, uint64 InOutputMask) const
 {
 	UOptimusGraphDataProvider* Provider = NewObject<UOptimusGraphDataProvider>();
-	Provider->SkinnedMeshComponent = Cast<USkinnedMeshComponent>(InBinding);
+	Provider->MeshComponent = Cast<UMeshComponent>(InBinding);
 	Provider->Variables = Variables;
 	Provider->ParameterBufferSize = ParameterBufferSize;
+	check(Provider->MeshComponent);
 	return Provider;
 }
 
@@ -113,12 +115,14 @@ void UOptimusGraphDataProvider::SetConstant(FString const& InVariableName, TArra
 
 FComputeDataProviderRenderProxy* UOptimusGraphDataProvider::GetRenderProxy()
 {
-	UOptimusDeformerInstance const* DeformerInstance = Cast<UOptimusDeformerInstance>(SkinnedMeshComponent->GetMeshDeformerInstance());
-
-	return new FOptimusGraphDataProviderProxy(DeformerInstance, Variables, ParameterBufferSize);
+	if (IMeshDeformerInterface* MeshDeformerInterface = Cast<IMeshDeformerInterface>(MeshComponent))
+	{
+		UOptimusDeformerInstance const* DeformerInstance = Cast<UOptimusDeformerInstance>(MeshDeformerInterface->GetMeshDeformerInstance());
+		return new FOptimusGraphDataProviderProxy(DeformerInstance, Variables, ParameterBufferSize);
+	}
+	return nullptr;
 }
 
-
 FOptimusGraphDataProviderProxy::FOptimusGraphDataProviderProxy(UOptimusDeformerInstance const* DeformerInstance, TArray<FOptimusGraphVariableDescription> const& Variables, int32 ParameterBufferSize)
 {
 	// Get all variables from deformer instance and fill buffer.
diff --git a/Engine/Plugins/Animation/DeformerGraph/Source/OptimusCore/Private/DataInterfaces/OptimusDataInterfaceGraph.h b/Engine/Plugins/Animation/DeformerGraph/Source/OptimusCore/Private/DataInterfaces/OptimusDataInterfaceGraph.h
--- a/Engine/Plugins/Animation/DeformerGraph/Source/OptimusCore/Private/DataInterfaces/OptimusDataInterfaceGraph.h
+++ b/Engine/Plugins/Animation/DeformerGraph/Source/OptimusCore/Private/DataInterfaces/OptimusDataInterfaceGraph.h
@@ -8,7 +8,7 @@
 #include "OptimusDataInterfaceGraph.generated.h"
 
 class UOptimusDeformerInstance;
-class USkinnedMeshComponent;
+class UMeshComponent;
 
 /** */
 USTRUCT()
@@ -62,7 +62,7 @@ class UOptimusGraphDataProvider : public UComputeDataProvider
 
 public:
 	UPROPERTY(EditAnywhere, BlueprintReadWrite, Category = Binding)
-	TObjectPtr<USkinnedMeshComponent> SkinnedMeshComponent = nullptr;
+	TObjectPtr<UMeshComponent> MeshComponent = nullptr;
 
 	UPROPERTY()
 	TArray<FOptimusGraphVariableDescription> Variables;
diff --git a/Engine/Plugins/Animation/DeformerGraph/Source/OptimusCore/Private/OptimusDeformerInstance.cpp b/Engine/Plugins/Animation/DeformerGraph/Source/OptimusCore/Private/OptimusDeformerInstance.cpp
--- a/Engine/Plugins/Animation/DeformerGraph/Source/OptimusCore/Private/OptimusDeformerInstance.cpp
+++ b/Engine/Plugins/Animation/DeformerGraph/Source/OptimusCore/Private/OptimusDeformerInstance.cpp
@@ -9,7 +9,6 @@
 #include "OptimusDeformer.h"
 #include "OptimusVariableDescription.h"
 #include "RenderGraphBuilder.h"
-#include "SkeletalRenderPublic.h"
 #include "DataInterfaces/OptimusDataInterfaceRawBuffer.h"
 
 
@@ -333,18 +332,6 @@ void UOptimusDeformerInstance::SetupFromDeformer(UOptimusDeformer* InDeformer)
 			}
 		}
 
-		int32 LODIndex = 0;
-		if (const USkinnedMeshComponent* SkinnedMeshComponent = Cast<USkinnedMeshComponent>(MeshComponent))
-		{
-			// This guff should be a utility function on USkinnedMeshComponent. 
-			LODIndex = SkinnedMeshComponent->GetPredictedLODLevel();
-			
-			if (SkinnedMeshComponent->GetSkinnedAsset() && SkinnedMeshComponent->GetSkinnedAsset()->IsStreamable() && SkinnedMeshComponent->MeshObject)
-			{
-				LODIndex = FMath::Max<int32>(LODIndex, SkinnedMeshComponent->MeshObject->GetSkeletalMeshRenderData().PendingFirstLODIdx);
-			}
-		}
-
 		for(TObjectPtr<UComputeDataProvider> DataProvider: Info.ComputeGraphInstance.GetDataProviders())
 		{
 			// Make the persistent buffer data provider aware of the buffer pool and current LOD index.
diff --git a/Engine/Plugins/Animation/DeformerGraph/Source/OptimusCore/Public/OptimusComponentSource.h b/Engine/Plugins/Animation/DeformerGraph/Source/OptimusCore/Public/OptimusComponentSource.h
--- a/Engine/Plugins/Animation/DeformerGraph/Source/OptimusCore/Public/OptimusComponentSource.h
+++ b/Engine/Plugins/Animation/DeformerGraph/Source/OptimusCore/Public/OptimusComponentSource.h
@@ -57,7 +57,7 @@ public:
 	}
 	
 	/** Returns true if the source can be used by primary bindings. */
-	bool IsUsableAsPrimarySource() const;
+	virtual bool IsUsableAsPrimarySource() const;
 
 	// TODO: Component color for additional indicator wire.
 
diff --git a/Engine/Plugins/Animation/DeformerGraph/Source/OptimusCore/Public/OptimusComputeDataInterface.h b/Engine/Plugins/Animation/DeformerGraph/Source/OptimusCore/Public/OptimusComputeDataInterface.h
--- a/Engine/Plugins/Animation/DeformerGraph/Source/OptimusCore/Public/OptimusComputeDataInterface.h
+++ b/Engine/Plugins/Animation/DeformerGraph/Source/OptimusCore/Public/OptimusComputeDataInterface.h
@@ -119,7 +119,7 @@ class OPTIMUSCORE_API UOptimusComputeDataInterface : public UComputeDataInterfac
 	GENERATED_BODY()
 	
 public:
-	struct CategoryName
+	struct OPTIMUSCORE_API CategoryName
 	{
 		static const FName DataInterfaces;
 		static const FName ExecutionDataInterfaces;
diff --git a/Engine/Source/Runtime/Engine/Classes/Components/SkinnedMeshComponent.h b/Engine/Source/Runtime/Engine/Classes/Components/SkinnedMeshComponent.h
--- a/Engine/Source/Runtime/Engine/Classes/Components/SkinnedMeshComponent.h
+++ b/Engine/Source/Runtime/Engine/Classes/Components/SkinnedMeshComponent.h
@@ -15,6 +15,7 @@
 #include "Components/MeshComponent.h"
 #include "Containers/SortedMap.h"
 #include "LODSyncInterface.h"
+#include "MeshDeformerInterface.h"
 #include "BoneContainer.h"
 #include "Rendering/MorphTargetVertexInfoBuffers.h"
 #include "ClothingSystemRuntimeTypes.h"
@@ -238,7 +239,7 @@ struct FVertexOffsetUsage
  * @see USkeletalMeshComponent
 */
 UCLASS(hidecategories=Object, config=Engine, editinlinenew, abstract)
-class ENGINE_API USkinnedMeshComponent : public UMeshComponent, public ILODSyncInterface
+class ENGINE_API USkinnedMeshComponent : public UMeshComponent, public ILODSyncInterface, public IMeshDeformerInterface
 {
 	GENERATED_UCLASS_BODY()
 
@@ -302,7 +303,7 @@ protected:
 
 public:
 	/** Get the currently active MeshDeformer Instance. */
-	UMeshDeformerInstance const* GetMeshDeformerInstance() const { return MeshDeformerInstance; }
+	virtual UMeshDeformerInstance const* GetMeshDeformerInstance() const override { return MeshDeformerInstance; }
 
 	/** const getters for previous transform idea */
 	const TArray<uint8>& GetPreviousBoneVisibilityStates() const
diff --git a/Engine/Source/Runtime/Engine/Private/Engine/MeshDeformerInterface.cpp b/Engine/Source/Runtime/Engine/Private/Engine/MeshDeformerInterface.cpp
new file mode 100644
--- /dev/null
+++ b/Engine/Source/Runtime/Engine/Private/Engine/MeshDeformerInterface.cpp
@@ -0,0 +1,11 @@
+// Copyright Epic Games, Inc. All Rights Reserved.
+
+#include "MeshDeformerInterface.h"
+
+#include UE_INLINE_GENERATED_CPP_BY_NAME(MeshDeformerInterface)
+
+UMeshDeformerInterface::UMeshDeformerInterface(const FObjectInitializer& ObjectInitializer)
+	: Super(ObjectInitializer)
+{
+}
+
diff --git a/Engine/Source/Runtime/Engine/Public/MeshDeformerInterface.h b/Engine/Source/Runtime/Engine/Public/MeshDeformerInterface.h
new file mode 100644
--- /dev/null
+++ b/Engine/Source/Runtime/Engine/Public/MeshDeformerInterface.h
@@ -0,0 +1,31 @@
+// Copyright Epic Games, Inc. All Rights Reserved.
+
+#pragma once
+
+#include "CoreMinimal.h"
+#include "UObject/ObjectMacros.h"
+#include "UObject/Object.h"
+#include "UObject/Interface.h"
+#include "MeshDeformerInterface.generated.h"
+
+class UMeshDeformerInstance;
+
+UINTERFACE(meta = (CannotImplementInterfaceInBlueprint))
+class ENGINE_API UMeshDeformerInterface : public UInterface
+{
+	GENERATED_UINTERFACE_BODY()
+};
+
+/* This is interface class for getting/setting deformer instance
+ *
+ * Implemented by SkeletalMeshComponent, GroomComponent
+ */
+
+class ENGINE_API IMeshDeformerInterface
+{
+	GENERATED_IINTERFACE_BODY()
+
+	/** Get the currently active MeshDeformer Instance. */
+	virtual UMeshDeformerInstance const* GetMeshDeformerInstance() const = 0;
+};
+
