Add ComponentToSim transform to JiggleSpring deformer, and 'amount' property that ramps in/out the effect of the deformer, post sim.  Add the ability to pass a FTransform to an optimus deformer via Blueprints.

diff --git a/Engine/Plugins/Animation/DeformerGraph/Source/OptimusCore/Private/OptimusDeformerInstance.cpp b/Engine/Plugins/Animation/DeformerGraph/Source/OptimusCore/Private/OptimusDeformerInstance.cpp
--- a/Engine/Plugins/Animation/DeformerGraph/Source/OptimusCore/Private/OptimusDeformerInstance.cpp
+++ b/Engine/Plugins/Animation/DeformerGraph/Source/OptimusCore/Private/OptimusDeformerInstance.cpp
@@ -500,6 +500,11 @@ bool UOptimusDeformerInstance::SetVector4Variable(FName InVariableName, const FV
 	return SetVariableValue<FVector4>(Variables, InVariableName, "FVector4", InValue);
 }
 
+bool UOptimusDeformerInstance::SetTransformVariable(FName InVariableName, const FTransform& InValue)
+{
+	return SetVariableValue<FTransform>(Variables, InVariableName, "FTransform", InValue);
+}
+
 const TArray<UOptimusVariableDescription*>& UOptimusDeformerInstance::GetVariables() const
 {
 	return Variables->Descriptions;
diff --git a/Engine/Plugins/Animation/DeformerGraph/Source/OptimusCore/Public/OptimusDeformerInstance.h b/Engine/Plugins/Animation/DeformerGraph/Source/OptimusCore/Public/OptimusDeformerInstance.h
--- a/Engine/Plugins/Animation/DeformerGraph/Source/OptimusCore/Public/OptimusDeformerInstance.h
+++ b/Engine/Plugins/Animation/DeformerGraph/Source/OptimusCore/Public/OptimusDeformerInstance.h
@@ -182,6 +182,10 @@ public:
 	UFUNCTION(BlueprintCallable, Category="Deformer", meta=(DisplayName="Set Variable (Vector4)"))
 	bool SetVector4Variable(FName InVariableName, const FVector4& InValue);
 
+	/** Set the value of a transform variable. */
+	UFUNCTION(BlueprintCallable, Category = "Deformer", meta = (DisplayName = "Set Variable (Transform)"))
+	bool SetTransformVariable(FName InVariableName, const FTransform& InValue);
+
 	/** Get an array containing all the variables. */
 	UFUNCTION(BlueprintGetter)
 	const TArray<UOptimusVariableDescription*>& GetVariables() const;
