Optimus: Fix shader compilation warnings for non initialized structure.

diff --git a/Engine/Plugins/Animation/DeformerGraph/Shaders/Private/DataInterfaceDebugDraw.ush b/Engine/Plugins/Animation/DeformerGraph/Shaders/Private/DataInterfaceDebugDraw.ush
--- a/Engine/Plugins/Animation/DeformerGraph/Shaders/Private/DataInterfaceDebugDraw.ush
+++ b/Engine/Plugins/Animation/DeformerGraph/Shaders/Private/DataInterfaceDebugDraw.ush
@@ -63,6 +63,8 @@ FDebugDraw ReadDebugDraw_{DataInterfaceName}()
 	DebugDraw.Context.StartPos = 0;
 	DebugDraw.Context.Pos = 0;
 	DebugDraw.Context.Config.Resolution = {DataInterfaceName}_Resolution;
+	DebugDraw.Context.CursorCoord = 0;
+	DebugDraw.Context.TranslatedWorldOffset = 0;
 	DebugDraw.Context.Config.FontSize = {DataInterfaceName}_FontSize;
 	DebugDraw.Context.Config.FontSpacing = {DataInterfaceName}_FontSpacing;
 	DebugDraw.Context.Config.MaxValueCount = {DataInterfaceName}_MaxValueCount;
