Optimus: Use a default deformer settings when none has been created.
Tidy up UI names for deformer settings.

diff --git a/Engine/Plugins/Animation/DeformerGraph/Source/OptimusCore/Private/OptimusDeformerInstance.cpp b/Engine/Plugins/Animation/DeformerGraph/Source/OptimusCore/Private/OptimusDeformerInstance.cpp
--- a/Engine/Plugins/Animation/DeformerGraph/Source/OptimusCore/Private/OptimusDeformerInstance.cpp
+++ b/Engine/Plugins/Animation/DeformerGraph/Source/OptimusCore/Private/OptimusDeformerInstance.cpp
@@ -295,10 +295,13 @@ void UOptimusDeformerInstance::SetupFromDeformer(UOptimusDeformer* InDeformer)
 	// The bindings are in the same order as the component bindings in the deformer.
 	TArray<UActorComponent*> BoundComponents;
 	UOptimusDeformerInstanceSettings* InstanceSettingsPtr = InstanceSettings.Get(); 
-	if (InstanceSettingsPtr)
+	if (InstanceSettingsPtr == nullptr)
 	{
-		InstanceSettingsPtr->GetComponentBindings(InDeformer, MeshComponent.Get(), BoundComponents);
+		// If we don't have any settings, then create a temporary object to get bindings.
+		InstanceSettingsPtr = NewObject<UOptimusDeformerInstanceSettings>();
+		InstanceSettingsPtr->InitializeSettings(InDeformer, MeshComponent.Get());
 	}
+	InstanceSettingsPtr->GetComponentBindings(InDeformer, MeshComponent.Get(), BoundComponents);
 
 	// Create the persistent buffer pool
 	BufferPool = MakeShared<FOptimusPersistentBufferPool>();
diff --git a/Engine/Plugins/Animation/DeformerGraph/Source/OptimusCore/Public/OptimusDeformerInstance.h b/Engine/Plugins/Animation/DeformerGraph/Source/OptimusCore/Public/OptimusDeformerInstance.h
--- a/Engine/Plugins/Animation/DeformerGraph/Source/OptimusCore/Public/OptimusDeformerInstance.h
+++ b/Engine/Plugins/Animation/DeformerGraph/Source/OptimusCore/Public/OptimusDeformerInstance.h
@@ -110,7 +110,7 @@ class OPTIMUSCORE_API UOptimusDeformerInstanceSettings :
 	TWeakObjectPtr<UOptimusDeformer> Deformer;
 
 	/** Array of binding descriptions. This is fixed and used by GetComponentBindings() to resolve final bindings for a given context. */
-	UPROPERTY(EditAnywhere, Category="Deformer", EditFixedSize, meta = (DisplayName = "Deformer Bindings", EditFixedOrder))
+	UPROPERTY(EditAnywhere, Category="Deformer|DeformerSettings", EditFixedSize, meta = (DisplayName = "Component Bindings", EditFixedOrder))
 	TArray<FOptimusDeformerInstanceComponentBinding> Bindings;
 
 public:
diff --git a/Engine/Source/Runtime/Engine/Classes/Animation/MeshDeformerInstance.h b/Engine/Source/Runtime/Engine/Classes/Animation/MeshDeformerInstance.h
--- a/Engine/Source/Runtime/Engine/Classes/Animation/MeshDeformerInstance.h
+++ b/Engine/Source/Runtime/Engine/Classes/Animation/MeshDeformerInstance.h
@@ -12,7 +12,7 @@ class FSceneInterface;
  * Base class for mesh deformers instance settings.
  * This contains the serialized user settings to apply to the UMeshDeformer.
  */
-UCLASS(Abstract, DefaultToInstanced)
+UCLASS(Abstract)
 class ENGINE_API UMeshDeformerInstanceSettings : public UObject
 {
 	GENERATED_BODY()
diff --git a/Engine/Source/Runtime/Engine/Classes/Components/SkinnedMeshComponent.h b/Engine/Source/Runtime/Engine/Classes/Components/SkinnedMeshComponent.h
--- a/Engine/Source/Runtime/Engine/Classes/Components/SkinnedMeshComponent.h
+++ b/Engine/Source/Runtime/Engine/Classes/Components/SkinnedMeshComponent.h
@@ -294,7 +294,7 @@ protected:
 	UMeshDeformer* GetActiveMeshDeformer() const;
 
 	/** Object containing instance settings for the bound MeshDeformer. */
-	UPROPERTY(VisibleAnywhere, BlueprintReadOnly, Category = "Deformer", meta = (DisplayName = "Settings", EditCondition = "MeshDeformerInstanceSettings!=nullptr", EditConditionHides, ShowOnlyInnerProperties))
+	UPROPERTY(VisibleAnywhere, BlueprintReadOnly, Instanced, Category = "Deformer", meta = (DisplayName = "Deformer Settings", EditCondition = "MeshDeformerInstanceSettings!=nullptr", ShowOnlyInnerProperties))
 	TObjectPtr<UMeshDeformerInstanceSettings> MeshDeformerInstanceSettings;
 
 	/** Object containing state for the bound MeshDeformer. */
