Fix missing printf format identifier.

diff --git a/Engine/Plugins/Animation/DeformerGraph/Source/OptimusCore/Private/Nodes/OptimusNode_ComputeKernelBase.cpp b/Engine/Plugins/Animation/DeformerGraph/Source/OptimusCore/Private/Nodes/OptimusNode_ComputeKernelBase.cpp
--- a/Engine/Plugins/Animation/DeformerGraph/Source/OptimusCore/Private/Nodes/OptimusNode_ComputeKernelBase.cpp
+++ b/Engine/Plugins/Animation/DeformerGraph/Source/OptimusCore/Private/Nodes/OptimusNode_ComputeKernelBase.cpp
@@ -94,7 +94,7 @@ FOptimus_ComputeKernelResult UOptimusNode_ComputeKernelBase::CreateComputeKernel
 						// The same friendly name cannot be claimed by two unique types;
 						if (*FoundUniqueName != UniqueName)
 						{
-							return FText::Format(LOCTEXT("InvalidPinFriendlyName", "Invalid unique friendly name on pin '{}'"), FText::FromName(InPin->GetUniqueName()));
+							return FText::Format(LOCTEXT("InvalidPinFriendlyName", "Invalid unique friendly name on pin '{0}'"), FText::FromName(InPin->GetUniqueName()));
 						}
 
 						// Type is already in the map, no more actions needed
