Optimus: Fix copy/paste/duplicate to do a deep copy.

#rnx

diff --git a/Engine/Plugins/Animation/DeformerGraph/Source/OptimusCore/Private/Actions/OptimusNodeGraphActions.cpp b/Engine/Plugins/Animation/DeformerGraph/Source/OptimusCore/Private/Actions/OptimusNodeGraphActions.cpp
--- a/Engine/Plugins/Animation/DeformerGraph/Source/OptimusCore/Private/Actions/OptimusNodeGraphActions.cpp
+++ b/Engine/Plugins/Animation/DeformerGraph/Source/OptimusCore/Private/Actions/OptimusNodeGraphActions.cpp
@@ -328,13 +328,8 @@ FOptimusNodeGraphAction_DuplicateNode::FOptimusNodeGraphAction_DuplicateNode(
 		NodeClassPath = InSourceNode->GetClass()->GetPathName();
 		ConfigureNodeFunc = InConfigureNodeFunc;
 		
-		// Photocopy any non-transient property data
-		{
-			FMemoryWriter NodeArchive(NodeData);
-			FObjectAndNameAsStringProxyArchive NodeProxyArchive(
-					NodeArchive, /* bInLoadIfFindFails=*/ false);
-			InSourceNode->SerializeScriptProperties(NodeProxyArchive);
-		}
+		// Photocopy the node.
+		Optimus::FBinaryObjectWriter(InSourceNode, NodeData);
 	}
 }
 
@@ -353,12 +348,8 @@ bool FOptimusNodeGraphAction_DuplicateNode::Do(
 {
 	auto BootstrapNodeFunc = [this](UOptimusNode* InNode) -> bool
 	{
-		{
-			FMemoryReader NodeArchive(NodeData);
-			FObjectAndNameAsStringProxyArchive NodeProxyArchive(
-					NodeArchive, /* bInLoadIfFindFails=*/true);
-			InNode->SerializeScriptProperties(NodeProxyArchive);
-		}
+		// Paste the data from the old node onto the new one.
+		Optimus::FBinaryObjectReader(InNode, NodeData);
 		if (!ConfigureNodeFunc)
 		{
 			return true;
@@ -505,7 +496,8 @@ bool FOptimusNodeGraphAction_RemoveNode::Undo(
 
 FOptimusNodeGraphAction_AddRemoveLink::FOptimusNodeGraphAction_AddRemoveLink(
 	UOptimusNodePin* InNodeOutputPin, 
-	UOptimusNodePin* InNodeInputPin
+	UOptimusNodePin* InNodeInputPin,
+	bool bInCanFail
 	)
 {
 	if (ensure(InNodeOutputPin != nullptr) && ensure(InNodeInputPin != nullptr) &&
@@ -517,19 +509,22 @@ FOptimusNodeGraphAction_AddRemoveLink::FOptimusNodeGraphAction_AddRemoveLink(
 	{
 		NodeOutputPinPath = InNodeOutputPin->GetPinPath();
 		NodeInputPinPath = InNodeInputPin->GetPinPath();
+		bCanFail = bInCanFail;
 	}
 }
 
 
 FOptimusNodeGraphAction_AddRemoveLink::FOptimusNodeGraphAction_AddRemoveLink(
 	const FString& InNodeOutputPinPath,
-	const FString& InNodeInputPinPath
+	const FString& InNodeInputPinPath,
+	bool bInCanFail
 	)
 {
 	if (ensure(!InNodeOutputPinPath.IsEmpty()) && ensure(!InNodeInputPinPath.IsEmpty()))
 	{
 		NodeOutputPinPath = InNodeOutputPinPath;
 		NodeInputPinPath = InNodeInputPinPath;
+		bCanFail = bInCanFail;
 	}
 }
 
@@ -541,13 +536,13 @@ bool FOptimusNodeGraphAction_AddRemoveLink::AddLink(
 	UOptimusNodePin* InNodeOutputPin = InRoot->ResolvePinPath(NodeOutputPinPath);
 	if (InNodeOutputPin == nullptr)
 	{
-		return false;
+		return bCanFail;
 	}
 
 	UOptimusNodePin* InNodeInputPin = InRoot->ResolvePinPath(NodeInputPinPath);
 	if (InNodeInputPin == nullptr)
 	{
-		return false;
+		return bCanFail;
 	}
 
 	UOptimusNodeGraph* Graph = InNodeOutputPin->GetOwningNode()->GetOwningGraph();
@@ -562,13 +557,13 @@ bool FOptimusNodeGraphAction_AddRemoveLink::RemoveLink(
 	UOptimusNodePin* InNodeOutputPin = InRoot->ResolvePinPath(NodeOutputPinPath);
 	if (InNodeOutputPin == nullptr)
 	{
-		return false;
+		return bCanFail;
 	}
 
 	UOptimusNodePin* InNodeInputPin = InRoot->ResolvePinPath(NodeInputPinPath);
 	if (InNodeInputPin == nullptr)
 	{
-		return false;
+		return bCanFail;
 	}
 
 	UOptimusNodeGraph* Graph = InNodeOutputPin->GetOwningNode()->GetOwningGraph();
@@ -580,9 +575,10 @@ bool FOptimusNodeGraphAction_AddRemoveLink::RemoveLink(
 
 FOptimusNodeGraphAction_AddLink::FOptimusNodeGraphAction_AddLink(
 	UOptimusNodePin* InNodeOutputPin, 
-	UOptimusNodePin* InNodeInputPin
+	UOptimusNodePin* InNodeInputPin,
+	bool bInCanFail
 	) :
-	FOptimusNodeGraphAction_AddRemoveLink(InNodeOutputPin, InNodeInputPin)
+	FOptimusNodeGraphAction_AddRemoveLink(InNodeOutputPin, InNodeInputPin, bInCanFail)
 {
 	// FIXME: Prettier name.
 	SetTitlef(TEXT("Add Link"));
@@ -591,9 +587,10 @@ FOptimusNodeGraphAction_AddLink::FOptimusNodeGraphAction_AddLink(
 
 FOptimusNodeGraphAction_AddLink::FOptimusNodeGraphAction_AddLink(
 	const FString& InNodeOutputPinPath,
-	const FString& InNodeInputPinPath
+	const FString& InNodeInputPinPath,
+	bool bInCanFail
 	) :
-	FOptimusNodeGraphAction_AddRemoveLink(InNodeOutputPinPath, InNodeInputPinPath)
+	FOptimusNodeGraphAction_AddRemoveLink(InNodeOutputPinPath, InNodeInputPinPath, bInCanFail)
 {
 	SetTitlef(TEXT("Add Link"));
 }
diff --git a/Engine/Plugins/Animation/DeformerGraph/Source/OptimusCore/Private/Actions/OptimusNodeGraphActions.h b/Engine/Plugins/Animation/DeformerGraph/Source/OptimusCore/Private/Actions/OptimusNodeGraphActions.h
--- a/Engine/Plugins/Animation/DeformerGraph/Source/OptimusCore/Private/Actions/OptimusNodeGraphActions.h
+++ b/Engine/Plugins/Animation/DeformerGraph/Source/OptimusCore/Private/Actions/OptimusNodeGraphActions.h
@@ -269,12 +269,14 @@ struct FOptimusNodeGraphAction_AddRemoveLink :
 
 	FOptimusNodeGraphAction_AddRemoveLink(
 		UOptimusNodePin* InNodeOutputPin, 
-		UOptimusNodePin* InNodeInputPin
+		UOptimusNodePin* InNodeInputPin,
+		bool bInCanFail = false
 		);
 
 	FOptimusNodeGraphAction_AddRemoveLink(
 		const FString& InNodeOutputPinPath,
-		const FString& InNodeInputPinPath
+		const FString& InNodeInputPinPath,
+		bool bInCanFail = false
 		);
 
 protected:
@@ -286,6 +288,10 @@ protected:
 
 	// The path of the output input on the node to connect/disconnect to/from.
 	FString NodeInputPinPath;
+
+	// The operation is allowed to fail gracefully without terminating the
+	// rest of the actions being performed.
+	bool bCanFail = false;
 };
 
 // Mark FOptimusNodeGraphAction_AddRemoveLink as pure virtual, so that the UObject machinery
@@ -312,12 +318,14 @@ struct FOptimusNodeGraphAction_AddLink :
 
 	FOptimusNodeGraphAction_AddLink(
 		UOptimusNodePin* InNodeOutputPin,
-		UOptimusNodePin* InNodeInputPin
+		UOptimusNodePin* InNodeInputPin,
+		bool bInCanFail = false
 	);
 
 	FOptimusNodeGraphAction_AddLink(
 		const FString& InNodeOutputPinPath,
-		const FString& InNodeInputPinPath
+		const FString& InNodeInputPinPath,
+		bool bInCanFail = false
 		);
 
 protected:
diff --git a/Engine/Plugins/Animation/DeformerGraph/Source/OptimusCore/Private/Actions/OptimusResourceActions.cpp b/Engine/Plugins/Animation/DeformerGraph/Source/OptimusCore/Private/Actions/OptimusResourceActions.cpp
--- a/Engine/Plugins/Animation/DeformerGraph/Source/OptimusCore/Private/Actions/OptimusResourceActions.cpp
+++ b/Engine/Plugins/Animation/DeformerGraph/Source/OptimusCore/Private/Actions/OptimusResourceActions.cpp
@@ -11,11 +11,13 @@
 
 FOptimusResourceAction_AddResource::FOptimusResourceAction_AddResource(
 	FOptimusDataTypeRef InDataType,
-    FName InName
+    FName InName,
+    FOptimusDataDomain InDataDomain
 	)
 {
 	ResourceName = InName;
 	DataType = InDataType;
+	DataDomain = InDataDomain;
 
 	SetTitlef(TEXT("Add resource '%s'"), *ResourceName.ToString());
 }
@@ -46,6 +48,7 @@ bool FOptimusResourceAction_AddResource::Do(
 
 	Resource->ResourceName = Resource->GetFName();
 	Resource->DataType = DataType;
+	Resource->DataDomain = DataDomain;
 
 	if (!Deformer->AddResourceDirect(Resource))
 	{
diff --git a/Engine/Plugins/Animation/DeformerGraph/Source/OptimusCore/Private/Actions/OptimusResourceActions.h b/Engine/Plugins/Animation/DeformerGraph/Source/OptimusCore/Private/Actions/OptimusResourceActions.h
--- a/Engine/Plugins/Animation/DeformerGraph/Source/OptimusCore/Private/Actions/OptimusResourceActions.h
+++ b/Engine/Plugins/Animation/DeformerGraph/Source/OptimusCore/Private/Actions/OptimusResourceActions.h
@@ -25,7 +25,8 @@ public:
 
 	FOptimusResourceAction_AddResource(
 		FOptimusDataTypeRef InDataType,
-		FName InName
+		FName InName,
+		FOptimusDataDomain InDataDomain = FOptimusDataDomain()
 	);
 
 	UOptimusResourceDescription* GetResource(IOptimusPathResolver* InRoot) const;
@@ -40,6 +41,9 @@ private:
 
 	// The data type of the resource
 	FOptimusDataTypeRef DataType;
+
+	// The desired data domain of the resource
+	FOptimusDataDomain DataDomain;
 };
 
 
diff --git a/Engine/Plugins/Animation/DeformerGraph/Source/OptimusCore/Private/Nodes/OptimusNode_ConstantValue.cpp b/Engine/Plugins/Animation/DeformerGraph/Source/OptimusCore/Private/Nodes/OptimusNode_ConstantValue.cpp
--- a/Engine/Plugins/Animation/DeformerGraph/Source/OptimusCore/Private/Nodes/OptimusNode_ConstantValue.cpp
+++ b/Engine/Plugins/Animation/DeformerGraph/Source/OptimusCore/Private/Nodes/OptimusNode_ConstantValue.cpp
@@ -2,6 +2,7 @@
 
 #include "Nodes/OptimusNode_ConstantValue.h"
 
+#include "OptimusDataTypeRegistry.h"
 #include "OptimusNodePin.h"
 #include "OptimusNodeGraph.h"
 
@@ -67,6 +68,7 @@ UClass* UOptimusNode_ConstantValueGeneratorClass::GetClassForType(UPackage* InPa
 	return TypeClass;
 }
 
+
 void UOptimusNode_ConstantValue::PostLoad()
 {
 	Super::PostLoad();
@@ -105,7 +107,7 @@ FString UOptimusNode_ConstantValue::GetValueName() const
 
 FOptimusDataTypeRef UOptimusNode_ConstantValue::GetValueType() const
 {
-	UOptimusNode_ConstantValueGeneratorClass* Class = Cast<UOptimusNode_ConstantValueGeneratorClass>(GetClass());
+	const UOptimusNode_ConstantValueGeneratorClass* Class = GetGeneratorClass();
 	if (ensure(Class))
 	{
 		return Class->DataType;
@@ -137,9 +139,52 @@ FShaderValueType::FValue UOptimusNode_ConstantValue::GetShaderValue() const
 }
 
 
+FTopLevelAssetPath UOptimusNode_ConstantValue::GetAssetPathForClassDefiner() const
+{
+	return StaticClass()->GetClassPathName();
+}
+
+
+FString UOptimusNode_ConstantValue::GetClassCreationString() const
+{
+	if (GetGeneratorClass())
+	{
+		return FString::Printf(TEXT("DataType=%s"), *GetGeneratorClass()->DataType->TypeName.ToString()); 
+	}
+	return {};
+}
+
+
+UClass* UOptimusNode_ConstantValue::GetClassFromCreationString(
+	UPackage* InPackage,
+	const TCHAR* InCreationString
+	) const
+{
+	FName DataTypeName;
+	if (!FParse::Value(InCreationString, TEXT("DataType="), DataTypeName))
+	{
+		return nullptr;
+	}
+
+	FOptimusDataTypeHandle FoundDataType = FOptimusDataTypeRegistry::Get().FindType(DataTypeName);
+	if (!FoundDataType.IsValid())
+	{
+		return nullptr;
+	}
+
+	return UOptimusNode_ConstantValueGeneratorClass::GetClassForType(InPackage, FoundDataType);
+}
+
+
 void UOptimusNode_ConstantValue::ConstructNode()
 {
 	SetDisplayName(FText::Format(FText::FromString(TEXT("{0} Constant")), GetValueType()->DisplayName));
 
 	UOptimusNode::ConstructNode();
 }
+
+
+UOptimusNode_ConstantValueGeneratorClass* UOptimusNode_ConstantValue::GetGeneratorClass() const
+{
+	return Cast<UOptimusNode_ConstantValueGeneratorClass>(GetClass());	
+}
diff --git a/Engine/Plugins/Animation/DeformerGraph/Source/OptimusCore/Private/Nodes/OptimusNode_ConstantValue.h b/Engine/Plugins/Animation/DeformerGraph/Source/OptimusCore/Private/Nodes/OptimusNode_ConstantValue.h
--- a/Engine/Plugins/Animation/DeformerGraph/Source/OptimusCore/Private/Nodes/OptimusNode_ConstantValue.h
+++ b/Engine/Plugins/Animation/DeformerGraph/Source/OptimusCore/Private/Nodes/OptimusNode_ConstantValue.h
@@ -2,9 +2,10 @@
 
 #pragma once
 
+#include "IOptimusGeneratedClassDefiner.h"
+#include "IOptimusValueProvider.h"
 #include "OptimusNode.h"
 #include "OptimusDataType.h"
-#include "IOptimusValueProvider.h"
 
 #include "OptimusNode_ConstantValue.generated.h"
 
@@ -41,7 +42,8 @@ public:
 UCLASS(Hidden)
 class UOptimusNode_ConstantValue :
 	public UOptimusNode,
-	public IOptimusValueProvider
+	public IOptimusValueProvider,
+	public IOptimusGeneratedClassDefiner
 {
 	GENERATED_BODY()
 
@@ -65,6 +67,17 @@ public:
 	FOptimusDataTypeRef GetValueType() const override;
 	FShaderValueType::FValue GetShaderValue() const override;
 
+	// IOptimusGeneratedClassDefiner implementation
+	FTopLevelAssetPath GetAssetPathForClassDefiner() const override;
+	FString GetClassCreationString() const override;
+	UClass* GetClassFromCreationString(
+		UPackage* InPackage,
+		const TCHAR* InCreationString
+		) const override;
+	
+	
 protected:
 	void ConstructNode() override;
+	
+	UOptimusNode_ConstantValueGeneratorClass* GetGeneratorClass() const;
 };
diff --git a/Engine/Plugins/Animation/DeformerGraph/Source/OptimusCore/Private/Nodes/OptimusNode_DataInterface.cpp b/Engine/Plugins/Animation/DeformerGraph/Source/OptimusCore/Private/Nodes/OptimusNode_DataInterface.cpp
--- a/Engine/Plugins/Animation/DeformerGraph/Source/OptimusCore/Private/Nodes/OptimusNode_DataInterface.cpp
+++ b/Engine/Plugins/Animation/DeformerGraph/Source/OptimusCore/Private/Nodes/OptimusNode_DataInterface.cpp
@@ -200,9 +200,12 @@ void UOptimusNode_DataInterface::PostLoad()
 void UOptimusNode_DataInterface::ConstructNode()
 {
 	// Create the component pin.
-	if (ensure(DataInterfaceClass) &&
-		ensure(DataInterfaceData))
+	if (ensure(DataInterfaceClass))
 	{
+		if (!DataInterfaceData)
+		{
+			DataInterfaceData = NewObject<UOptimusComputeDataInterface>(this, DataInterfaceClass);
+		}
 		SetDisplayName(FText::FromString(DataInterfaceData->GetDisplayName()));
 		CreateComponentPin();
 		CreatePinsFromDataInterface(DataInterfaceData);
diff --git a/Engine/Plugins/Animation/DeformerGraph/Source/OptimusCore/Private/Nodes/OptimusNode_GetVariable.cpp b/Engine/Plugins/Animation/DeformerGraph/Source/OptimusCore/Private/Nodes/OptimusNode_GetVariable.cpp
--- a/Engine/Plugins/Animation/DeformerGraph/Source/OptimusCore/Private/Nodes/OptimusNode_GetVariable.cpp
+++ b/Engine/Plugins/Animation/DeformerGraph/Source/OptimusCore/Private/Nodes/OptimusNode_GetVariable.cpp
@@ -3,44 +3,160 @@
 #include "Nodes/OptimusNode_GetVariable.h"
 
 #include "OptimusCoreModule.h"
+#include "OptimusDataTypeRegistry.h"
 #include "OptimusDeformer.h"
 #include "OptimusNodePin.h"
+#include "OptimusValueContainer.h"
 #include "OptimusVariableDescription.h"
+#include "Actions/OptimusAction.h"
+#include "Actions/OptimusVariableActions.h"
 
 #define LOCTEXT_NAMESPACE "OptimusGetVariable"
 
-void UOptimusNode_GetVariable::PostDuplicate(EDuplicateMode::Type DuplicateMode)
-{
-	Super::PostDuplicate(DuplicateMode);
 
-	if (!VariableDesc.IsValid())
+void UOptimusNode_GetVariable::PreDuplicateRequirementActions(
+	const UOptimusNodeGraph* InTargetGraph,
+	FOptimusCompoundAction* InCompoundAction
+	)
+{
+	if (DuplicationInfo.VariableName.IsNone())
 	{
 		return;
 	}
+
+	// Check if the deformer has a variable that matches the name and type.
+	UOptimusDeformer* Deformer = Cast<UOptimusDeformer>(InTargetGraph->GetCollectionRoot());
 	
+	for (const UOptimusVariableDescription* ExistingVariableDesc: Deformer->GetVariables())
+	{
+		if (ExistingVariableDesc->VariableName == DuplicationInfo.VariableName &&
+			ExistingVariableDesc->DataType == DuplicationInfo.DataType)
+		{
+			// Note: We don't care about the default value in this context.
+			return;
+		}
+	}
+
+	// Note: This will fail if there are multiple overlapping names. A better system is needed.
+	bool bFoundDuplicate = false;
+	do
+	{
+		for (const UOptimusVariableDescription* ExistingVariableDesc: Deformer->GetVariables())
+		{
+			if (ExistingVariableDesc->VariableName == DuplicationInfo.VariableName)
+			{
+				DuplicationInfo.VariableName.SetNumber(DuplicationInfo.VariableName.GetNumber() + 1);
+				bFoundDuplicate = true;
+			}
+		}
+	}
+	while(bFoundDuplicate);
+
+	InCompoundAction->AddSubAction<FOptimusVariableAction_AddVariable>(DuplicationInfo.DataType, DuplicationInfo.VariableName);
+}
+
+
+void UOptimusNode_GetVariable::PostDuplicate(EDuplicateMode::Type DuplicateMode)
+{
+	Super::PostDuplicate(DuplicateMode);
+
 	if (!GetOwningGraph())
 	{
 		return;
 	}
-	
-	const UOptimusDeformer* OldDescOwner = VariableDesc->GetOwningDeformer();
+		
 	const UOptimusDeformer* NewDescOwner = Cast<UOptimusDeformer>(GetOwningGraph()->GetCollectionRoot());
-
 	if (!NewDescOwner)
 	{
 		return;
 	}
-	
-	// No action needed if we are copying/pasting within the same deformer asset 
-	if (OldDescOwner == NewDescOwner)
+
+	if (VariableDesc.IsValid())
 	{
-		return;
+		const UOptimusDeformer* OldDescOwner = VariableDesc->GetOwningDeformer();
+		
+		// No action needed if we are copying/pasting within the same deformer asset 
+		if (OldDescOwner == NewDescOwner)
+		{
+			return;
+		}
+		
+		// Refresh the variable description so that we don't hold a reference to a VariableDesc in another deformer asset
+		VariableDesc = NewDescOwner->ResolveVariable(VariableDesc->GetFName());
 	}
+	else if (!DuplicationInfo.VariableName.IsNone())
+	{
+		// Resolve the variable description from the duplication information that likely came from a copy/paste operation
+		// where the pointer didn't survive and the variable didn't exist and we had to set it up in PreDuplicateRequirementActions
+		VariableDesc = NewDescOwner->ResolveVariable(DuplicationInfo.VariableName);
+
+		if (VariableDesc.IsValid())
+		{
+			// Re-create the pins. Don't call PostCreateNode, since it won't notify the graph that the pin layout has
+			// changed.
+			check(GetPins().IsEmpty());
+			ConstructNode();
+		}
 
-	// Refresh the ResourceDesc so that we don't hold a reference to a VariableDesc in another deformer asset
-	VariableDesc = NewDescOwner->ResolveVariable(VariableDesc->GetFName());
+		// Empty the duplication info before we carry on.
+		DuplicationInfo = FOptimusNode_GetVariable_DuplicationInfo();
+	}
+
+}
+
+
+void UOptimusNode_GetVariable::ExportCustomProperties(FOutputDevice& Out, uint32 Indent)
+{
+	if (const UOptimusVariableDescription* Var = VariableDesc.Get())
+	{
+		Out.Logf(TEXT("%sCustomProperties VariableDefinition Name=\"%s\" Type=%s"),
+			FCString::Spc(Indent), *Var->VariableName.ToString(), *Var->DataType->TypeName.ToString());
+
+		if (const FProperty* Property = Var->DefaultValue->GetClass()->FindPropertyByName(UOptimusValueContainerGeneratorClass::ValuePropertyName))
+		{
+			FString ValueStr;
+			Property->ExportTextItem_InContainer(ValueStr, Var->DefaultValue.Get(), nullptr, nullptr, PPF_None);
+			Out.Logf(TEXT(" DefaultValue=\"%s\""), *ValueStr.ReplaceCharWithEscapedChar());
+		}
+		Out.Logf(TEXT("\n"));
+	}
+}
+
+
+void UOptimusNode_GetVariable::ImportCustomProperties(const TCHAR* SourceText, FFeedbackContext* Warn)
+{
+	if (FParse::Command(&SourceText, TEXT("VariableDefinition")))
+	{
+		FName VariableName;
+		if (!FParse::Value(SourceText, TEXT("Name="), VariableName))
+		{
+			return;
+		}
+
+		FName DataTypeName;
+		if (!FParse::Value(SourceText, TEXT("Type="), DataTypeName))
+		{
+			return;
+		}
+
+		const FOptimusDataTypeRef DataType = FOptimusDataTypeRegistry::Get().FindType(DataTypeName);
+		if (!DataType.IsValid())
+		{
+			return;
+		}
+
+		DuplicationInfo.VariableName = VariableName;
+		DuplicationInfo.DataType = DataType;
+
+		if (FParse::Value(SourceText, TEXT("DefaultValue="), DuplicationInfo.DefaultValue))
+		{
+			// Unescape control characters.
+			DuplicationInfo.DefaultValue.ReplaceEscapedCharWithCharInline();
+		}
+	}
 }
 
+
 void UOptimusNode_GetVariable::SetVariableDescription(UOptimusVariableDescription* InVariableDesc)
 {
 	if (!ensure(InVariableDesc))
@@ -127,4 +243,5 @@ void UOptimusNode_GetVariable::ConstructNode()
 	}
 }
 
-#undef LOCTEXT_NAMESPACE
\ No newline at end of file
+
+#undef LOCTEXT_NAMESPACE
diff --git a/Engine/Plugins/Animation/DeformerGraph/Source/OptimusCore/Private/Nodes/OptimusNode_GetVariable.h b/Engine/Plugins/Animation/DeformerGraph/Source/OptimusCore/Private/Nodes/OptimusNode_GetVariable.h
--- a/Engine/Plugins/Animation/DeformerGraph/Source/OptimusCore/Private/Nodes/OptimusNode_GetVariable.h
+++ b/Engine/Plugins/Animation/DeformerGraph/Source/OptimusCore/Private/Nodes/OptimusNode_GetVariable.h
@@ -11,6 +11,24 @@
 class UOptimusVariableDescription;
 
 
+// Information to hold data on the variable definition that persists over node duplication.
+// FIXME: This could be generalized in a better fashion.
+USTRUCT()
+struct FOptimusNode_GetVariable_DuplicationInfo
+{
+	GENERATED_BODY()
+	
+	UPROPERTY()
+	FName VariableName;
+
+	UPROPERTY()
+	FOptimusDataTypeRef DataType;
+
+	UPROPERTY()
+	FString DefaultValue;
+};
+
+
 UCLASS(Hidden)
 class UOptimusNode_GetVariable : 
 	public UOptimusNode,
@@ -39,10 +57,18 @@ public:
 protected:
 	void ConstructNode() override;
 
+	void PreDuplicateRequirementActions(const UOptimusNodeGraph* InTargetGraph, FOptimusCompoundAction* InCompoundAction) override;
+
 	// UObject overrides
-	virtual void PostDuplicate(EDuplicateMode::Type DuplicateMode) override;
+	void PostDuplicate(EDuplicateMode::Type DuplicateMode) override;
+	void ExportCustomProperties(FOutputDevice& Out, uint32 Indent) override;
+	void ImportCustomProperties(const TCHAR* SourceText, FFeedbackContext* Warn) override;
 	
 private:
 	UPROPERTY()
 	TWeakObjectPtr<UOptimusVariableDescription> VariableDesc;
+	
+	// Duplication data across graphs
+	UPROPERTY(DuplicateTransient)
+	FOptimusNode_GetVariable_DuplicationInfo DuplicationInfo;
 };
diff --git a/Engine/Plugins/Animation/DeformerGraph/Source/OptimusCore/Private/Nodes/OptimusNode_ResourceAccessorBase.cpp b/Engine/Plugins/Animation/DeformerGraph/Source/OptimusCore/Private/Nodes/OptimusNode_ResourceAccessorBase.cpp
--- a/Engine/Plugins/Animation/DeformerGraph/Source/OptimusCore/Private/Nodes/OptimusNode_ResourceAccessorBase.cpp
+++ b/Engine/Plugins/Animation/DeformerGraph/Source/OptimusCore/Private/Nodes/OptimusNode_ResourceAccessorBase.cpp
@@ -3,45 +3,153 @@
 #include "OptimusNode_ResourceAccessorBase.h"
 
 #include "OptimusCoreModule.h"
+#include "OptimusDataTypeRegistry.h"
 #include "OptimusDeformer.h"
 #include "OptimusResourceDescription.h"
 #include "DataInterfaces/OptimusDataInterfaceRawBuffer.h"
+#include "Actions/OptimusAction.h"
+#include "Actions/OptimusResourceActions.h"
+
 
 #define LOCTEXT_NAMESPACE "OptimusResourceAccessorBase"
 
 
-void UOptimusNode_ResourceAccessorBase::PostDuplicate(EDuplicateMode::Type DuplicateMode)
+void UOptimusNode_ResourceAccessorBase::PreDuplicateRequirementActions(
+	const UOptimusNodeGraph* InTargetGraph,
+	FOptimusCompoundAction* InCompoundAction
+	)
 {
-	Super::PostDuplicate(DuplicateMode);
-
-	if (!ResourceDesc.IsValid())
+	if (DuplicationInfo.ResourceName.IsNone())
 	{
 		return;
 	}
 	
+	// Check if the deformer has a variable that matches the name and type.
+	UOptimusDeformer* Deformer = Cast<UOptimusDeformer>(InTargetGraph->GetCollectionRoot());
+
+	for (const UOptimusResourceDescription* ExistingResourceDesc: Deformer->GetResources())
+	{
+		if (ExistingResourceDesc->ResourceName == DuplicationInfo.ResourceName &&
+			ExistingResourceDesc->DataType == DuplicationInfo.DataType &&
+			ExistingResourceDesc->DataDomain == DuplicationInfo.DataDomain)
+		{
+			// We don't care about the binding itself, the user can fix that up later.
+			return;
+		}
+	}
+
+	// Note: This will fail if there are multiple overlapping names. A better system is needed.
+	bool bFoundDuplicate = false;
+	do
+	{
+		for (const UOptimusResourceDescription* ExistingResourceDesc: Deformer->GetResources())
+		{
+			if (ExistingResourceDesc->ResourceName == DuplicationInfo.ResourceName)
+			{
+				DuplicationInfo.ResourceName.SetNumber(DuplicationInfo.ResourceName.GetNumber() + 1);
+				bFoundDuplicate = true;
+			}
+		}
+	}
+	while(bFoundDuplicate);
+
+	InCompoundAction->AddSubAction<FOptimusResourceAction_AddResource>(DuplicationInfo.DataType, DuplicationInfo.ResourceName, DuplicationInfo.DataDomain);
+}
+
+
+void UOptimusNode_ResourceAccessorBase::PostDuplicate(EDuplicateMode::Type DuplicateMode)
+{
+	Super::PostDuplicate(DuplicateMode);
+
 	if (!GetOwningGraph())
 	{
 		return;
 	}
-	
-	const UOptimusDeformer* OldDescOwner = ResourceDesc->GetOwningDeformer();
-	const UOptimusDeformer* NewDescOwner = Cast<UOptimusDeformer>(GetOwningGraph()->GetCollectionRoot());
 
+	const UOptimusDeformer* NewDescOwner = Cast<UOptimusDeformer>(GetOwningGraph()->GetCollectionRoot());
 	if (!NewDescOwner)
 	{
 		return;
 	}
 	
-	// No action needed if we are copying/pasting within the same deformer asset 
-	if (OldDescOwner == NewDescOwner)
+	if (ResourceDesc.IsValid())
 	{
-		return;
+		const UOptimusDeformer* OldDescOwner = ResourceDesc->GetOwningDeformer();
+		
+		// No action needed if we are copying/pasting within the same deformer asset 
+		if (OldDescOwner == NewDescOwner)
+		{
+			return;
+		}
+
+		// Refresh the ResourceDesc so that we don't hold a reference to a ResourceDesc in another deformer asset
+		ResourceDesc = NewDescOwner->ResolveResource(ResourceDesc->GetFName());
+	}
+	else if (!DuplicationInfo.ResourceName.IsNone())
+	{
+		ResourceDesc = NewDescOwner->ResolveResource(DuplicationInfo.ResourceName);
+
+		if (ResourceDesc.IsValid())
+		{
+			// Re-create the pins. Don't call PostCreateNode, since it won't notify the graph that the pin layout has
+			// changed.
+			check(GetPins().IsEmpty());
+			ConstructNode();
+		}
+
+		// Empty the duplication info before we carry on.
+		DuplicationInfo = FOptimusNode_ResourceAccessorBase_DuplicationInfo();
+	}
+}
+
+
+void UOptimusNode_ResourceAccessorBase::ExportCustomProperties(FOutputDevice& Out, uint32 Indent)
+{
+	if (const UOptimusResourceDescription* Res = ResourceDesc.Get())
+	{
+		Out.Logf(TEXT("%sCustomProperties ResourceDefinition Name=\"%s\" Type=%s DataDomain=%s\n"),
+			FCString::Spc(Indent), *Res->ResourceName.ToString(), *Res->DataType->TypeName.ToString(),
+			*Res->DataDomain.ToString());
 	}
+	
+}
+
 
-	// Refresh the ResourceDesc so that we don't hold a reference to a ResourceDesc in another deformer asset
-	ResourceDesc = NewDescOwner->ResolveResource(ResourceDesc->GetFName());
+void UOptimusNode_ResourceAccessorBase::ImportCustomProperties(const TCHAR* SourceText, FFeedbackContext* Warn)
+{
+	if (FParse::Command(&SourceText, TEXT("ResourceDefinition")))
+	{
+		FName ResourceName;
+		if (!FParse::Value(SourceText, TEXT("Name="), ResourceName))
+		{
+			return;
+		}
+		
+		FName DataTypeName;
+		if (!FParse::Value(SourceText, TEXT("Type="), DataTypeName))
+		{
+			return;
+		}
+
+		const FOptimusDataTypeRef DataType = FOptimusDataTypeRegistry::Get().FindType(DataTypeName);
+		if (!DataType.IsValid())
+		{
+			return;
+		}
+
+		FString DataDomainStr;
+		if (!FParse::Value(SourceText, TEXT("DataDomain="), DataDomainStr))
+		{
+			return;
+		}
+
+		DuplicationInfo.ResourceName = ResourceName;
+		DuplicationInfo.DataType = DataType;
+		DuplicationInfo.DataDomain = FOptimusDataDomain::FromString(DataDomainStr);
+	}
 }
 
+
 void UOptimusNode_ResourceAccessorBase::SetResourceDescription(UOptimusResourceDescription* InResourceDesc)
 {
 	if (!ensure(InResourceDesc))
diff --git a/Engine/Plugins/Animation/DeformerGraph/Source/OptimusCore/Private/Nodes/OptimusNode_ResourceAccessorBase.h b/Engine/Plugins/Animation/DeformerGraph/Source/OptimusCore/Private/Nodes/OptimusNode_ResourceAccessorBase.h
--- a/Engine/Plugins/Animation/DeformerGraph/Source/OptimusCore/Private/Nodes/OptimusNode_ResourceAccessorBase.h
+++ b/Engine/Plugins/Animation/DeformerGraph/Source/OptimusCore/Private/Nodes/OptimusNode_ResourceAccessorBase.h
@@ -14,6 +14,22 @@ class UOptimusComputeDataInterface;
 class UOptimusResourceDescription;
 
 
+USTRUCT()
+struct FOptimusNode_ResourceAccessorBase_DuplicationInfo
+{
+	GENERATED_BODY()
+	
+	UPROPERTY()
+	FName ResourceName;
+
+	UPROPERTY()
+	FOptimusDataTypeRef DataType;
+
+	UPROPERTY()
+	FOptimusDataDomain DataDomain;
+};
+
+
 UCLASS(Abstract)
 class UOptimusNode_ResourceAccessorBase : 
 	public UOptimusNode,
@@ -46,8 +62,12 @@ public:
 
 	
 protected:
+	void PreDuplicateRequirementActions(const UOptimusNodeGraph* InTargetGraph, FOptimusCompoundAction* InCompoundAction) override;
+	
 	// UObject overrides
-	virtual void PostDuplicate(EDuplicateMode::Type DuplicateMode) override;
+	void PostDuplicate(EDuplicateMode::Type DuplicateMode) override;
+	void ExportCustomProperties(FOutputDevice& Out, uint32 Indent) override;
+	void ImportCustomProperties(const TCHAR* SourceText, FFeedbackContext* Warn) override;
 	
 protected:
 	UPROPERTY()
@@ -56,4 +76,7 @@ protected:
 	/** Logical operation when writing to the resource. */
 	UPROPERTY(EditAnywhere, Category = Resource)
 	EOptimusBufferWriteType WriteType = EOptimusBufferWriteType::Write;
+
+	UPROPERTY(DuplicateTransient)
+	FOptimusNode_ResourceAccessorBase_DuplicationInfo DuplicationInfo;
 };
diff --git a/Engine/Plugins/Animation/DeformerGraph/Source/OptimusCore/Private/OptimusDataDomain.cpp b/Engine/Plugins/Animation/DeformerGraph/Source/OptimusCore/Private/OptimusDataDomain.cpp
--- a/Engine/Plugins/Animation/DeformerGraph/Source/OptimusCore/Private/OptimusDataDomain.cpp
+++ b/Engine/Plugins/Animation/DeformerGraph/Source/OptimusCore/Private/OptimusDataDomain.cpp
@@ -4,6 +4,7 @@
 
 #include "OptimusExpressionEvaluator.h"
 #include "OptimusObjectVersion.h"
+#include "String/ParseTokens.h"
 
 
 namespace Optimus::DomainName
@@ -54,6 +55,64 @@ TOptional<int32> FOptimusDataDomain::GetElementCount(TMap<FName, int32> InDomain
 }
 
 
+FString FOptimusDataDomain::ToString() const
+{
+	switch(Type)
+	{
+	case EOptimusDataDomainType::Dimensional:
+		{
+			if (DimensionNames.IsEmpty())
+			{
+				return TEXT("Empty");
+			}
+			else if (DimensionNames.Num() == 1)
+			{
+				FString Result;
+				for (FName Name: DimensionNames)
+				{
+					Result += Name.ToString();
+					Result += TEXT("/");
+				}
+				Result += FString::FromInt(Multiplier);
+				return Result;
+			}
+		}
+		break;
+	case EOptimusDataDomainType::Expression:
+		return FString::Printf(TEXT("'%s'"), *Expression);
+	}
+	
+	return {};
+}
+
+
+FOptimusDataDomain FOptimusDataDomain::FromString(
+	const FString& InString
+	)
+{
+	FString String = InString.TrimStartAndEnd();
+	if (String.StartsWith(TEXT("'")))
+	{
+		return FOptimusDataDomain(String.Mid(1, String.Len() - 2));
+	}
+
+	TArray<FStringView> Parts;
+	UE::String::ParseTokens(InString, TEXT("/"), Parts);
+	if (Parts.Num() < 2)
+	{
+		return {};
+	}
+
+	const int32 Multiplier = FCString::Atoi(*FString(Parts.Pop()));
+	TArray<FName> Names;
+	for (FStringView Part: Parts)
+	{
+		Names.Add(FName(Part));
+	}
+	return FOptimusDataDomain(Names, Multiplier);
+}
+
+
 bool FOptimusDataDomain::operator==(const FOptimusDataDomain& InOtherDomain) const
 {
 	if (Type != InOtherDomain.Type)
diff --git a/Engine/Plugins/Animation/DeformerGraph/Source/OptimusCore/Private/OptimusNodeGraph.cpp b/Engine/Plugins/Animation/DeformerGraph/Source/OptimusCore/Private/OptimusNodeGraph.cpp
--- a/Engine/Plugins/Animation/DeformerGraph/Source/OptimusCore/Private/OptimusNodeGraph.cpp
+++ b/Engine/Plugins/Animation/DeformerGraph/Source/OptimusCore/Private/OptimusNodeGraph.cpp
@@ -580,6 +580,13 @@ bool UOptimusNodeGraph::DuplicateNodes(
 		Action->SetTitlef(TEXT("%s %d Nodes"), *InActionName, InNodes.Num());
 	}
 
+	// Add all the pre-duplicate requirement actions first. This allows certain nodes to set up their
+	// operating environment correctly (e.g. a missing variable description for get variable nodes, etc.)
+	for (UOptimusNode* Node: InNodes)
+	{
+		Node->PreDuplicateRequirementActions(this, Action);
+	}
+
 	// Duplicate the nodes and place them correctly
 	for (UOptimusNode* Node: InNodes)
 	{
@@ -592,10 +599,14 @@ bool UOptimusNodeGraph::DuplicateNodes(
 		Action->AddSubAction(DuplicateNodeAction);
 	}
 
-	// Add any links that the nodes may have had.
+	// Add any links that the nodes may have had. These operations are allowed to fail, in which case
+	// we just end up with dangling nodes.
+	// In the future we would like to introduce connections that are in an error state (will fail compile)
+	// but show up as disconnectable wires in the graph (e.g. if a type no longer matches).
 	for (const TTuple<FString, FString>& LinkInfo: NodeLinks)
 	{
-		Action->AddSubAction<FOptimusNodeGraphAction_AddLink>(LinkInfo.Key, LinkInfo.Value);
+		constexpr bool bCanFail = true;
+		Action->AddSubAction<FOptimusNodeGraphAction_AddLink>(LinkInfo.Key, LinkInfo.Value, bCanFail);
 	}
 
 	return GetActionStack()->RunAction(Action);
diff --git a/Engine/Plugins/Animation/DeformerGraph/Source/OptimusCore/Public/IOptimusGeneratedClassDefiner.h b/Engine/Plugins/Animation/DeformerGraph/Source/OptimusCore/Public/IOptimusGeneratedClassDefiner.h
new file mode 100644
--- /dev/null
+++ b/Engine/Plugins/Animation/DeformerGraph/Source/OptimusCore/Public/IOptimusGeneratedClassDefiner.h
@@ -0,0 +1,40 @@
+// Copyright Epic Games, Inc. All Rights Reserved.
+
+#pragma once
+
+#include "OptimusDataType.h"
+#include "UObject/Interface.h"
+
+#include "IOptimusGeneratedClassDefiner.generated.h"
+
+UINTERFACE()
+class OPTIMUSCORE_API UOptimusGeneratedClassDefiner :
+	public UInterface
+{
+	GENERATED_BODY()
+};
+
+
+/** FIXME: A stop-gap shader value provider until we have a proper pin evaluation that handles
+  * paths that have a constant, computed, varying and a mix thereof, results.
+  */
+class IOptimusGeneratedClassDefiner
+{
+	GENERATED_BODY()
+
+public:
+	/** Returns the class whose CDO we can get to get this interface to create a new generated class from
+	 *  a string definition. */
+	virtual FTopLevelAssetPath GetAssetPathForClassDefiner() const = 0;
+	
+	/** Export definition values for the given class which are sufficient to pass into a call to
+	  * CreateClassFromCreationString to generate a class based on the class that hosts this interface. */
+	virtual FString GetClassCreationString() const = 0;
+
+	/** Create a generated class, within the package provided, from the creation string that was originally
+	 *  generated by GetClassCreationString. */
+	virtual UClass* GetClassFromCreationString(
+		UPackage* InPackage,
+		const TCHAR* InCreationString
+		) const = 0; 
+};
diff --git a/Engine/Plugins/Animation/DeformerGraph/Source/OptimusCore/Public/OptimusDataDomain.h b/Engine/Plugins/Animation/DeformerGraph/Source/OptimusCore/Public/OptimusDataDomain.h
--- a/Engine/Plugins/Animation/DeformerGraph/Source/OptimusCore/Public/OptimusDataDomain.h
+++ b/Engine/Plugins/Animation/DeformerGraph/Source/OptimusCore/Public/OptimusDataDomain.h
@@ -90,7 +90,13 @@ struct OPTIMUSCORE_API FOptimusDataDomain
 	bool IsSingleton() const
 	{
 		return Type == EOptimusDataDomainType::Dimensional && DimensionNames.IsEmpty();
-	} 
+	}
+
+	/** Convert the data domain to a text serializable string */
+	FString ToString() const;
+
+	/** Create a data domain object from a serializable string generated by the above function. */
+	static FOptimusDataDomain FromString(const FString& InString);
 
 	bool operator==(const FOptimusDataDomain& InOtherDomain) const;
 	bool operator!=(const FOptimusDataDomain& InOtherDomain) const { return !operator==(InOtherDomain); }
diff --git a/Engine/Plugins/Animation/DeformerGraph/Source/OptimusCore/Public/OptimusHelpers.h b/Engine/Plugins/Animation/DeformerGraph/Source/OptimusCore/Public/OptimusHelpers.h
--- a/Engine/Plugins/Animation/DeformerGraph/Source/OptimusCore/Public/OptimusHelpers.h
+++ b/Engine/Plugins/Animation/DeformerGraph/Source/OptimusCore/Public/OptimusHelpers.h
@@ -54,7 +54,7 @@ namespace Optimus
 			// FIXME: The constructor is broken. It only needs a const ref.
 		    : FObjectReader(const_cast<TArray<uint8>&>(InBytes))
 		{
-			this->SetWantBinaryPropertySerialization(true);
+			FObjectReader::SetWantBinaryPropertySerialization(true);
 			Obj->Serialize(*this);
 		}
 	};
@@ -65,7 +65,7 @@ namespace Optimus
 		FBinaryObjectWriter(UObject* Obj, TArray<uint8>& OutBytes)
 		    : FObjectWriter(OutBytes)
 		{
-			this->SetWantBinaryPropertySerialization(true);
+			FObjectWriter::SetWantBinaryPropertySerialization(true);
 			Obj->Serialize(*this);
 		}
 	};
diff --git a/Engine/Plugins/Animation/DeformerGraph/Source/OptimusCore/Public/OptimusNode.h b/Engine/Plugins/Animation/DeformerGraph/Source/OptimusCore/Public/OptimusNode.h
--- a/Engine/Plugins/Animation/DeformerGraph/Source/OptimusCore/Public/OptimusNode.h
+++ b/Engine/Plugins/Animation/DeformerGraph/Source/OptimusCore/Public/OptimusNode.h
@@ -18,6 +18,7 @@ enum class EOptimusNodePinStorageType : uint8;
 class UOptimusActionStack;
 class UOptimusNodeGraph;
 class UOptimusNodePin;
+struct FOptimusCompoundAction;
 struct FOptimusDataTypeRef;
 struct FOptimusParameterBinding;
 
@@ -217,6 +218,13 @@ protected:
 	{
 		return {};
 	}
+
+	/** Called prior to duplicate to allow the node to add its own graph requirements to
+	 *  to the list of actions being performed.
+	 */
+	virtual void PreDuplicateRequirementActions(
+		const UOptimusNodeGraph* InTargetGraph, 
+		FOptimusCompoundAction *InCompoundAction) {}
 	
 
 	void EnableDynamicPins();
diff --git a/Engine/Plugins/Animation/DeformerGraph/Source/OptimusCore/Public/OptimusNodeGraph.h b/Engine/Plugins/Animation/DeformerGraph/Source/OptimusCore/Public/OptimusNodeGraph.h
--- a/Engine/Plugins/Animation/DeformerGraph/Source/OptimusCore/Public/OptimusNodeGraph.h
+++ b/Engine/Plugins/Animation/DeformerGraph/Source/OptimusCore/Public/OptimusNodeGraph.h
@@ -343,7 +343,7 @@ protected:
 	friend class UOptimusNode;
 	friend class UOptimusNodePin;
 	friend class UOptimusNode_ConstantValue;
-	friend class UOptimusClipboardContent;
+	friend class FOptimusEditorClipboard;
 	friend struct FOptimusNodeGraphAction_AddNode;
 	friend struct FOptimusNodeGraphAction_DuplicateNode;
 	friend struct FOptimusNodeGraphAction_RemoveNode;
diff --git a/Engine/Plugins/Animation/DeformerGraph/Source/OptimusEditor/Private/OptimusEditor.cpp b/Engine/Plugins/Animation/DeformerGraph/Source/OptimusEditor/Private/OptimusEditor.cpp
--- a/Engine/Plugins/Animation/DeformerGraph/Source/OptimusEditor/Private/OptimusEditor.cpp
+++ b/Engine/Plugins/Animation/DeformerGraph/Source/OptimusEditor/Private/OptimusEditor.cpp
@@ -448,9 +448,7 @@ void FOptimusEditor::CopySelectedNodes() const
 {
 	FOptimusEditorClipboard& Clipboard = FOptimusEditorModule::Get().GetClipboard();
 
-	const TArray<UOptimusNode*> ModelNodes = GetSelectedModelNodes();
-	UOptimusClipboardContent* Content = UOptimusClipboardContent::Create(EditorGraph->GetModelGraph(), ModelNodes);
-	Clipboard.SetClipboardContent(Content);
+	Clipboard.SetClipboardFromNodes(GetSelectedModelNodes());
 }
 
 
@@ -466,8 +464,7 @@ void FOptimusEditor::CutSelectedNodes() const
 
 	UOptimusNodeGraph* ModelGraph = EditorGraph->GetModelGraph();
 	const TArray<UOptimusNode*> ModelNodes = GetSelectedModelNodes();
-	UOptimusClipboardContent* Content = UOptimusClipboardContent::Create(ModelGraph, ModelNodes);
-	Clipboard.SetClipboardContent(Content);
+	Clipboard.SetClipboardFromNodes(ModelNodes);
 
 	ModelGraph->RemoveNodes(ModelNodes, TEXT("Cut"));
 }
@@ -482,14 +479,12 @@ bool FOptimusEditor::CanCutNodes() const
 void FOptimusEditor::PasteNodes() const
 {
 	const FOptimusEditorClipboard& Clipboard = FOptimusEditorModule::Get().GetClipboard();
-	const UOptimusClipboardContent* Content = Clipboard.GetClipboardContent();
-	if (!Content)
+	const UOptimusNodeGraph* TransientGraph = Clipboard.GetGraphFromClipboardContent(DeformerObject->GetPackage());
+	if (!TransientGraph)
 	{
 		return;
 	}
 
-	const UOptimusNodeGraph* TransientGraph = Content->GetGraphFromClipboardContent();
-
 	UOptimusNodeGraph* ModelGraph = EditorGraph->GetModelGraph();
 	ModelGraph->DuplicateNodes(
 		TransientGraph->GetAllNodes(), GraphEditorWidget->GetPasteLocation(), TEXT("Paste"));
diff --git a/Engine/Plugins/Animation/DeformerGraph/Source/OptimusEditor/Private/OptimusEditorClipboard.cpp b/Engine/Plugins/Animation/DeformerGraph/Source/OptimusEditor/Private/OptimusEditorClipboard.cpp
--- a/Engine/Plugins/Animation/DeformerGraph/Source/OptimusEditor/Private/OptimusEditorClipboard.cpp
+++ b/Engine/Plugins/Animation/DeformerGraph/Source/OptimusEditor/Private/OptimusEditorClipboard.cpp
@@ -6,136 +6,343 @@
 #include "OptimusNodeGraph.h"
 #include "OptimusNodeLink.h"
 #include "OptimusNodePin.h"
-#include "UnrealExporter.h"
+
+#include "Editor.h"
+#include "IOptimusGeneratedClassDefiner.h"
 #include "Exporters/Exporter.h"
+#include "UnrealExporter.h"
 #include "HAL/PlatformApplicationMisc.h"
+#include "Misc/EngineVersion.h"
 
 
-UOptimusClipboardContent* UOptimusClipboardContent::Create(
-	const UOptimusNodeGraph* InGraph,
-	const TArray<UOptimusNode*>& InNodes
-	)
+void FOptimusEditorClipboard::SetClipboardFromNodes(const TArray<UOptimusNode*>& InNodes)
 {
-	for (const UOptimusNode* Node: InNodes)
+	if (InNodes.IsEmpty())
 	{
-		if (Node->GetOwningGraph() != InGraph)
-		{
-			return nullptr;
-		}
+		return;
 	}
+	
+	UOptimusNodeGraph* Graph = InNodes[0]->GetOwningGraph();
 
-	UOptimusClipboardContent* Content = NewObject<UOptimusClipboardContent>(GetTransientPackage());
+	// Export the clipboard to text.
+	const FExportObjectInnerContext Context;
+	FStringOutputDevice Archive;
 
-	TMap<const UOptimusNode*, int32> NodeIndexMap;
-	for (const UOptimusNode* Node: InNodes)
+	// Clear all stray export marks to ensure all the objects get properly exported. 
+	UnMarkAllObjects(OBJECTMARK_TagExp);
+
+	// Export the nodes first.
+	TSet<UOptimusNode*> NodesToCopy;
+	NodesToCopy.Reserve(InNodes.Num());
+	for (UOptimusNode* Node: InNodes)
 	{
-		NodeIndexMap.Add(Node, Content->Nodes.Num());
-		Content->Nodes.Add(Node);
+		if (!ensure(Node->GetOwningGraph() == Graph))
+		{
+			continue;
+		}
+
+		// Export any generated classes first. If the node implements the IOptimusGeneratedClassDefiner interface,
+		// it means that it uses a generated class as a base. We need to ensure that we have all the relevant information
+		// to construct the generated class on the receiving side.
+		if (const IOptimusGeneratedClassDefiner* ClassDefiner = Cast<IOptimusGeneratedClassDefiner>(Node))
+		{
+			Archive.Logf(TEXT("GeneratedClass Class=%s GeneratorPath=%s %s\n"),
+				*Node->GetClass()->GetPathName(), *ClassDefiner->GetAssetPathForClassDefiner().ToString(),
+				*ClassDefiner->GetClassCreationString());
+		}
+		
+		UExporter::ExportToOutputDevice(
+			&Context,
+			Node,
+			nullptr,
+			Archive,
+			TEXT("copy"),
+			0,
+			PPF_ExportsNotFullyQualified | PPF_Copy | PPF_Delimited,
+			false,
+			Graph);
+		
+		NodesToCopy.Add(Node);
 	}
 
-	// Make a copy of all the links that connect the nodes that are being copied.
-	for (const UOptimusNodeLink* Link: InGraph->GetAllLinks())
+	// Export all internal links as well.
+	for (UOptimusNodeLink* Link: Graph->GetAllLinks())
 	{
 		if (ensure(Link->GetNodeOutputPin() != nullptr))
 		{
 			const UOptimusNode *OutputNode = Link->GetNodeOutputPin()->GetOwningNode();
 			const UOptimusNode *InputNode = Link->GetNodeInputPin()->GetOwningNode();
 
-			if (NodeIndexMap.Contains(OutputNode) && NodeIndexMap.Contains(InputNode))
+			// Only include the link if the both end points are connected to nodes in the list of nodes
+			// we're interested in copying.
+			if (NodesToCopy.Contains(OutputNode) && NodesToCopy.Contains(InputNode))
 			{
-				FOptimusClipboardNodeLink ClipboardLink;
-				ClipboardLink.NodeOutputIndex = NodeIndexMap[OutputNode];
-				ClipboardLink.NodeOutputPinName = Link->GetNodeOutputPin()->GetUniqueName().ToString();
-			
-				ClipboardLink.NodeInputIndex = NodeIndexMap[InputNode];
-				ClipboardLink.NodeInputPinName = Link->GetNodeInputPin()->GetUniqueName().ToString();
-
-				Content->NodeLinks.Add(ClipboardLink);
+				UExporter::ExportToOutputDevice(
+					&Context,
+					Link,
+					nullptr,
+					Archive,
+					TEXT("copy"),
+					0,
+					PPF_ExportsNotFullyQualified | PPF_Copy | PPF_Delimited,
+					false,
+					Graph);
 			}
 		}
 	}
+	
+	FPlatformApplicationMisc::ClipboardCopy(*Archive);
+}
+
+
+UOptimusNodeGraph* FOptimusEditorClipboard::GetGraphFromClipboardContent(
+	UPackage* InTargetPackage
+	)
+{
+	// Get the text from the clipboard.
+	FString ClipboardText;
+	FPlatformApplicationMisc::ClipboardPaste(ClipboardText);
+
+	// Try to create optimus clipboard content from that. We're not using a custom version of FCustomizableTextObjectFactory
+	// since the original does not handle sub-objects at all.
+	if (CanCreateObjectsFromText(*ClipboardText))
+	{
+		UOptimusNodeGraph* Graph = NewObject<UOptimusNodeGraph>(GetTransientPackage(), FName("ClipboardGraph"), RF_Transient);
+		
+		ProcessObjectBuffer(InTargetPackage, Graph, *ClipboardText);
+		return Graph;
+	}
 
-	return Content;
+	return nullptr;
 }
 
 
-UOptimusNodeGraph* UOptimusClipboardContent::GetGraphFromClipboardContent() const
+bool FOptimusEditorClipboard::HasValidClipboardContent()
 {
-	// Create a temporary graph that will hold the nodes that the caller will then duplicate
-	// out of.
-	UOptimusNodeGraph* Graph = NewObject<UOptimusNodeGraph>(GetTransientPackage(), FName("ClipboardGraph"), RF_Transient);
+	// Get the text from the clipboard.
+	FString ClipboardText;
+	FPlatformApplicationMisc::ClipboardPaste(ClipboardText);
+	return CanCreateObjectsFromText(*ClipboardText);
+}
+
+
+bool FOptimusEditorClipboard::CanCreateObjectsFromText(
+	const TCHAR* InBuffer
+	)
+{
+	FParse::Next( &InBuffer );
 	
-	// Duplicate the nodes into it.
-	TMap<const UOptimusNode*, UOptimusNode*> OldNodeToNewNodeMap;
-	for (const UOptimusNode* Node: Nodes)
+	// Find at least one object we can create.
+	FString StrLine;
+	while (FParse::Line(&InBuffer,StrLine))
 	{
-		UOptimusNode* NewNode = DuplicateObject(Node, Graph, Node->GetFName());
-		Graph->AddNodeDirect(NewNode);
-		OldNodeToNewNodeMap.Add(Node, NewNode);
+		const TCHAR* Str = *StrLine;
+		if (FParse::Command(&Str, TEXT("Begin")) && FParse::Command(&Str,TEXT("Object")))
+		{
+			UClass* ObjClass = nullptr;
+			if (ParseObject<UClass>(Str, TEXT("Class="), ObjClass, nullptr))
+			{
+				if (CanCreateClass(ObjClass))
+				{
+					return true;
+				}
+			}
+		}
 	}
+	return false;
+}
 
-	// Link them up.
-	for (const FOptimusClipboardNodeLink& LinkInfo: NodeLinks)
+
+bool FOptimusEditorClipboard::ProcessObjectBuffer(
+	UPackage* InPackage, 
+	UObject* InRootOuter,
+	const TCHAR* InBuffer
+	)
+{
+	FParse::Next(&InBuffer);
+	
+	FObjectInstancingGraph InstancingGraph;
+	TMap<TPair<FName /*ObjectName*/, UObject* /*Outer*/>, UObject*> ObjectMap;
+	TArray<UObject*> CurrentObjectStack;
+	TArray<UObject*> CreatedObjects;
+	TMap<FName, UClass*> GeneratedClasses;
+	CurrentObjectStack.Push(InRootOuter);
+
+	auto ParseObjectCommand = [](const TCHAR** Str, const TCHAR* Match) -> bool
+	{
+		const TCHAR* Tmp = *Str;
+		if (FParse::Command(Str, Match) && FParse::Command(Str, TEXT("Object")))
+		{
+			return true;
+		}
+		*Str = Tmp;
+		return false;
+	};
+	auto ParseCommand = [](const TCHAR** Str, const TCHAR* Match) -> bool
 	{
-		const UOptimusNode* OutputNode = OldNodeToNewNodeMap[Nodes[LinkInfo.NodeOutputIndex]];
-		const UOptimusNode* InputNode = OldNodeToNewNodeMap[Nodes[LinkInfo.NodeInputIndex]];
+		const TCHAR* Tmp = *Str;
+		if (FParse::Command(Str, Match))
+		{
+			return true;
+		}
+		*Str = Tmp;
+		return false;
+	};
+	
+	// 
+	FString LineStr;
+	while (FParse::Line(&InBuffer, LineStr))
+	{
+		const TCHAR* LineBuffer = *LineStr;
 
-		UOptimusNodePin* OutputPin = OutputNode->FindPin(LinkInfo.NodeOutputPinName);
-		UOptimusNodePin* InputPin = InputNode->FindPin(LinkInfo.NodeInputPinName);
+		if (ParseObjectCommand(&LineBuffer, TEXT("Begin")))
+		{
+			// Get the name, it should always be there.
+			FName ObjectName;
+			if (!FParse::Value(LineBuffer, TEXT("Name="), ObjectName))
+			{
+				return false;
+			}
+			
+			UClass *ObjectClass = nullptr;
+			bool bInvalidObject = false;
+			if (!ParseObject<UClass>(LineBuffer, TEXT("Class="), ObjectClass, nullptr, &bInvalidObject) &&
+				bInvalidObject)
+			{
+				// Check if we have it in the list of generated classes. We already know the tag is there.
+				FName ClassName;
+				(void)FParse::Value(LineBuffer, TEXT("Class="), ClassName);
+				UClass **ObjectClassResult = GeneratedClasses.Find(ClassName);
+				if (!ObjectClassResult)
+				{
+					// Class= found but no matching class object found.
+					return false;
+				}
 
-		Graph->AddLinkDirect(OutputPin, InputPin);
-	}
+				ObjectClass = *ObjectClassResult;
+			}
 
-	return Graph;
-}
+			// If a class is given, create the object now and add it to the outer stack.
+			if (ObjectClass)
+			{
+				UObject* CreatedObject = NewObject<UObject>(CurrentObjectStack.Last(), ObjectClass, ObjectName, RF_NoFlags, nullptr, true, &InstancingGraph);
 
+				ObjectMap.Add(MakeTuple(ObjectName, CurrentObjectStack.Last()), CreatedObject);
+				CreatedObjects.Add(CreatedObject);
 
-void FOptimusEditorClipboard::SetClipboardContent(UOptimusClipboardContent* InContent)
-{
-	// Clear the mark state for saving.
-	UnMarkAllObjects(EObjectMark(OBJECTMARK_TagExp | OBJECTMARK_TagImp));
+				// Mark the newly created object as the current object on the stack, it will receive any properties found
+				// at this level and be the outer for any newly created objects as well.
+				CurrentObjectStack.Push(CreatedObject);
+			}
+			// Else it's a definition containing the properties, set the current object to the named object.
+			else if (UObject **FoundObject = ObjectMap.Find(MakeTuple(ObjectName, CurrentObjectStack.Last())))
+			{
+				CurrentObjectStack.Push(*FoundObject);
+			}
+			else
+			{
+				// Something's gone wrong.
+				return false;
+			}
+		}
+		else if (ParseObjectCommand(&LineBuffer, TEXT("End")))
+		{
+			CurrentObjectStack.Pop();
+		}
+		else if (ParseCommand(&LineBuffer, TEXT("GeneratedClass")))
+		{
+			FName ClassName;
+			if (!FParse::Value(LineBuffer, TEXT("Class="), ClassName))
+			{
+				return false;
+			}
 
-	// Export the clipboard to text.
-	FStringOutputDevice Archive;
-	const FExportObjectInnerContext Context;
-	UExporter::ExportToOutputDevice(
-		&Context,
-		InContent,
-		nullptr,
-		Archive,
-		TEXT("copy"),
-		0,
-		PPF_ExportsNotFullyQualified | PPF_Copy | PPF_Delimited,
-		false,
-		InContent->GetOuter());
-	FPlatformApplicationMisc::ClipboardCopy(*Archive);
-}
+			FString GeneratorPath;
+			if (!FParse::Value(LineBuffer, TEXT("GeneratorPath="), GeneratorPath))
+			{
+				return false;
+			}
 
+			// Find the class whose object implements the definer that we need to create the generated class. 
+			FTopLevelAssetPath GeneratorAssetPath(GeneratorPath);
+			UClass* GeneratorClass = FindObject<UClass>(GeneratorAssetPath);
+			if (!GeneratorClass)
+			{
+				return false;
+			}
+			IOptimusGeneratedClassDefiner* ClassDefiner = Cast<IOptimusGeneratedClassDefiner>(GeneratorClass->GetDefaultObject(true));
+			if (!ClassDefiner)
+			{
+				return false;
+			}
 
-UOptimusClipboardContent* FOptimusEditorClipboard::GetClipboardContent() const
-{
-	// Get the text from the clipboard.
-	FString ClipboardText;
-	FPlatformApplicationMisc::ClipboardPaste(ClipboardText);
+			UClass* GeneratedClass = ClassDefiner->GetClassFromCreationString(InPackage, LineBuffer);
+			if (!GeneratedClass)
+			{
+				return false;
+			}
+			GeneratedClasses.Add(ClassName, GeneratedClass);
+		}
+		// Custom properties?
+		else if (!CurrentObjectStack.IsEmpty())
+		{
+			UObject* PropertyObject = CurrentObjectStack.Last();
+			FImportObjectParams ImportParams;
 
-	// Try to create optimus clipboard content from that.
-	FOptimusEditorClipboardContentTextObjectFactory ClipboardContentFactory;
-	if (ClipboardContentFactory.CanCreateObjectsFromText(ClipboardText))
+			ImportParams.DestData = reinterpret_cast<uint8*>(PropertyObject);
+			ImportParams.SourceText = LineBuffer;
+			ImportParams.ObjectStruct = PropertyObject->GetClass();
+			ImportParams.SubobjectRoot = PropertyObject;
+			ImportParams.SubobjectOuter = PropertyObject;
+			ImportParams.Warn = GWarn;
+			ImportParams.Depth = 0;
+			ImportParams.LineNumber = 0;
+			ImportParams.InInstanceGraph = &InstancingGraph;
+			ImportParams.bShouldCallEditChange = false;
+			
+			ImportObjectProperties(ImportParams);
+		}
+	}
+
+	// Process the created objects in reverse order of creation.
+	for (int32 Index = CreatedObjects.Num() - 1; Index >= 0; Index--)
 	{
-		ClipboardContentFactory.ProcessBuffer(GetTransientPackage(), RF_Transactional, ClipboardText);
-		return ClipboardContentFactory.ClipboardContent;
+		if (!ProcessPostCreateObject(InRootOuter, CreatedObjects[Index]))
+		{
+			return false;
+		}
 	}
 
-	return nullptr;
+	return true;
 }
 
 
-bool FOptimusEditorClipboard::HasValidClipboardContent() const
+bool FOptimusEditorClipboard::CanCreateClass(const UClass* InClass)
 {
-	// Get the text from the clipboard.
-	FString ClipboardText;
-	FPlatformApplicationMisc::ClipboardPaste(ClipboardText);
+	return InClass->IsChildOf(UOptimusNode::StaticClass()) ||
+		   InClass == UOptimusNodePin::StaticClass() ||
+		   InClass == UOptimusNodeLink::StaticClass();
+}
+
+
+bool FOptimusEditorClipboard::ProcessPostCreateObject(UObject* InRootOuter, UObject* InNewObject)
+{
+	if (InNewObject->GetOuter() == InRootOuter && InRootOuter->IsA<UOptimusNodeGraph>())
+	{
+		UOptimusNodeGraph* Graph = Cast<UOptimusNodeGraph>(InRootOuter);
+		if (InNewObject->IsA<UOptimusNode>())
+		{
+			UOptimusNode* Node = Cast<UOptimusNode>(InNewObject);
+
+			return Graph->AddNodeDirect(Node);
+		}
+		else if (InNewObject->IsA<UOptimusNodeLink>())
+		{
+			const UOptimusNodeLink* Link = Cast<UOptimusNodeLink>(InNewObject);
+
+			return Graph->AddLinkDirect(Link->GetNodeOutputPin(), Link->GetNodeInputPin());
+		}
+	}
 
-	const FOptimusEditorClipboardContentTextObjectFactory ClipboardContentFactory;
-	return ClipboardContentFactory.CanCreateObjectsFromText(ClipboardText);
+	return true;
 }
diff --git a/Engine/Plugins/Animation/DeformerGraph/Source/OptimusEditor/Public/OptimusEditorClipboard.h b/Engine/Plugins/Animation/DeformerGraph/Source/OptimusEditor/Public/OptimusEditorClipboard.h
--- a/Engine/Plugins/Animation/DeformerGraph/Source/OptimusEditor/Public/OptimusEditorClipboard.h
+++ b/Engine/Plugins/Animation/DeformerGraph/Source/OptimusEditor/Public/OptimusEditorClipboard.h
@@ -2,89 +2,32 @@
 
 #pragma once
 
-#include "Factories.h"
+#include "Containers/Array.h"
+#include "HAL/Platform.h"
 
-#include "OptimusEditorClipboard.generated.h"
 
+class UClass;
+class UObject;
 class UOptimusNodeGraph;
 class UOptimusNode;
 
-USTRUCT()
-struct FOptimusClipboardNodeLink
-{
-	GENERATED_BODY()
-	
-	// Index into the UOptimusClipboardContent::Nodes list.
-	UPROPERTY()
-	int32 NodeOutputIndex = INDEX_NONE;
-
-	UPROPERTY()
-	FString NodeOutputPinName;
-
-	UPROPERTY()
-	int32 NodeInputIndex = INDEX_NONE;
-
-	UPROPERTY()
-	FString NodeInputPinName;
-};
-
-UCLASS()
-class UOptimusClipboardContent :
-	public UObject
-{
-	GENERATED_BODY()
-	
-public:
-	static UOptimusClipboardContent* Create(
-		const UOptimusNodeGraph* InGraph,
-		const TArray<UOptimusNode*>& InNodes
-		);
-
-	// Creates a temporary node graph with the contents of the clipboard
-	UOptimusNodeGraph* GetGraphFromClipboardContent() const;
-
-	UPROPERTY()
-	TArray<TObjectPtr<const UOptimusNode>> Nodes;
-
-	UPROPERTY()
-	TArray<FOptimusClipboardNodeLink> NodeLinks;
-};
-
 
 class OPTIMUSEDITOR_API FOptimusEditorClipboard
 {
 public:
-	FOptimusEditorClipboard() = default;
-	
-	void SetClipboardContent(UOptimusClipboardContent *InContent);
-
-	UOptimusClipboardContent *GetClipboardContent() const;
+	static void SetClipboardFromNodes(const TArray<UOptimusNode*>& InNodes);
 
-	bool HasValidClipboardContent() const;
-};
-
-
-struct FOptimusEditorClipboardContentTextObjectFactory final :
-	FCustomizableTextObjectFactory
-{
-	UOptimusClipboardContent* ClipboardContent = nullptr;
+	// Creates a self-contained 
+	static UOptimusNodeGraph* GetGraphFromClipboardContent(
+		UPackage* InTargetPackage
+		);
 
-	FOptimusEditorClipboardContentTextObjectFactory() :
-		FCustomizableTextObjectFactory(GWarn)
-	{
-	}
+	static bool HasValidClipboardContent();
 
-protected:
-	bool CanCreateClass(UClass* ObjectClass, bool& bOmitSubObjs) const override
-	{
-		return ObjectClass == UOptimusClipboardContent::StaticClass();
-	}
+private:
+	static bool CanCreateObjectsFromText(const TCHAR* InBuffer);
 
-	void ProcessConstructedObject(UObject* CreatedObject) override
-	{
-		if (CreatedObject->IsA<UOptimusClipboardContent>())
-		{
-			ClipboardContent = CastChecked<UOptimusClipboardContent>(CreatedObject);
-		}
-	}
+	static bool ProcessObjectBuffer(UPackage* InPackage, UObject* InRootOuter, const TCHAR* InBuffer);
+	static bool CanCreateClass(const UClass* ObjectClass);
+	static bool ProcessPostCreateObject(UObject* InRootOuter, UObject* InNewObject);
 };
diff --git a/Engine/Source/Editor/UnrealEd/Public/Editor.h b/Engine/Source/Editor/UnrealEd/Public/Editor.h
--- a/Engine/Source/Editor/UnrealEd/Public/Editor.h
+++ b/Engine/Source/Editor/UnrealEd/Public/Editor.h
@@ -502,7 +502,7 @@ struct FImportObjectParams
  *
  * @return	NULL if the default values couldn't be imported
  */
-const TCHAR* ImportObjectProperties( FImportObjectParams& InParams );
+UNREALED_API const TCHAR* ImportObjectProperties( FImportObjectParams& InParams );
 
 /**
  * Parse and import text as property values for the object specified.
