Optimus: Add data interface for duplicate vertices.
This is needed for a better recompute normal calculation that doesn't generate seams at UV borders.

diff --git a/Engine/Plugins/Animation/DeformerGraph/Shaders/Private/DataInterfaceDuplicateVertices.ush b/Engine/Plugins/Animation/DeformerGraph/Shaders/Private/DataInterfaceDuplicateVertices.ush
new file mode 100644
--- /dev/null
+++ b/Engine/Plugins/Animation/DeformerGraph/Shaders/Private/DataInterfaceDuplicateVertices.ush
@@ -0,0 +1,36 @@
+// Copyright Epic Games, Inc. All Rights Reserved.
+
+uint {DataInterfaceName}_NumVertices;
+uint {DataInterfaceName}_NumDuplicateVertices;
+uint {DataInterfaceName}_InputStreamStart;
+Buffer<uint> {DataInterfaceName}_DuplicateVertexStartAndLength;
+Buffer<uint> {DataInterfaceName}_DuplicateVertices;
+
+uint ReadNumVertices_{DataInterfaceName}()
+{
+	return {DataInterfaceName}_NumVertices;
+}
+
+uint ReadNumDuplicateVertices_{DataInterfaceName}()
+{
+	return {DataInterfaceName}_NumDuplicateVertices;
+}
+
+int2 ReadDuplicateVerticesStartAndLength_{DataInterfaceName}(uint VertexIndex)
+{
+#if ENABLE_DUPLICATED_VERTICES
+	uint BufferIndex = ({DataInterfaceName}_InputStreamStart + VertexIndex) * 2;
+	return int2({DataInterfaceName}_DuplicateVertexStartAndLength[BufferIndex  + 1], {DataInterfaceName}_DuplicateVertexStartAndLength[BufferIndex]);
+#else	
+	return 0;
+#endif
+}
+
+uint ReadDuplicateVertex_{DataInterfaceName}(uint Index)
+{
+#if ENABLE_DUPLICATED_VERTICES
+	return {DataInterfaceName}_DuplicateVertices[Index];
+#else	
+	return 0;
+#endif
+}
diff --git a/Engine/Plugins/Animation/DeformerGraph/Shaders/Private/DataInterfaceSkinnedMesh.ush b/Engine/Plugins/Animation/DeformerGraph/Shaders/Private/DataInterfaceSkinnedMesh.ush
--- a/Engine/Plugins/Animation/DeformerGraph/Shaders/Private/DataInterfaceSkinnedMesh.ush
+++ b/Engine/Plugins/Animation/DeformerGraph/Shaders/Private/DataInterfaceSkinnedMesh.ush
@@ -3,9 +3,6 @@
 // Include required for platform TANGENT_RWBUFFER_FORMAT.
 #include "/Engine/Private/GpuSkinCommon.ush"
 
-// todo: This should be set by compute graph permutations.
-#define MERGE_DUPLICATED_VERTICES 0
-
 uint {DataInterfaceName}_NumVertices;
 uint {DataInterfaceName}_NumTriangles;
 uint {DataInterfaceName}_NumUVChannels;
@@ -16,8 +13,6 @@ Buffer<float> {DataInterfaceName}_PositionInputBuffer;
 Buffer<SNORM float4> {DataInterfaceName}_TangentInputBuffer;
 Buffer<float4> {DataInterfaceName}_ColorInputBuffer;
 Buffer<float2> {DataInterfaceName}_UVInputBuffer;
-Buffer<uint> {DataInterfaceName}_DuplicatedIndicesIndices;
-Buffer<uint> {DataInterfaceName}_DuplicatedIndices;
 
 uint ReadNumVertices_{DataInterfaceName}()
 {
@@ -69,30 +64,3 @@ float2 ReadUV_{DataInterfaceName}(uint VertexIndex, uint UVChannelIndex)
 	uint BufferIndex = ({DataInterfaceName}_InputStreamStart + VertexIndex) * {DataInterfaceName}_NumUVChannels + UVChannelIndex;
 	return {DataInterfaceName}_UVInputBuffer[BufferIndex];
 }
-
-uint ReadDuplicatedIndicesStart_{DataInterfaceName}(uint VertexIndex)
-{
-#if MERGE_DUPLICATED_VERTICES 
-	return {DataInterfaceName}_DuplicatedIndicesIndices[VertexIndex * 2 + 1];
-#else
-	return 0;
-#endif
-}
-
-uint ReadDuplicatedIndicesLength_{DataInterfaceName}(uint VertexIndex)
-{
-#if MERGE_DUPLICATED_VERTICES 
-	return {DataInterfaceName}_DuplicatedIndicesIndices[VertexIndex * 2];
-#else
-	return 0;
-#endif
-}
-
-uint ReadDuplicatedIndex_{DataInterfaceName}(uint Index)
-{
-#if MERGE_DUPLICATED_VERTICES 
-	return {DataInterfaceName}_DuplicatedIndices[Index];
-#else
-	return 0;
-#endif
-}
diff --git a/Engine/Plugins/Animation/DeformerGraph/Source/OptimusCore/Private/ComponentSources/OptimusSkinnedMeshComponentSource.cpp b/Engine/Plugins/Animation/DeformerGraph/Source/OptimusCore/Private/ComponentSources/OptimusSkinnedMeshComponentSource.cpp
--- a/Engine/Plugins/Animation/DeformerGraph/Source/OptimusCore/Private/ComponentSources/OptimusSkinnedMeshComponentSource.cpp
+++ b/Engine/Plugins/Animation/DeformerGraph/Source/OptimusCore/Private/ComponentSources/OptimusSkinnedMeshComponentSource.cpp
@@ -14,6 +14,7 @@
 
 FName UOptimusSkinnedMeshComponentSource::Domains::Vertex("Vertex");
 FName UOptimusSkinnedMeshComponentSource::Domains::Triangle("Triangle");
+FName UOptimusSkinnedMeshComponentSource::Domains::DuplicateVertex("DuplicateVertex");
 
 
 FText UOptimusSkinnedMeshComponentSource::GetDisplayName() const
diff --git a/Engine/Plugins/Animation/DeformerGraph/Source/OptimusCore/Private/ComponentSources/OptimusSkinnedMeshComponentSource.h b/Engine/Plugins/Animation/DeformerGraph/Source/OptimusCore/Private/ComponentSources/OptimusSkinnedMeshComponentSource.h
--- a/Engine/Plugins/Animation/DeformerGraph/Source/OptimusCore/Private/ComponentSources/OptimusSkinnedMeshComponentSource.h
+++ b/Engine/Plugins/Animation/DeformerGraph/Source/OptimusCore/Private/ComponentSources/OptimusSkinnedMeshComponentSource.h
@@ -17,6 +17,7 @@ public:
 	{
 		static FName Vertex;
 		static FName Triangle;
+		static FName DuplicateVertex;
 	};
 	
 	// UOptimusComponentSource implementations
diff --git a/Engine/Plugins/Animation/DeformerGraph/Source/OptimusCore/Private/DataInterfaces/OptimusDataInterfaceDuplicateVertices.cpp b/Engine/Plugins/Animation/DeformerGraph/Source/OptimusCore/Private/DataInterfaces/OptimusDataInterfaceDuplicateVertices.cpp
new file mode 100644
--- /dev/null
+++ b/Engine/Plugins/Animation/DeformerGraph/Source/OptimusCore/Private/DataInterfaces/OptimusDataInterfaceDuplicateVertices.cpp
@@ -0,0 +1,172 @@
+// Copyright Epic Games, Inc. All Rights Reserved.
+
+#include "OptimusDataInterfaceDuplicateVertices.h"
+
+#include "ComponentSources/OptimusSkinnedMeshComponentSource.h"
+#include "ComputeFramework/ComputeKernelPermutationVector.h"
+#include "ComputeFramework/ShaderParamTypeDefinition.h"
+#include "OptimusDataDomain.h"
+#include "Rendering/SkeletalMeshLODRenderData.h"
+#include "Rendering/SkeletalMeshRenderData.h"
+#include "ShaderParameterMetadataBuilder.h"
+#include "SkeletalRenderPublic.h"
+
+
+FString UOptimusDuplicateVerticesDataInterface::GetDisplayName() const
+{
+	return TEXT("Duplicate Vertices");
+}
+
+TArray<FOptimusCDIPinDefinition> UOptimusDuplicateVerticesDataInterface::GetPinDefinitions() const
+{
+	FName Vertex(UOptimusSkinnedMeshComponentSource::Domains::Vertex);
+	FName Triangle(UOptimusSkinnedMeshComponentSource::Domains::Triangle);
+	FName DuplicateVertex(UOptimusSkinnedMeshComponentSource::Domains::DuplicateVertex);
+	
+	TArray<FOptimusCDIPinDefinition> Defs;
+	Defs.Add({"NumVertices", "ReadNumVertices"});
+	Defs.Add({"DuplicateVerticesStartAndLength", "ReadDuplicateVerticesStartAndLength", Vertex, "ReadDuplicateVerticesStartAndLength"});
+	Defs.Add({"DuplicateVertex", "ReadDuplicateVertex", DuplicateVertex, "ReadDuplicateVertex"});
+	return Defs;
+}
+
+
+TSubclassOf<UActorComponent> UOptimusDuplicateVerticesDataInterface::GetRequiredComponentClass() const
+{
+	return USkinnedMeshComponent::StaticClass();
+}
+
+
+void UOptimusDuplicateVerticesDataInterface::GetSupportedInputs(TArray<FShaderFunctionDefinition>& OutFunctions) const
+{
+	OutFunctions.AddDefaulted_GetRef()
+		.SetName(TEXT("ReadNumVertices"))
+		.AddReturnType(EShaderFundamentalType::Uint);
+
+	OutFunctions.AddDefaulted_GetRef()
+		.SetName(TEXT("ReadNumDuplicateVertices"))
+		.AddReturnType(EShaderFundamentalType::Uint);
+
+	OutFunctions.AddDefaulted_GetRef()
+		.SetName(TEXT("ReadDuplicateVerticesStartAndLength"))
+		.AddReturnType(EShaderFundamentalType::Int, 2)
+		.AddParam(EShaderFundamentalType::Uint);
+
+	OutFunctions.AddDefaulted_GetRef()
+		.SetName(TEXT("ReadDuplicateVertex"))
+		.AddReturnType(EShaderFundamentalType::Uint)
+		.AddParam(EShaderFundamentalType::Uint);
+}
+
+BEGIN_SHADER_PARAMETER_STRUCT(FDuplicateVerticesDataInterfaceParameters, )
+	SHADER_PARAMETER(uint32, NumVertices)
+	SHADER_PARAMETER(uint32, NumDuplicateVertices)
+	SHADER_PARAMETER(uint32, InputStreamStart)
+	SHADER_PARAMETER_SRV(Buffer<uint>, DuplicateVertexStartAndLength)
+	SHADER_PARAMETER_SRV(Buffer<uint>, DuplicateVertices)
+END_SHADER_PARAMETER_STRUCT()
+
+void UOptimusDuplicateVerticesDataInterface::GetShaderParameters(TCHAR const* UID, FShaderParametersMetadataBuilder& InOutBuilder, FShaderParametersMetadataAllocations& InOutAllocations) const
+{
+	InOutBuilder.AddNestedStruct<FDuplicateVerticesDataInterfaceParameters>(UID);
+}
+
+void UOptimusDuplicateVerticesDataInterface::GetPermutations(FComputeKernelPermutationVector& OutPermutationVector) const
+{
+	OutPermutationVector.AddPermutation(TEXT("ENABLE_DUPLICATED_VERTICES"), 2);
+}
+
+void UOptimusDuplicateVerticesDataInterface::GetShaderHash(FString& InOutKey) const
+{
+	GetShaderFileHash(TEXT("/Plugin/Optimus/Private/DataInterfaceSkinnedMesh.ush"), EShaderPlatform::SP_PCD3D_SM5).AppendString(InOutKey);
+}
+
+void UOptimusDuplicateVerticesDataInterface::GetHLSL(FString& OutHLSL, FString const& InDataInterfaceName) const
+{
+	TMap<FString, FStringFormatArg> TemplateArgs =
+	{
+		{ TEXT("DataInterfaceName"), InDataInterfaceName },
+	};
+
+	FString TemplateFile;
+	LoadShaderSourceFile(TEXT("/Plugin/Optimus/Private/DataInterfaceDuplicateVertices.ush"), EShaderPlatform::SP_PCD3D_SM5, &TemplateFile, nullptr);
+	OutHLSL += FString::Format(*TemplateFile, TemplateArgs);
+}
+
+UComputeDataProvider* UOptimusDuplicateVerticesDataInterface::CreateDataProvider(TObjectPtr<UObject> InBinding, uint64 InInputMask, uint64 InOutputMask) const
+{
+	UOptimusDuplicateVerticesDataProvider* Provider = NewObject<UOptimusDuplicateVerticesDataProvider>();
+	Provider->SkinnedMesh = Cast<USkinnedMeshComponent>(InBinding);
+	return Provider;
+}
+
+
+bool UOptimusDuplicateVerticesDataProvider::IsValid() const
+{
+	return
+		SkinnedMesh != nullptr &&
+		SkinnedMesh->MeshObject != nullptr;
+}
+
+FComputeDataProviderRenderProxy* UOptimusDuplicateVerticesDataProvider::GetRenderProxy()
+{
+	return new FOptimusDuplicateVerticesDataProviderProxy(SkinnedMesh);
+}
+
+
+FOptimusDuplicateVerticesDataProviderProxy::FOptimusDuplicateVerticesDataProviderProxy(USkinnedMeshComponent* SkinnedMeshComponent)
+{
+	SkeletalMeshObject = SkinnedMeshComponent->MeshObject;
+}
+
+struct FDuplicateVerticesDataInterfacePermutationIds
+{
+	uint32 EnableDuplicateVertices = 0;
+
+	FDuplicateVerticesDataInterfacePermutationIds(FComputeKernelPermutationVector const& PermutationVector)
+	{
+		{
+			static FString Name(TEXT("ENABLE_DUPLICATED_VERTICES"));
+			static uint32 Hash = GetTypeHash(Name);
+			EnableDuplicateVertices = PermutationVector.GetPermutationBits(Name, Hash, 1);
+		}
+	}
+};
+
+void FOptimusDuplicateVerticesDataProviderProxy::GatherDispatchData(FDispatchSetup const& InDispatchSetup, FCollectedDispatchData& InOutDispatchData)
+{
+	if (!ensure(InDispatchSetup.ParameterStructSizeForValidation == sizeof(FDuplicateVerticesDataInterfaceParameters)))
+	{
+		return;
+	}
+
+	const int32 LodIndex = SkeletalMeshObject->GetLOD();
+	FSkeletalMeshRenderData const& SkeletalMeshRenderData = SkeletalMeshObject->GetSkeletalMeshRenderData();
+	FSkeletalMeshLODRenderData const* LodRenderData = &SkeletalMeshRenderData.LODRenderData[LodIndex];
+	if (!ensure(LodRenderData->RenderSections.Num() == InDispatchSetup.NumInvocations))
+	{
+		return;
+	}
+
+	FDuplicateVerticesDataInterfacePermutationIds PermutationIds(InDispatchSetup.PermutationVector);
+
+	FRHIShaderResourceView* NullSRVBinding = GWhiteVertexBufferWithSRV->ShaderResourceViewRHI.GetReference();
+
+	for (int32 InvocationIndex = 0; InvocationIndex < InDispatchSetup.NumInvocations; ++InvocationIndex)
+	{
+		FSkelMeshRenderSection const& RenderSection = LodRenderData->RenderSections[InvocationIndex];
+
+		FRHIShaderResourceView* DuplicateVertexStartAndLengthSRV = RenderSection.DuplicatedVerticesBuffer.LengthAndIndexDuplicatedVerticesIndexBuffer.VertexBufferSRV;
+		FRHIShaderResourceView* DuplicateVerticesSRV = RenderSection.DuplicatedVerticesBuffer.DuplicatedVerticesIndexBuffer.VertexBufferSRV;
+		const bool bValidDuplicateIndices = (DuplicateVertexStartAndLengthSRV != nullptr) && (DuplicateVerticesSRV != nullptr);
+
+		FDuplicateVerticesDataInterfaceParameters* Parameters = (FDuplicateVerticesDataInterfaceParameters*)(InOutDispatchData.ParameterBuffer + InDispatchSetup.ParameterBufferOffset + InDispatchSetup.ParameterBufferStride * InvocationIndex);
+		Parameters->NumVertices = RenderSection.NumVertices;
+		Parameters->NumDuplicateVertices = bValidDuplicateIndices ? RenderSection.DuplicatedVerticesBuffer.DupVertData.Num() : 0;
+		Parameters->InputStreamStart = RenderSection.BaseVertexIndex;
+		Parameters->DuplicateVertexStartAndLength = bValidDuplicateIndices ? DuplicateVertexStartAndLengthSRV : NullSRVBinding;
+		Parameters->DuplicateVertices = bValidDuplicateIndices ? DuplicateVerticesSRV : NullSRVBinding;
+
+		InOutDispatchData.PermutationId[InvocationIndex] |= (bValidDuplicateIndices ? PermutationIds.EnableDuplicateVertices : 0);
+	}
+}
diff --git a/Engine/Plugins/Animation/DeformerGraph/Source/OptimusCore/Private/DataInterfaces/OptimusDataInterfaceDuplicateVertices.h b/Engine/Plugins/Animation/DeformerGraph/Source/OptimusCore/Private/DataInterfaces/OptimusDataInterfaceDuplicateVertices.h
new file mode 100644
--- /dev/null
+++ b/Engine/Plugins/Animation/DeformerGraph/Source/OptimusCore/Private/DataInterfaces/OptimusDataInterfaceDuplicateVertices.h
@@ -0,0 +1,67 @@
+// Copyright Epic Games, Inc. All Rights Reserved.
+
+#pragma once
+
+#include "OptimusComputeDataInterface.h"
+#include "ComputeFramework/ComputeDataProvider.h"
+
+#include "OptimusDataInterfaceDuplicateVertices.generated.h"
+
+class FSkeletalMeshObject;
+class USkinnedMeshComponent;
+
+/** 
+ * Compute Framework Data Interface for reading duplicate vertices on a mesh.
+ * These are vertices that are at the same location but that have been split because of discontinous color/UV etc.
+ */
+UCLASS(Category = ComputeFramework)
+class OPTIMUSCORE_API UOptimusDuplicateVerticesDataInterface : public UOptimusComputeDataInterface
+{
+	GENERATED_BODY()
+
+public:
+	//~ Begin UOptimusComputeDataInterface Interface
+	FString GetDisplayName() const override;
+	TArray<FOptimusCDIPinDefinition> GetPinDefinitions() const override;
+	TSubclassOf<UActorComponent> GetRequiredComponentClass() const override;
+	//~ End UOptimusComputeDataInterface Interface
+	
+	//~ Begin UComputeDataInterface Interface
+	TCHAR const* GetClassName() const override { return TEXT("SkinnedMesh"); }
+	void GetSupportedInputs(TArray<FShaderFunctionDefinition>& OutFunctions) const override;
+	void GetShaderParameters(TCHAR const* UID, FShaderParametersMetadataBuilder& InOutBuilder, FShaderParametersMetadataAllocations& InOutAllocations) const override;
+	void GetPermutations(FComputeKernelPermutationVector& OutPermutationVector) const override;
+	void GetShaderHash(FString& InOutKey) const override;
+	void GetHLSL(FString& OutHLSL, FString const& InDataInterfaceName) const override;
+	UComputeDataProvider* CreateDataProvider(TObjectPtr<UObject> InBinding, uint64 InInputMask, uint64 InOutputMask) const override;
+	//~ End UComputeDataInterface Interface
+};
+
+/** Compute Framework Data Provider for reading skeletal mesh. */
+UCLASS(BlueprintType, editinlinenew, Category = ComputeFramework)
+class UOptimusDuplicateVerticesDataProvider : public UComputeDataProvider
+{
+	GENERATED_BODY()
+
+public:
+	UPROPERTY(EditAnywhere, BlueprintReadWrite, Category = Binding)
+	TObjectPtr<USkinnedMeshComponent> SkinnedMesh = nullptr;
+
+	//~ Begin UComputeDataProvider Interface
+	bool IsValid() const override;
+	FComputeDataProviderRenderProxy* GetRenderProxy() override;
+	//~ End UComputeDataProvider Interface
+};
+
+class FOptimusDuplicateVerticesDataProviderProxy : public FComputeDataProviderRenderProxy
+{
+public:
+	FOptimusDuplicateVerticesDataProviderProxy(USkinnedMeshComponent* SkinnedMeshComponent);
+
+	//~ Begin FComputeDataProviderRenderProxy Interface
+	void GatherDispatchData(FDispatchSetup const& InDispatchSetup, FCollectedDispatchData& InOutDispatchData) override;
+	//~ End FComputeDataProviderRenderProxy Interface
+
+private:
+	FSkeletalMeshObject* SkeletalMeshObject;
+};
diff --git a/Engine/Plugins/Animation/DeformerGraph/Source/OptimusCore/Private/DataInterfaces/OptimusDataInterfaceSkinnedMesh.cpp b/Engine/Plugins/Animation/DeformerGraph/Source/OptimusCore/Private/DataInterfaces/OptimusDataInterfaceSkinnedMesh.cpp
--- a/Engine/Plugins/Animation/DeformerGraph/Source/OptimusCore/Private/DataInterfaces/OptimusDataInterfaceSkinnedMesh.cpp
+++ b/Engine/Plugins/Animation/DeformerGraph/Source/OptimusCore/Private/DataInterfaces/OptimusDataInterfaceSkinnedMesh.cpp
@@ -86,21 +86,6 @@ void UOptimusSkinnedMeshDataInterface::GetSupportedInputs(TArray<FShaderFunction
 		.SetName(TEXT("ReadColor"))
 		.AddReturnType(EShaderFundamentalType::Float, 4)
 		.AddParam(EShaderFundamentalType::Uint);
-
-	OutFunctions.AddDefaulted_GetRef()
-		.SetName(TEXT("ReadDuplicatedIndicesStart"))
-		.AddReturnType(EShaderFundamentalType::Uint)
-		.AddParam(EShaderFundamentalType::Uint);
-
-	OutFunctions.AddDefaulted_GetRef()
-		.SetName(TEXT("ReadDuplicatedIndicesLength"))
-		.AddReturnType(EShaderFundamentalType::Uint)
-		.AddParam(EShaderFundamentalType::Uint);
-
-	OutFunctions.AddDefaulted_GetRef()
-		.SetName(TEXT("ReadDuplicatedIndex"))
-		.AddReturnType(EShaderFundamentalType::Uint)
-		.AddParam(EShaderFundamentalType::Uint);
 }
 
 BEGIN_SHADER_PARAMETER_STRUCT(FSkinnedMeshDataInterfaceParameters, )
@@ -114,8 +99,6 @@ BEGIN_SHADER_PARAMETER_STRUCT(FSkinnedMeshDataInterfaceParameters, )
 	SHADER_PARAMETER_SRV(Buffer<SNORM float4>, TangentInputBuffer)
 	SHADER_PARAMETER_SRV(Buffer<float2>, UVInputBuffer)
 	SHADER_PARAMETER_SRV(Buffer<float4>, ColorInputBuffer)
-	SHADER_PARAMETER_SRV(Buffer<uint>, DuplicatedIndicesIndices)
-	SHADER_PARAMETER_SRV(Buffer<uint>, DuplicatedIndices)
 END_SHADER_PARAMETER_STRUCT()
 
 void UOptimusSkinnedMeshDataInterface::GetShaderParameters(TCHAR const* UID, FShaderParametersMetadataBuilder& InOutBuilder, FShaderParametersMetadataAllocations& InOutAllocations) const
@@ -193,10 +176,6 @@ void FOptimusSkinnedMeshDataProviderProxy::GatherDispatchData(FDispatchSetup con
 		FRHIShaderResourceView* MeshUVBufferSRV = LodRenderData->StaticVertexBuffers.StaticMeshVertexBuffer.GetTexCoordsSRV();
 		FRHIShaderResourceView* MeshColorBufferSRV = LodRenderData->StaticVertexBuffers.ColorVertexBuffer.GetColorComponentsSRV();
 
-		FRHIShaderResourceView* DuplicatedIndicesIndicesSRV = RenderSection.DuplicatedVerticesBuffer.LengthAndIndexDuplicatedVerticesIndexBuffer.VertexBufferSRV;
-		FRHIShaderResourceView* DuplicatedIndicesSRV = RenderSection.DuplicatedVerticesBuffer.DuplicatedVerticesIndexBuffer.VertexBufferSRV;
-		const bool bValidDuplicatedIndices = (DuplicatedIndicesIndicesSRV != nullptr) && (DuplicatedIndicesSRV != nullptr);
-
 		FSkinnedMeshDataInterfaceParameters* Parameters = (FSkinnedMeshDataInterfaceParameters*)(InOutDispatchData.ParameterBuffer + InDispatchSetup.ParameterBufferOffset + InDispatchSetup.ParameterBufferStride * InvocationIndex);
 		Parameters->NumVertices = RenderSection.NumVertices;
 		Parameters->NumTriangles = RenderSection.NumTriangles;
@@ -208,7 +187,5 @@ void FOptimusSkinnedMeshDataProviderProxy::GatherDispatchData(FDispatchSetup con
 		Parameters->TangentInputBuffer = MeshTangentBufferSRV != nullptr ? MeshTangentBufferSRV : NullSRVBinding;
 		Parameters->UVInputBuffer = MeshUVBufferSRV != nullptr ? MeshUVBufferSRV : NullSRVBinding;
 		Parameters->ColorInputBuffer = MeshColorBufferSRV != nullptr ? MeshColorBufferSRV : NullSRVBinding;
-		Parameters->DuplicatedIndicesIndices = DuplicatedIndicesIndicesSRV != nullptr ? DuplicatedIndicesIndicesSRV : NullSRVBinding;
-		Parameters->DuplicatedIndices = DuplicatedIndicesSRV != nullptr ? DuplicatedIndicesSRV : NullSRVBinding;
 	}
 }
