Optimus: Return a nullptr from CreateInstance if ComputeFramework is disabled for a platform.
This allows clients to fall back to whatever deformation logic they would use without the deformer graph.

diff --git a/Engine/Plugins/Animation/DeformerGraph/Source/OptimusCore/Private/OptimusDeformer.cpp b/Engine/Plugins/Animation/DeformerGraph/Source/OptimusCore/Private/OptimusDeformer.cpp
--- a/Engine/Plugins/Animation/DeformerGraph/Source/OptimusCore/Private/OptimusDeformer.cpp
+++ b/Engine/Plugins/Animation/DeformerGraph/Source/OptimusCore/Private/OptimusDeformer.cpp
@@ -6,11 +6,16 @@
 #include "Actions/OptimusNodeGraphActions.h"
 #include "Actions/OptimusResourceActions.h"
 #include "Actions/OptimusVariableActions.h"
+#include "Components/MeshComponent.h"
+#include "ComputeFramework/ComputeFramework.h"
+#include "ComputeFramework/ComputeKernel.h"
+#include "Containers/Queue.h"
 #include "DataInterfaces/OptimusDataInterfaceGraph.h"
 #include "DataInterfaces/OptimusDataInterfaceRawBuffer.h"
 #include "IOptimusComputeKernelProvider.h"
 #include "IOptimusDataInterfaceProvider.h"
 #include "IOptimusValueProvider.h"
+#include "Misc/UObjectToken.h"
 #include "OptimusActionStack.h"
 #include "OptimusComputeGraph.h"
 #include "OptimusDataTypeRegistry.h"
@@ -25,12 +30,8 @@
 #include "OptimusObjectVersion.h"
 #include "OptimusResourceDescription.h"
 #include "OptimusVariableDescription.h"
-
-#include "Components/MeshComponent.h"
-#include "ComputeFramework/ComputeKernel.h"
-#include "Containers/Queue.h"
-#include "Misc/UObjectToken.h"
 #include "RenderingThread.h"
+#include "SceneInterface.h"
 #include "UObject/Package.h"
 
 // FIXME: We should not be accessing nodes directly.
@@ -2303,6 +2304,13 @@ UMeshDeformerInstance* UOptimusDeformer::CreateInstance(
 		return nullptr;
 	}
 
+	// Return nullptr if ComputeFramework is disabled. Clients can then fallback to some other behaviour.
+	EShaderPlatform Platform = InMeshComponent->GetScene() != nullptr ? InMeshComponent->GetScene()->GetShaderPlatform() : GMaxRHIShaderPlatform;
+	if (!ComputeFramework::IsEnabled() || !ComputeFramework::IsSupported(Platform))
+	{
+		return nullptr;
+	}
+
 	const FName InstanceName(GetName() + TEXT("_Instance"));
 	UOptimusDeformerInstance* Instance = NewObject<UOptimusDeformerInstance>(InMeshComponent, InstanceName);
 	Instance->SetMeshComponent(InMeshComponent);
