Tweaked unity sizes.

diff --git a/Engine/Plugins/Animation/ControlRig/Source/ControlRig/ControlRig.Build.cs b/Engine/Plugins/Animation/ControlRig/Source/ControlRig/ControlRig.Build.cs
--- a/Engine/Plugins/Animation/ControlRig/Source/ControlRig/ControlRig.Build.cs
+++ b/Engine/Plugins/Animation/ControlRig/Source/ControlRig/ControlRig.Build.cs
@@ -6,7 +6,9 @@ namespace UnrealBuildTool.Rules
     {
         public ControlRig(ReadOnlyTargetRules Target) : base(Target)
         {
-            PrivateIncludePaths.Add("ControlRig/ThirdParty/AHEasing");
+			NumIncludedBytesPerUnityCPPOverride = 688128; // best unity size found from using UBT ProfileUnitySizes mode
+
+			PrivateIncludePaths.Add("ControlRig/ThirdParty/AHEasing");
 
             PrivateDependencyModuleNames.AddRange(
                 new string[]
diff --git a/Engine/Plugins/Animation/DeformerGraph/Source/OptimusCore/OptimusCore.Build.cs b/Engine/Plugins/Animation/DeformerGraph/Source/OptimusCore/OptimusCore.Build.cs
--- a/Engine/Plugins/Animation/DeformerGraph/Source/OptimusCore/OptimusCore.Build.cs
+++ b/Engine/Plugins/Animation/DeformerGraph/Source/OptimusCore/OptimusCore.Build.cs
@@ -5,8 +5,10 @@ namespace UnrealBuildTool.Rules
     public class OptimusCore : ModuleRules
     {
         public OptimusCore(ReadOnlyTargetRules Target) : base(Target)
-        {
-            PrivateIncludePaths.AddRange(
+		{
+			NumIncludedBytesPerUnityCPPOverride = 688128; // best unity size found from using UBT ProfileUnitySizes mode
+
+			PrivateIncludePaths.AddRange(
                 new string[] {
                     "OptimusCore/Private",
 				}
diff --git a/Engine/Plugins/FX/Niagara/Source/Niagara/Niagara.Build.cs b/Engine/Plugins/FX/Niagara/Source/Niagara/Niagara.Build.cs
--- a/Engine/Plugins/FX/Niagara/Source/Niagara/Niagara.Build.cs
+++ b/Engine/Plugins/FX/Niagara/Source/Niagara/Niagara.Build.cs
@@ -6,8 +6,10 @@ using UnrealBuildTool;
 public class Niagara : ModuleRules
 {
     public Niagara(ReadOnlyTargetRules Target) : base(Target)
-    {
-        PrivateIncludePaths.Add("../../../../Shaders/Shared");
+	{
+		NumIncludedBytesPerUnityCPPOverride = 491520; // best unity size found from using UBT ProfileUnitySizes mode
+
+		PrivateIncludePaths.Add("../../../../Shaders/Shared");
 
 		// Specific to OpenVDB support
 		if (Target.IsInPlatformGroup(UnrealPlatformGroup.Windows))
diff --git a/Engine/Plugins/Runtime/GameplayAbilities/Source/GameplayAbilities/GameplayAbilities.Build.cs b/Engine/Plugins/Runtime/GameplayAbilities/Source/GameplayAbilities/GameplayAbilities.Build.cs
--- a/Engine/Plugins/Runtime/GameplayAbilities/Source/GameplayAbilities/GameplayAbilities.Build.cs
+++ b/Engine/Plugins/Runtime/GameplayAbilities/Source/GameplayAbilities/GameplayAbilities.Build.cs
@@ -6,6 +6,8 @@ namespace UnrealBuildTool.Rules
 	{
 		public GameplayAbilities(ReadOnlyTargetRules Target) : base(Target)
 		{
+			NumIncludedBytesPerUnityCPPOverride = 491520; // best unity size found from using UBT ProfileUnitySizes mode
+
 			PrivateIncludePaths.Add("GameplayAbilities/Private");
 			
 			PublicDependencyModuleNames.AddRange(
diff --git a/Engine/Plugins/Runtime/GeometryProcessing/Source/DynamicMesh/DynamicMesh.Build.cs b/Engine/Plugins/Runtime/GeometryProcessing/Source/DynamicMesh/DynamicMesh.Build.cs
--- a/Engine/Plugins/Runtime/GeometryProcessing/Source/DynamicMesh/DynamicMesh.Build.cs
+++ b/Engine/Plugins/Runtime/GeometryProcessing/Source/DynamicMesh/DynamicMesh.Build.cs
@@ -6,7 +6,9 @@ public class DynamicMesh : ModuleRules
 {	
 	public DynamicMesh(ReadOnlyTargetRules Target) : base(Target)
 	{
-        PCHUsage = ModuleRules.PCHUsageMode.UseExplicitOrSharedPCHs;
+		NumIncludedBytesPerUnityCPPOverride = 589824; // best unity size found from using UBT ProfileUnitySizes mode
+
+		PCHUsage = ModuleRules.PCHUsageMode.UseExplicitOrSharedPCHs;
 
         PublicDependencyModuleNames.AddRange(
 			new string[] {
diff --git a/Engine/Source/Runtime/Core/Core.Build.cs b/Engine/Source/Runtime/Core/Core.Build.cs
--- a/Engine/Source/Runtime/Core/Core.Build.cs
+++ b/Engine/Source/Runtime/Core/Core.Build.cs
@@ -8,6 +8,8 @@ public class Core : ModuleRules
 {
 	public Core(ReadOnlyTargetRules Target) : base(Target)
 	{
+		NumIncludedBytesPerUnityCPPOverride = 491520; // best unity size found from using UBT ProfileUnitySizes mode
+
 		PrivatePCHHeaderFile = "Private/CorePrivatePCH.h";
 
 		SharedPCHHeaderFile = "Public/CoreSharedPCH.h";
diff --git a/Engine/Source/Runtime/Engine/Engine.Build.cs b/Engine/Source/Runtime/Engine/Engine.Build.cs
--- a/Engine/Source/Runtime/Engine/Engine.Build.cs
+++ b/Engine/Source/Runtime/Engine/Engine.Build.cs
@@ -7,6 +7,8 @@ public class Engine : ModuleRules
 {
 	public Engine(ReadOnlyTargetRules Target) : base(Target)
 	{
+		NumIncludedBytesPerUnityCPPOverride = 589824; // best unity size found from using UBT ProfileUnitySizes mode
+
 		PublicIncludePaths.Add("../Shaders/Shared");
 		
 		PrivatePCHHeaderFile = "Private/EnginePrivatePCH.h";
