Optimus: fix a crash caused by anim attribute nodes. It was a bad idea to cache a pointer to an TArray element, which can be invalidated when the array expands in size

diff --git a/Engine/Plugins/Animation/DeformerGraph/Source/OptimusCore/Private/DataInterfaces/OptimusDataInterfaceAnimAttribute.cpp b/Engine/Plugins/Animation/DeformerGraph/Source/OptimusCore/Private/DataInterfaces/OptimusDataInterfaceAnimAttribute.cpp
--- a/Engine/Plugins/Animation/DeformerGraph/Source/OptimusCore/Private/DataInterfaces/OptimusDataInterfaceAnimAttribute.cpp
+++ b/Engine/Plugins/Animation/DeformerGraph/Source/OptimusCore/Private/DataInterfaces/OptimusDataInterfaceAnimAttribute.cpp
@@ -602,11 +602,8 @@ void UOptimusAnimAttributeDataProvider::Init(
 		
 		RuntimeData.Offset = Members[Index].GetOffset();
 
-		if (RuntimeData.ArrayMetadata)
-		{
-			RuntimeData.ArrayIndexStart = TotalNumArrays;
-			TotalNumArrays += RuntimeData.ArrayMetadata->Num();
-		}
+		RuntimeData.ArrayIndexStart = TotalNumArrays;
+		TotalNumArrays += RuntimeData.ArrayMetadata.Num();
 	}
 
 	AttributeBufferSize = ShaderParameterMetadata->GetSize();
@@ -634,10 +631,10 @@ FComputeDataProviderRenderProxy* UOptimusAnimAttributeDataProvider::GetRenderPro
 	{
 		const FOptimusAnimAttributeRuntimeData& AttributeData = AttributeRuntimeData[Index];
 
-		for (int32 ArrayIndex = 0; AttributeData.ArrayMetadata && ArrayIndex < AttributeData.ArrayMetadata->Num(); ArrayIndex++)
+		for (int32 ArrayIndex = 0; ArrayIndex < AttributeData.ArrayMetadata.Num(); ArrayIndex++)
 		{
 			int32 ToplevelArrayIndex = AttributeData.ArrayIndexStart + ArrayIndex;
-			const FOptimusDataTypeRegistry::FArrayMetadata& Metadata = (*AttributeData.ArrayMetadata)[ArrayIndex];
+			const FOptimusDataTypeRegistry::FArrayMetadata& Metadata = AttributeData.ArrayMetadata[ArrayIndex];
 			if (ensure(Proxy->AttributeArrayMetadata.IsValidIndex(ToplevelArrayIndex)))
 			{
 				Proxy->AttributeArrayMetadata[ToplevelArrayIndex].Offset = AttributeData.Offset + Metadata.ShaderValueOffset;
@@ -670,7 +667,7 @@ FComputeDataProviderRenderProxy* UOptimusAnimAttributeDataProvider::GetRenderPro
 					{ValuePtr, ValueSize},
 					{
 						{Proxy->AttributeBuffer.GetData() + AttributeData.Offset, AttributeData.Size},
-						{Proxy->AttributeArrayData.GetData() + AttributeData.ArrayIndexStart, AttributeData.ArrayMetadata ? AttributeData.ArrayMetadata->Num() : 0}	
+						{Proxy->AttributeArrayData.GetData() + AttributeData.ArrayIndexStart, AttributeData.ArrayMetadata.Num()}	
 					});
 			}
 		
@@ -682,20 +679,17 @@ FComputeDataProviderRenderProxy* UOptimusAnimAttributeDataProvider::GetRenderPro
 
 				FMemory::Memcpy(&Proxy->AttributeBuffer[AttributeData.Offset], DefaultValuePtr, DefaultValueSize);
 
-				if (AttributeData.ArrayMetadata)
+				if (ensure(AttributeData.ArrayMetadata.Num() == AttributeData.CachedDefaultValue.ArrayList.Num()))
 				{
-					if (ensure(AttributeData.ArrayMetadata->Num() == AttributeData.CachedDefaultValue.ArrayList.Num()))
+					for (int32 ArrayIndex = 0; ArrayIndex < AttributeData.CachedDefaultValue.ArrayList.Num(); ArrayIndex++)
 					{
-						for (int32 ArrayIndex = 0; ArrayIndex < AttributeData.CachedDefaultValue.ArrayList.Num(); ArrayIndex++)
+						const int32 ToplevelArrayIndex = AttributeData.ArrayIndexStart + ArrayIndex;
+						if (ensure(Proxy->AttributeArrayData.IsValidIndex(ToplevelArrayIndex)))
 						{
-							const int32 ToplevelArrayIndex = AttributeData.ArrayIndexStart + ArrayIndex;
-							if (ensure(Proxy->AttributeArrayData.IsValidIndex(ToplevelArrayIndex)))
-							{
-								Proxy->AttributeArrayData[ToplevelArrayIndex] = AttributeData.CachedDefaultValue.ArrayList[ArrayIndex];
-							}
+							Proxy->AttributeArrayData[ToplevelArrayIndex] = AttributeData.CachedDefaultValue.ArrayList[ArrayIndex];
 						}
-					}	
-				}
+					}
+				}	
 			}
 		}
 	}
diff --git a/Engine/Plugins/Animation/DeformerGraph/Source/OptimusCore/Private/DataInterfaces/OptimusDataInterfaceAnimAttribute.h b/Engine/Plugins/Animation/DeformerGraph/Source/OptimusCore/Private/DataInterfaces/OptimusDataInterfaceAnimAttribute.h
--- a/Engine/Plugins/Animation/DeformerGraph/Source/OptimusCore/Private/DataInterfaces/OptimusDataInterfaceAnimAttribute.h
+++ b/Engine/Plugins/Animation/DeformerGraph/Source/OptimusCore/Private/DataInterfaces/OptimusDataInterfaceAnimAttribute.h
@@ -165,7 +165,7 @@ struct FOptimusAnimAttributeRuntimeData
 
 	FOptimusDataTypeRegistry::PropertyValueConvertFuncT	ConvertFunc = nullptr;
 
-	const TArray<FOptimusDataTypeRegistry::FArrayMetadata>* ArrayMetadata = nullptr;
+	TArray<FOptimusDataTypeRegistry::FArrayMetadata> ArrayMetadata;
 
 	UScriptStruct* AttributeType = nullptr;
 
diff --git a/Engine/Plugins/Animation/DeformerGraph/Source/OptimusCore/Private/OptimusDataType.cpp b/Engine/Plugins/Animation/DeformerGraph/Source/OptimusCore/Private/OptimusDataType.cpp
--- a/Engine/Plugins/Animation/DeformerGraph/Source/OptimusCore/Private/OptimusDataType.cpp
+++ b/Engine/Plugins/Animation/DeformerGraph/Source/OptimusCore/Private/OptimusDataType.cpp
@@ -118,25 +118,20 @@ bool FOptimusDataType::CanCreateProperty() const
 
 int32 FOptimusDataType::GetNumArrays() const
 {
-	const TArray<FOptimusDataTypeRegistry::FArrayMetadata>* ArrayMetadata =
+	const TArray<FOptimusDataTypeRegistry::FArrayMetadata> ArrayMetadata =
 		FOptimusDataTypeRegistry::Get().FindArrayMetadata(TypeName);
 	
-	if (ensure(ArrayMetadata))
-	{
-		return ArrayMetadata->Num();
-	}
-
-	return 0;
+	return ArrayMetadata.Num();
 }
 
 int32 FOptimusDataType::GetArrayShaderValueOffset(int32 InArrayIndex) const
 {
-	const TArray<FOptimusDataTypeRegistry::FArrayMetadata>* ArrayMetadata =
+	const TArray<FOptimusDataTypeRegistry::FArrayMetadata> ArrayMetadata =
 		FOptimusDataTypeRegistry::Get().FindArrayMetadata(TypeName);
 
-	if (ensure(ArrayMetadata) && ensure(ArrayMetadata->IsValidIndex(InArrayIndex)))
+	if (ensure(ArrayMetadata.IsValidIndex(InArrayIndex)))
 	{
-		return (*ArrayMetadata)[InArrayIndex].ShaderValueOffset;
+		return ArrayMetadata[InArrayIndex].ShaderValueOffset;
 	}
 
 	return INDEX_NONE;
@@ -144,12 +139,12 @@ int32 FOptimusDataType::GetArrayShaderValueOffset(int32 InArrayIndex) const
 
 int32 FOptimusDataType::GetArrayElementShaderValueSize(int32 InArrayIndex) const
 {
-	const TArray<FOptimusDataTypeRegistry::FArrayMetadata>* ArrayMetadata =
+	const TArray<FOptimusDataTypeRegistry::FArrayMetadata> ArrayMetadata =
 		FOptimusDataTypeRegistry::Get().FindArrayMetadata(TypeName);
 
-	if (ensure(ArrayMetadata && ArrayMetadata->IsValidIndex(InArrayIndex)))
+	if (ensure(ArrayMetadata.IsValidIndex(InArrayIndex)))
 	{
-		return (*ArrayMetadata)[InArrayIndex].ElementShaderValueSize;
+		return ArrayMetadata[InArrayIndex].ElementShaderValueSize;
 	}
 
 	return INDEX_NONE;
diff --git a/Engine/Plugins/Animation/DeformerGraph/Source/OptimusCore/Private/OptimusDataTypeRegistry.cpp b/Engine/Plugins/Animation/DeformerGraph/Source/OptimusCore/Private/OptimusDataTypeRegistry.cpp
--- a/Engine/Plugins/Animation/DeformerGraph/Source/OptimusCore/Private/OptimusDataTypeRegistry.cpp
+++ b/Engine/Plugins/Animation/DeformerGraph/Source/OptimusCore/Private/OptimusDataTypeRegistry.cpp
@@ -1436,17 +1436,17 @@ FOptimusDataTypeRegistry::PropertyValueConvertFuncT FOptimusDataTypeRegistry::Fi
 	return InfoPtr->PropertyValueConvertFunc;
 }
 
-const TArray<FOptimusDataTypeRegistry::FArrayMetadata>* FOptimusDataTypeRegistry::FindArrayMetadata(
+TArray<FOptimusDataTypeRegistry::FArrayMetadata> FOptimusDataTypeRegistry::FindArrayMetadata(
 	FName InTypeName) const
 {
 	const FTypeInfo* InfoPtr = RegisteredTypes.Find(InTypeName);
 	if (!InfoPtr)
 	{
 		UE_LOG(LogOptimusCore, Fatal, TEXT("CreateProperty: Invalid type name."));
-		return nullptr;
+		return {};
 	}
 
-	return &InfoPtr->ArrayMetadata;
+	return InfoPtr->ArrayMetadata;
 }
 
 UScriptStruct* FOptimusDataTypeRegistry::FindAttributeType(FName InTypeName) const
diff --git a/Engine/Plugins/Animation/DeformerGraph/Source/OptimusCore/Public/OptimusDataTypeRegistry.h b/Engine/Plugins/Animation/DeformerGraph/Source/OptimusCore/Public/OptimusDataTypeRegistry.h
--- a/Engine/Plugins/Animation/DeformerGraph/Source/OptimusCore/Public/OptimusDataTypeRegistry.h
+++ b/Engine/Plugins/Animation/DeformerGraph/Source/OptimusCore/Public/OptimusDataTypeRegistry.h
@@ -165,7 +165,7 @@ public:
 	PropertyValueConvertFuncT FindPropertyValueConvertFunc(FName InTypeName) const;
 
 	/** A helper function to return the array metadata. */
-	const TArray<FArrayMetadata>* FindArrayMetadata(FName InTypeName) const;
+	TArray<FArrayMetadata> FindArrayMetadata(FName InTypeName) const;
 
 	/** A helper function to return the corresponding animation attribute type. */
 	UScriptStruct* FindAttributeType(FName InTypeName) const;
