Optimus: changed to pack structuredbuffer<float3> as structuredbuffer <float4>

diff --git a/Engine/Plugins/Animation/DeformerGraph/Source/OptimusCore/Private/OptimusDataTypeRegistry.cpp b/Engine/Plugins/Animation/DeformerGraph/Source/OptimusCore/Private/OptimusDataTypeRegistry.cpp
--- a/Engine/Plugins/Animation/DeformerGraph/Source/OptimusCore/Private/OptimusDataTypeRegistry.cpp
+++ b/Engine/Plugins/Animation/DeformerGraph/Source/OptimusCore/Private/OptimusDataTypeRegistry.cpp
@@ -575,6 +575,15 @@ bool FOptimusDataTypeRegistry::RegisterStructType(UScriptStruct* InStructType)
 			}
 
 			DataType = FindType(*PropertyForValidation);
+
+			// Special logic for things like StructuredBuffer<float3> which should be packed as buffer of float4 in Vulkan
+			if (bIsArrayMember)
+			{
+				if (DataType->ShaderValueType->Type != EShaderFundamentalType::Struct && DataType->ShaderValueType->VectorElemCount == 3)
+				{
+					DataType = FindType(FShaderValueType::Get(DataType->ShaderValueType->Type, 4));
+				}
+			}
 			
 			if (DataType->GetNumArrays() > 0)
 			{
@@ -960,20 +969,25 @@ bool FOptimusDataTypeRegistry::RegisterType(
 				FShaderValueType::FValueView OutShaderValue
 				) -> bool
 			{
-				if (ensure(InRawValue.Num() == ExpectedPropertySize) &&
+				// we can be copying a smaller property into a larger shader side array element
+				// for example, FVector3 -> StructuredBuffer<float4>, see special logic in RegisterStructType
+				if (ensure(InRawValue.Num() <= ExpectedPropertySize) && 
 					ensure(OutShaderValue.ShaderValue.Num() == ExpectedShaderValueSize))
 				{
 					uint8* ShaderValuePtr = OutShaderValue.ShaderValue.GetData();
 					for (const FPropertyConversionInfo& ConversionInfo: ConversionEntries)
 					{
-						TArrayView<const uint8> PropertyData(InRawValue.GetData() + ConversionInfo.PropertyOffset, ConversionInfo.PropertySize);
-						TArrayView<uint8> ShaderValueData(ShaderValuePtr, ConversionInfo.ShaderValueSize);
-						if (!ConversionInfo.ConversionFunc(PropertyData, ShaderValueData))
+						if (ConversionInfo.PropertyOffset + ConversionInfo.PropertySize <= InRawValue.Num())
 						{
-							return false;
-						}
+							TArrayView<const uint8> PropertyData(InRawValue.GetData() + ConversionInfo.PropertyOffset, ConversionInfo.PropertySize);
+							TArrayView<uint8> ShaderValueData(ShaderValuePtr, ConversionInfo.ShaderValueSize);
+							if (!ConversionInfo.ConversionFunc(PropertyData, ShaderValueData))
+							{
+								return false;
+							}
 
-						ShaderValuePtr += ConversionInfo.ShaderValueSize;
+							ShaderValuePtr += ConversionInfo.ShaderValueSize;
+						}
 					}
 					return true;
 				}
