Fix linux build error.

diff --git a/Engine/Plugins/Animation/DeformerGraph/Source/OptimusSettings/OptimusSettings.Build.cs b/Engine/Plugins/Animation/DeformerGraph/Source/OptimusSettings/OptimusSettings.Build.cs
--- a/Engine/Plugins/Animation/DeformerGraph/Source/OptimusSettings/OptimusSettings.Build.cs
+++ b/Engine/Plugins/Animation/DeformerGraph/Source/OptimusSettings/OptimusSettings.Build.cs
@@ -27,6 +27,7 @@ namespace UnrealBuildTool.Rules
 					"DeveloperSettings",
 					"Engine",
 					"RenderCore",
+					"RHI",
 				}
 			);
         }
diff --git a/Engine/Plugins/Animation/DeformerGraph/Source/OptimusSettings/Public/OptimusSettings.h b/Engine/Plugins/Animation/DeformerGraph/Source/OptimusSettings/Public/OptimusSettings.h
--- a/Engine/Plugins/Animation/DeformerGraph/Source/OptimusSettings/Public/OptimusSettings.h
+++ b/Engine/Plugins/Animation/DeformerGraph/Source/OptimusSettings/Public/OptimusSettings.h
@@ -4,6 +4,7 @@
 
 #include "CoreMinimal.h"
 #include "Engine/DeveloperSettings.h"
+#include "RHIDefinitions.h"
 #include "OptimusSettings.generated.h"
 
 class UMeshDeformer;
