Remove IMeshDeformerInterface

diff --git a/Engine/Plugins/Animation/DeformerGraph/Source/OptimusCore/Private/DataInterfaces/OptimusDataInterfaceGraph.cpp b/Engine/Plugins/Animation/DeformerGraph/Source/OptimusCore/Private/DataInterfaces/OptimusDataInterfaceGraph.cpp
--- a/Engine/Plugins/Animation/DeformerGraph/Source/OptimusCore/Private/DataInterfaces/OptimusDataInterfaceGraph.cpp
+++ b/Engine/Plugins/Animation/DeformerGraph/Source/OptimusCore/Private/DataInterfaces/OptimusDataInterfaceGraph.cpp
@@ -3,13 +3,12 @@
 #include "OptimusDataInterfaceGraph.h"
 
 #include "Components/SkinnedMeshComponent.h"
+#include "ComputeFramework/ComputeMetadataBuilder.h"
 #include "ComputeFramework/ShaderParameterMetadataAllocation.h"
 #include "ComputeFramework/ShaderParamTypeDefinition.h"
 #include "OptimusDeformerInstance.h"
 #include "OptimusVariableDescription.h"
 #include "ShaderParameterMetadataBuilder.h"
-#include "ComputeFramework/ComputeMetadataBuilder.h"
-#include "MeshDeformerInterface.h"
 
 #include UE_INLINE_GENERATED_CPP_BY_NAME(OptimusDataInterfaceGraph)
 
@@ -117,12 +116,7 @@ void UOptimusGraphDataProvider::SetConstant(FString const& InVariableName, TArra
 
 FComputeDataProviderRenderProxy* UOptimusGraphDataProvider::GetRenderProxy()
 {
-	if (IMeshDeformerInterface* MeshDeformerInterface = Cast<IMeshDeformerInterface>(MeshComponent))
-	{
-		UOptimusDeformerInstance const* DeformerInstance = Cast<UOptimusDeformerInstance>(MeshDeformerInterface->GetMeshDeformerInstance());
-		return new FOptimusGraphDataProviderProxy(DeformerInstance, Variables, ParameterBufferSize);
-	}
-	return nullptr;
+	return new FOptimusGraphDataProviderProxy(DeformerInstance, Variables, ParameterBufferSize);
 }
 
 FOptimusGraphDataProviderProxy::FOptimusGraphDataProviderProxy(UOptimusDeformerInstance const* DeformerInstance, TArray<FOptimusGraphVariableDescription> const& Variables, int32 ParameterBufferSize)
diff --git a/Engine/Plugins/Animation/DeformerGraph/Source/OptimusCore/Private/DataInterfaces/OptimusDataInterfaceGraph.h b/Engine/Plugins/Animation/DeformerGraph/Source/OptimusCore/Private/DataInterfaces/OptimusDataInterfaceGraph.h
--- a/Engine/Plugins/Animation/DeformerGraph/Source/OptimusCore/Private/DataInterfaces/OptimusDataInterfaceGraph.h
+++ b/Engine/Plugins/Animation/DeformerGraph/Source/OptimusCore/Private/DataInterfaces/OptimusDataInterfaceGraph.h
@@ -64,6 +64,9 @@ public:
 	UPROPERTY(EditAnywhere, BlueprintReadWrite, Category = Binding)
 	TObjectPtr<UMeshComponent> MeshComponent = nullptr;
 
+	UPROPERTY()
+	TObjectPtr<UOptimusDeformerInstance> DeformerInstance = nullptr;
+
 	UPROPERTY()
 	TArray<FOptimusGraphVariableDescription> Variables;
 
diff --git a/Engine/Plugins/Animation/DeformerGraph/Source/OptimusCore/Private/OptimusDeformerInstance.cpp b/Engine/Plugins/Animation/DeformerGraph/Source/OptimusCore/Private/OptimusDeformerInstance.cpp
--- a/Engine/Plugins/Animation/DeformerGraph/Source/OptimusCore/Private/OptimusDeformerInstance.cpp
+++ b/Engine/Plugins/Animation/DeformerGraph/Source/OptimusCore/Private/OptimusDeformerInstance.cpp
@@ -6,13 +6,13 @@
 #include "ComputeFramework/ComputeFramework.h"
 #include "ComputeWorkerInterface.h"
 #include "DataInterfaces/OptimusDataInterfaceGraph.h"
+#include "DataInterfaces/OptimusDataInterfaceRawBuffer.h"
 #include "OptimusComputeGraph.h"
 #include "OptimusDataTypeRegistry.h"
 #include "OptimusDeformer.h"
 #include "OptimusVariableDescription.h"
 #include "RenderGraphBuilder.h"
 #include "Rendering/SkeletalMeshRenderData.h"
-#include "DataInterfaces/OptimusDataInterfaceRawBuffer.h"
 
 #include UE_INLINE_GENERATED_CPP_BY_NAME(OptimusDeformerInstance)
 
@@ -348,6 +348,12 @@ void UOptimusDeformerInstance::SetupFromDeformer(UOptimusDeformer* InDeformer)
 			{
 				PersistentBufferProvider->BufferPool = BufferPool;
 			}
+
+			// Set this instance on the graph data provider so that it can query variables.
+			if (UOptimusGraphDataProvider* GraphProvider = Cast<UOptimusGraphDataProvider>(DataProvider))
+			{
+				GraphProvider->DeformerInstance = this;
+			}
 		}
 
 		// Schedule the setup graph to run.
diff --git a/Engine/Plugins/Runtime/HairStrands/Source/HairStrandsCore/Public/GroomComponent.h b/Engine/Plugins/Runtime/HairStrands/Source/HairStrandsCore/Public/GroomComponent.h
--- a/Engine/Plugins/Runtime/HairStrands/Source/HairStrandsCore/Public/GroomComponent.h
+++ b/Engine/Plugins/Runtime/HairStrands/Source/HairStrandsCore/Public/GroomComponent.h
@@ -11,7 +11,6 @@
 #include "RHIDefinitions.h"
 #include "GroomDesc.h"
 #include "LODSyncInterface.h"
-#include "MeshDeformerInterface.h"
 #include "GroomInstance.h"
 #include "NiagaraDataInterfacePhysicsAsset.h"
 
@@ -23,7 +22,7 @@ class UMeshDeformerInstance;
 class UMeshDeformerInstanceSettings;
 
 UCLASS(HideCategories = (Object, Physics, Activation, Mobility, "Components|Activation"), editinlinenew, meta = (BlueprintSpawnableComponent), ClassGroup = Rendering)
-class HAIRSTRANDSCORE_API UGroomComponent : public UMeshComponent, public ILODSyncInterface, public INiagaraPhysicsAssetDICollectorInterface, public IMeshDeformerInterface
+class HAIRSTRANDSCORE_API UGroomComponent : public UMeshComponent, public ILODSyncInterface, public INiagaraPhysicsAssetDICollectorInterface
 {
 	GENERATED_UCLASS_BODY()
 
@@ -240,10 +239,6 @@ public:
 	virtual int32 GetCurrentSyncLOD() const override;
 	//~ End ILODSyncInterface
 
-	///~ Begin IMeshDeformerInterface Interface.
-	virtual UMeshDeformerInstance const* GetMeshDeformerInstance() const override { return MeshDeformerInstance; }
-	//~ End IMeshDeformerInterface
-
 	int32 GetNumLODs() const;
 	int32 GetForcedLOD() const;
 	void  SetForcedLOD(int32 LODIndex);
diff --git a/Engine/Source/Runtime/Engine/Classes/Components/SkinnedMeshComponent.h b/Engine/Source/Runtime/Engine/Classes/Components/SkinnedMeshComponent.h
--- a/Engine/Source/Runtime/Engine/Classes/Components/SkinnedMeshComponent.h
+++ b/Engine/Source/Runtime/Engine/Classes/Components/SkinnedMeshComponent.h
@@ -18,7 +18,6 @@
 #include "Components/MeshComponent.h"
 #include "Containers/SortedMap.h"
 #include "LODSyncInterface.h"
-#include "MeshDeformerInterface.h"
 #include "BoneContainer.h"
 #include "ClothingSystemRuntimeTypes.h"
 #include "SkinnedMeshComponent.generated.h"
@@ -193,7 +192,7 @@ using FExternalMorphSets = TMap<int32, TSharedPtr<FExternalMorphSet>>;
  * @see USkeletalMeshComponent
 */
 UCLASS(hidecategories=Object, config=Engine, editinlinenew, abstract)
-class ENGINE_API USkinnedMeshComponent : public UMeshComponent, public ILODSyncInterface, public IMeshDeformerInterface
+class ENGINE_API USkinnedMeshComponent : public UMeshComponent, public ILODSyncInterface
 {
 	GENERATED_UCLASS_BODY()
 
@@ -258,12 +257,12 @@ protected:
 	TObjectPtr<UMeshDeformerInstanceSettings> MeshDeformerInstanceSettings;
 
 	/** Object containing state for the bound MeshDeformer. */
-	UPROPERTY(Transient, BlueprintReadOnly, Category = "Deformer")
+	UPROPERTY(Transient, BlueprintReadOnly, BlueprintGetter = GetMeshDeformerInstance, Category = "Deformer")
 	TObjectPtr<UMeshDeformerInstance> MeshDeformerInstance;
 
 public:
-	/** Get the currently active MeshDeformer Instance. */
-	UMeshDeformerInstance const* GetMeshDeformerInstance() const override { return MeshDeformerInstance; }
+	UFUNCTION(BlueprintGetter)
+	UMeshDeformerInstance* GetMeshDeformerInstance() const { return MeshDeformerInstance; }
 
 	/** Max LOD at which to update or apply the MeshDeformer. */
 	int32 GetMeshDeformerMaxLOD() const;
diff --git a/Engine/Source/Runtime/Engine/Private/Engine/MeshDeformerInterface.cpp b/Engine/Source/Runtime/Engine/Private/Engine/MeshDeformerInterface.cpp
deleted file mode 100644
diff --git a/Engine/Source/Runtime/Engine/Public/MeshDeformerInterface.h b/Engine/Source/Runtime/Engine/Public/MeshDeformerInterface.h
deleted file mode 100644
