Optimus: Fix bug in skeleton data interface with unlimited weights.

diff --git a/Engine/Plugins/Animation/DeformerGraph/Shaders/Private/DataInterfaceSkeleton.ush b/Engine/Plugins/Animation/DeformerGraph/Shaders/Private/DataInterfaceSkeleton.ush
--- a/Engine/Plugins/Animation/DeformerGraph/Shaders/Private/DataInterfaceSkeleton.ush
+++ b/Engine/Plugins/Animation/DeformerGraph/Shaders/Private/DataInterfaceSkeleton.ush
@@ -172,16 +172,16 @@ float3x4 ReadWeightedBoneMatrix_{DataInterfaceName}(uint VertexIndex)
 	{
 		int BoneIndexOffset = IndicesOffset + (BoneIndexSize * InfluenceIndex);
 		int BoneIndex = {DataInterfaceName}_InputWeightStream[BoneIndexOffset];
-		if ({DataInterfaceName}_InputWeightIndexSize > 1)
+		if (BoneIndexSize > 1)
 		{
 			BoneIndex = {DataInterfaceName}_InputWeightStream[BoneIndexOffset + 1] << 8 | BoneIndex;
 		}
 		int BoneWeightOffset = WeightsOffset + (BoneWeightSize * InfluenceIndex);
-		float BoneWeight;
+		float BoneWeight = float({ DataInterfaceName }_InputWeightStream[BoneWeightOffset]) / 255.0f;
 		if (BoneWeightSize > 1)
+		{
 			BoneWeight = float({DataInterfaceName}_InputWeightStream[BoneWeightOffset + 1] << 8 | {DataInterfaceName}_InputWeightStream[BoneWeightOffset]) / 65535.0;
-		else
-			BoneWeight = float({DataInterfaceName}_InputWeightStream[BoneWeightOffset]) / 255.0f;
+		}
 		BlendMatrix += BoneWeight * GetBoneMatrix_{DataInterfaceName}(BoneIndex);
 	}
 
