Optimus: Fix LOD calculation used for setting up transient buffers.
Need to use predicted LOD so that game/render thread calculations match.
The render thread calculation that we need to match is calculated in the FOptimusSkinnedMeshExecDataProviderProxy::GetDispatchThreadCount.
Added some additional ensure checks to validate.

diff --git a/Engine/Plugins/Animation/DeformerGraph/Source/OptimusCore/Private/ComponentSources/OptimusSkinnedMeshComponentSource.cpp b/Engine/Plugins/Animation/DeformerGraph/Source/OptimusCore/Private/ComponentSources/OptimusSkinnedMeshComponentSource.cpp
--- a/Engine/Plugins/Animation/DeformerGraph/Source/OptimusCore/Private/ComponentSources/OptimusSkinnedMeshComponentSource.cpp
+++ b/Engine/Plugins/Animation/DeformerGraph/Source/OptimusCore/Private/ComponentSources/OptimusSkinnedMeshComponentSource.cpp
@@ -38,8 +38,7 @@ TArray<FName> UOptimusSkinnedMeshComponentSource::GetExecutionDomains() const
 int32 UOptimusSkinnedMeshComponentSource::GetLodIndex(const UActorComponent* InComponent) const
 {
 	const USkinnedMeshComponent* SkinnedMeshComponent = Cast<USkinnedMeshComponent>(InComponent);
-	const FSkeletalMeshObject* SkeletalMeshObject = SkinnedMeshComponent ? SkinnedMeshComponent->MeshObject : nullptr;
-	return SkeletalMeshObject ? SkeletalMeshObject->GetLOD() : 0;
+	return SkinnedMeshComponent ? SkinnedMeshComponent->GetPredictedLODLevel() : 0;
 }
 
 bool UOptimusSkinnedMeshComponentSource::GetComponentElementCountsForExecutionDomain(
diff --git a/Engine/Plugins/Animation/DeformerGraph/Source/OptimusCore/Private/DataInterfaces/OptimusDataInterfaceRawBuffer.cpp b/Engine/Plugins/Animation/DeformerGraph/Source/OptimusCore/Private/DataInterfaces/OptimusDataInterfaceRawBuffer.cpp
--- a/Engine/Plugins/Animation/DeformerGraph/Source/OptimusCore/Private/DataInterfaces/OptimusDataInterfaceRawBuffer.cpp
+++ b/Engine/Plugins/Animation/DeformerGraph/Source/OptimusCore/Private/DataInterfaces/OptimusDataInterfaceRawBuffer.cpp
@@ -386,6 +386,11 @@ void FOptimusTransientBufferDataProviderProxy::GatherDispatchData(FDispatchSetup
 		return;
 	}
 
+	if (!ensure(InDispatchSetup.NumInvocations == Buffer.Num()))
+	{
+		return;
+	}
+
 	for (int32 InvocationIndex = 0; InvocationIndex < InvocationElementCounts.Num(); ++InvocationIndex)
 	{
 		FTransientBufferDataInterfaceParameters* Parameters =
@@ -422,6 +427,7 @@ void FOptimusPersistentBufferDataProviderProxy::AllocateResources(
 	)
 {
 	BufferPool->GetResourceBuffers(GraphBuilder, ResourceName, LODIndex, ElementStride, RawStride, InvocationElementCounts, Buffers);
+	ensure(Buffers.Num() == InvocationElementCounts.Num());
 	BufferUAVs.Reserve(Buffers.Num());
 	for (FRDGBufferRef BufferRef : Buffers)
 	{
@@ -440,11 +446,11 @@ void FOptimusPersistentBufferDataProviderProxy::GatherDispatchData(
 		return;
 	}
 
-	if (!ensure(Buffers.Num() == InvocationElementCounts.Num()))
+	if (!ensure(InDispatchSetup.NumInvocations == Buffers.Num()))
 	{
 		return;
 	}
-	
+
 	for (int32 InvocationIndex = 0; InvocationIndex < InvocationElementCounts.Num(); ++InvocationIndex)
 	{
 		FPersistentBufferDataInterfaceParameters* Parameters =
