Optimus: display name of kernel nodes can now change according to kernel name + fixed broken undo/redo

diff --git a/Engine/Plugins/Animation/DeformerGraph/Source/OptimusCore/Private/Nodes/OptimusNode_CustomComputeKernel.cpp b/Engine/Plugins/Animation/DeformerGraph/Source/OptimusCore/Private/Nodes/OptimusNode_CustomComputeKernel.cpp
--- a/Engine/Plugins/Animation/DeformerGraph/Source/OptimusCore/Private/Nodes/OptimusNode_CustomComputeKernel.cpp
+++ b/Engine/Plugins/Animation/DeformerGraph/Source/OptimusCore/Private/Nodes/OptimusNode_CustomComputeKernel.cpp
@@ -2,6 +2,7 @@
 
 #include "OptimusNode_CustomComputeKernel.h"
 
+#include "OptimusActionStack.h"
 #include "OptimusComponentSource.h"
 #include "DataInterfaces/OptimusDataInterfaceRawBuffer.h"
 #include "OptimusDataTypeRegistry.h"
@@ -10,6 +11,7 @@
 #include "OptimusNodePin.h"
 #include "OptimusNode_DataInterface.h"
 #include "OptimusObjectVersion.h"
+#include "Actions/OptimusNodeActions.h"
 #include "ComponentSources/OptimusSkeletalMeshComponentSource.h"
 #include "ComponentSources/OptimusSkinnedMeshComponentSource.h"
 #include "DataInterfaces/OptimusDataInterfaceSkinnedMeshExec.h"
@@ -345,9 +347,12 @@ void UOptimusNode_CustomComputeKernel::PropertyValueChanged(
 	};
 	
 	
-	if (InPropertyChangedEvent.GetPropertyName() == GET_MEMBER_NAME_STRING_CHECKED(UOptimusNode_CustomComputeKernel, KernelName))
+	if (InPropertyChangedEvent.GetMemberPropertyName() == GET_MEMBER_NAME_STRING_CHECKED(UOptimusNode_CustomComputeKernel, KernelName) && 
+		InPropertyChangedEvent.GetPropertyName() == GET_MEMBER_NAME_STRING_CHECKED(FOptimusValidatedName, Name))
 	{
-		SetDisplayName(FText::FromName(KernelName));
+		// Making sure the display name(NonTransactional) is updated as KernelName(Transactional) changes
+		GetActionStack()->RunAction<FOptimusNodeAction_RenameNode>(this, FText::FromName(KernelName));
+		
 		UpdatePreamble();
 	}
 	else if (InPropertyChangedEvent.GetMemberPropertyName() == GET_MEMBER_NAME_STRING_CHECKED(UOptimusNode_CustomComputeKernel, SecondaryInputBindingGroups) &&
@@ -964,12 +969,15 @@ PRAGMA_ENABLE_DEPRECATION_WARNINGS
 			}
 		}
 	}
-	
+
+	SetDisplayName(FText::FromName(KernelName));
 }
 
 
 void UOptimusNode_CustomComputeKernel::ConstructNode()
 {
+	SetDisplayName(FText::FromName(KernelName));
+	
 	// After a duplicate, the kernel node has no pins, so we need to reconstruct them from
 	// the bindings. We can assume that all naming clashes have already been dealt with.
 	for (const FOptimusParameterBinding& Binding: InputBindingArray)
