Optimus: node name changes/deletion is now synced with shader text editor

diff --git a/Engine/Plugins/Animation/DeformerGraph/Source/OptimusEditor/Private/OptimusEditor.cpp b/Engine/Plugins/Animation/DeformerGraph/Source/OptimusEditor/Private/OptimusEditor.cpp
--- a/Engine/Plugins/Animation/DeformerGraph/Source/OptimusEditor/Private/OptimusEditor.cpp
+++ b/Engine/Plugins/Animation/DeformerGraph/Source/OptimusEditor/Private/OptimusEditor.cpp
@@ -626,6 +626,34 @@ bool FOptimusEditor::CanExpandCollapsedNode() const
 	return !ModelNodes.IsEmpty();
 }
 
+void FOptimusEditor::SubscribeToGraphNotifies(UOptimusNodeGraph* Graph, const FOnGraphNotified& OnGraphNotifiedDelegate)
+{
+	if (IsValid(Graph))
+	{
+		TSet<TWeakObjectPtr<UOptimusNodeGraph>>& Graphs = GraphSubscribers.FindOrAdd(OnGraphNotifiedDelegate.GetHandle());
+		if (!Graphs.Contains(Graph))
+		{
+			Graphs.Add(Graph);
+			Graph->GetNotifyDelegate().Add(OnGraphNotifiedDelegate);
+		}
+	}
+}
+
+void FOptimusEditor::UnsubscribeToGraphNotifies(FDelegateHandle DelegateHandle)
+{
+	if (TSet<TWeakObjectPtr<UOptimusNodeGraph>>* Graphs = GraphSubscribers.Find(DelegateHandle))
+	{
+		for (TWeakObjectPtr<UOptimusNodeGraph> Graph : *Graphs)
+		{
+			if (Graph.IsValid())
+			{
+				Graph->GetNotifyDelegate().Remove(DelegateHandle);
+			}
+		}
+			
+		GraphSubscribers.Remove(DelegateHandle);
+	}
+}
 
 void FOptimusEditor::OnSelectedNodesChanged(const TSet<UObject*>& NewSelection)
 {
diff --git a/Engine/Plugins/Animation/DeformerGraph/Source/OptimusEditor/Private/OptimusEditor.h b/Engine/Plugins/Animation/DeformerGraph/Source/OptimusEditor/Private/OptimusEditor.h
--- a/Engine/Plugins/Animation/DeformerGraph/Source/OptimusEditor/Private/OptimusEditor.h
+++ b/Engine/Plugins/Animation/DeformerGraph/Source/OptimusEditor/Private/OptimusEditor.h
@@ -3,6 +3,8 @@
 #pragma once
 
 #include "IOptimusEditor.h"
+#include "OptimusCoreNotify.h"
+#include "OptimusNodeGraph.h"
 
 #include "Misc/NotifyHook.h"
 
@@ -30,6 +32,7 @@ class USkeletalMesh;
 enum class EOptimusGlobalNotifyType;
 struct FGraphAppearanceInfo;
 struct FOptimusCompilerDiagnostic;
+enum class EOptimusGraphNotifyType;
 
 
 class FOptimusEditor :
@@ -168,6 +171,12 @@ private:
 	FReply OnSpawnGraphNodeByShortcut(FInputChord InChord, const FVector2D& InPosition, UEdGraph* InGraph);
 
 public:
+	// ----------------------------------------------------------------------------------------
+	// Utility for other widgets to subscribe to Graph Notifies 
+	DECLARE_DELEGATE_ThreeParams(FOnGraphNotified,  EOptimusGraphNotifyType /* type */, UOptimusNodeGraph*/*graph */, UObject* /* subject */);
+	void SubscribeToGraphNotifies(UOptimusNodeGraph* Graph, const FOnGraphNotified& OnGraphNotifiedDelegate);
+	void UnsubscribeToGraphNotifies(FDelegateHandle DelegateHandle);
+
 	// ----------------------------------------------------------------------------------------
 	// Broadcast Graph Event
 	DECLARE_EVENT_OneParam(FOptimusEditor, FOnSelectedNodesChanged, const TArray<TWeakObjectPtr<UObject>>&);
@@ -255,4 +264,6 @@ private:
 	FOnRefreshEvent RefreshEvent;
 	FOnSelectedNodesChanged SelectedNodesChangedEvent;
 	FOnDiagnosticsUpdated DiagnosticsUpdatedEvent;
+	
+	TMap<FDelegateHandle, TSet<TWeakObjectPtr<UOptimusNodeGraph>>> GraphSubscribers;
 };
diff --git a/Engine/Plugins/Animation/DeformerGraph/Source/OptimusEditor/Private/Widgets/SOptimusShaderTextEditor.cpp b/Engine/Plugins/Animation/DeformerGraph/Source/OptimusEditor/Private/Widgets/SOptimusShaderTextEditor.cpp
--- a/Engine/Plugins/Animation/DeformerGraph/Source/OptimusEditor/Private/Widgets/SOptimusShaderTextEditor.cpp
+++ b/Engine/Plugins/Animation/DeformerGraph/Source/OptimusEditor/Private/Widgets/SOptimusShaderTextEditor.cpp
@@ -3,6 +3,7 @@
 #include "SOptimusShaderTextEditor.h"
 #include "SOptimusShaderTextDocumentTab.h"
 #include "OptimusEditor.h"
+#include "OptimusNode.h"
 
 #include "IOptimusShaderTextProvider.h"
 
@@ -77,6 +78,10 @@ TArray<FName> SOptimusShaderTextEditor::GetAllTabIds()
 
 SOptimusShaderTextEditor::~SOptimusShaderTextEditor()
 {
+	if (OwningEditor.Pin())
+	{
+		OwningEditor.Pin()->UnsubscribeToGraphNotifies(OnGraphNotifyDelegate.GetHandle());
+	}	
 }
 
 void SOptimusShaderTextEditor::Construct(const FArguments& InArgs, TWeakPtr<FOptimusEditor> InEditor, const TSharedRef<SDockTab>& InShaderTextEditorHostTab)
@@ -120,6 +125,8 @@ void SOptimusShaderTextEditor::Construct(const FArguments& InArgs, TWeakPtr<FOpt
 	{
 		OwningEditor.Pin()->OnSelectedNodesChanged().AddSP(this, &SOptimusShaderTextEditor::HandleSelectedNodesChanged);
 	}
+	
+	OnGraphNotifyDelegate.BindSP(this, &SOptimusShaderTextEditor::HandleGraphNotified);
 }
 
 void SOptimusShaderTextEditor::OnHostTabClosed(TSharedRef<SDockTab> InShaderTextEditorHostTab)
@@ -148,15 +155,52 @@ void SOptimusShaderTextEditor::HandleSelectedNodesChanged(const TArray<TWeakObje
 {
 	for (const TWeakObjectPtr<UObject>& Object : InSelectedObjects)
 	{
+		TSharedRef<FTabPayload_UObject> Payload = FTabPayload_UObject::Make(Object.Get());
+
+		if (TSharedPtr<SDockTab> OpenedTab = DocumentTracker->OpenDocument(Payload, FDocumentTracker::OpenNewDocument))
 		{
-			TSharedRef<FTabPayload_UObject> Payload = FTabPayload_UObject::Make(Object.Get());
-			
-			if (TabManager->GetOwnerTab()->IsForeground())
+			if (UOptimusNode* Node = Cast<UOptimusNode>(Object))
 			{
-				DocumentTracker->OpenDocument(Payload, FDocumentTracker::OpenNewDocument);
+				if (ensure(OwningEditor.Pin()))
+				{
+					// Register our handler with the owning graph of each kernel node
+					// OwningEditor should detect duplicated registration 
+					OwningEditor.Pin()->SubscribeToGraphNotifies(Node->GetOwningGraph(), OnGraphNotifyDelegate);
+				}
 			}
 		}
 	}
 }
 
+void SOptimusShaderTextEditor::HandleGraphNotified(EOptimusGraphNotifyType InNotifyType, UOptimusNodeGraph* InGraph, UObject* InSubject) const
+{
+	switch (InNotifyType)
+	{
+		case EOptimusGraphNotifyType::NodeRemoved:
+		{
+			TSharedRef<FTabPayload_UObject> Payload = FTabPayload_UObject::Make(InSubject);
+			DocumentTracker->CloseTab(Payload);
+			break;
+		}
+		case EOptimusGraphNotifyType::NodeDisplayNameChanged:
+		{
+			if (IOptimusShaderTextProvider* ShaderTextProvider = Cast<IOptimusShaderTextProvider>(InSubject))
+			{
+				TSharedRef<FTabPayload_UObject> Payload = FTabPayload_UObject::Make(InSubject);
+				TArray<TSharedPtr<SDockTab>> Results;
+				DocumentTracker->FindMatchingTabs(Payload,Results);
+				for (TSharedPtr<SDockTab> Tab : Results)
+				{
+					Tab->SetLabel(FText::FromString(ShaderTextProvider->GetNameForShaderTextEditor()));
+				}
+			}
+			break;
+		}
+		default:
+		{
+			break;
+		}
+	}
+}
+
 #undef LOCTEXT_NAMESPACE
\ No newline at end of file
diff --git a/Engine/Plugins/Animation/DeformerGraph/Source/OptimusEditor/Private/Widgets/SOptimusShaderTextEditor.h b/Engine/Plugins/Animation/DeformerGraph/Source/OptimusEditor/Private/Widgets/SOptimusShaderTextEditor.h
--- a/Engine/Plugins/Animation/DeformerGraph/Source/OptimusEditor/Private/Widgets/SOptimusShaderTextEditor.h
+++ b/Engine/Plugins/Animation/DeformerGraph/Source/OptimusEditor/Private/Widgets/SOptimusShaderTextEditor.h
@@ -4,11 +4,14 @@
 
 #include "CoreMinimal.h"
 #include "Widgets/SCompoundWidget.h"
+#include "OptimusEditor.h"
 
 class FOptimusEditor;
 class FDocumentTracker;
 class FTabManager;
 class SDockTab;
+class UOptimusNodeGraph;
+enum class EOptimusGraphNotifyType;
 
 
 class SOptimusShaderTextEditor : public SCompoundWidget
@@ -32,10 +35,13 @@ public:
 
 private:
 	void HandleSelectedNodesChanged(const TArray<TWeakObjectPtr<UObject>>& InSelectedObjects) const;
+	void HandleGraphNotified(EOptimusGraphNotifyType InNotifyType, UOptimusNodeGraph* InGraph, UObject* InSubject) const;
 
 	TSharedPtr<FTabManager> TabManager;
 
 	TSharedPtr<FDocumentTracker> DocumentTracker;
+
+	FOptimusEditor::FOnGraphNotified OnGraphNotifyDelegate;
 };
 
 
