Optimus: Copy paste binding arrays now create the pins accordingly

diff --git a/Engine/Plugins/Animation/DeformerGraph/Source/OptimusCore/Private/Nodes/OptimusNode_CustomComputeKernel.cpp b/Engine/Plugins/Animation/DeformerGraph/Source/OptimusCore/Private/Nodes/OptimusNode_CustomComputeKernel.cpp
--- a/Engine/Plugins/Animation/DeformerGraph/Source/OptimusCore/Private/Nodes/OptimusNode_CustomComputeKernel.cpp
+++ b/Engine/Plugins/Animation/DeformerGraph/Source/OptimusCore/Private/Nodes/OptimusNode_CustomComputeKernel.cpp
@@ -267,7 +267,15 @@ void UOptimusNode_CustomComputeKernel::PostEditChangeProperty(
 
 	if (PropertyChangedEvent.ChangeType & EPropertyChangeType::ValueSet)
 	{
-		PropertyValueChanged(PropertyChangedEvent);
+		if (CastField<FArrayProperty>(PropertyChangedEvent.Property) &&
+			PropertyChangedEvent.GetArrayIndex(PropertyChangedEvent.GetPropertyName().ToString()) == INDEX_NONE)
+		{
+			PropertyArrayPasted(PropertyChangedEvent);
+		}
+		else
+		{
+			PropertyValueChanged(PropertyChangedEvent);
+		}
 	}
 	else if (PropertyChangedEvent.ChangeType & EPropertyChangeType::ArrayAdd)
 	{
@@ -287,6 +295,75 @@ void UOptimusNode_CustomComputeKernel::PostEditChangeProperty(
 	}
 }
 
+void UOptimusNode_CustomComputeKernel::PropertyArrayPasted(const FPropertyChangedEvent& InPropertyChangedEvent)
+{
+	auto RemoveAllPinsByPredicate = [this](
+	TArrayView<UOptimusNodePin* const> InPins,
+	TFunction<bool(const UOptimusNodePin*)> InPredicate)
+	{
+		// Make a copy of the pins, since we're removing from the array represented by the view.
+		for (UOptimusNodePin* Pin: TArray<UOptimusNodePin*>(InPins))
+		{
+			if (InPredicate(Pin))
+			{
+				RemovePin(Pin);
+			}
+		}
+	};
+	
+	FOptimusActionScope(*GetActionStack(), TEXT("Paste Bindings"));
+	
+	if (InPropertyChangedEvent.GetMemberPropertyName() == ExtraInputBindingGroupsName)
+	{
+		// Two cases here
+		// 1. Pasting into the group array - Recreate the entire group array, including the sub pins
+		// 2. Pasting into the binding array of one of the groups
+		//		- Unfortunately we cannot know which group is the subject, so recreate the entire group array as well
+
+		RemoveAllPinsByPredicate(GetPins(), IsGroupingInputPin);
+
+		for (const FOptimusSecondaryInputBindingsGroup& InputGroup: SecondaryInputBindingGroups)
+		{
+			UOptimusNodePin* GroupPin = AddGroupingPin(InputGroup.GroupName, EOptimusNodePinDirection::Input);
+
+			for (const FOptimusParameterBinding& Binding: InputGroup.BindingArray)
+			{
+				AddPin(Binding.Name, EOptimusNodePinDirection::Input, Binding.DataDomain, Binding.DataType, nullptr, GroupPin);
+			}
+
+			if (InputGroup.BindingArray.Num() > 0)
+			{
+				// Make sure the group pin is expanded so that the change is visible in the graph.
+				GroupPin->SetIsExpanded(true);
+			}
+		}
+
+		UpdatePreamble();
+	}
+	else if (InPropertyChangedEvent.GetMemberPropertyName() == InputBindingsName)
+	{
+		RemoveAllPinsByPredicate(GetPins(), IsConnectableInputPin);
+
+		for (const FOptimusParameterBinding& Binding: InputBindingArray)
+		{
+			AddPin(Binding.Name, EOptimusNodePinDirection::Input, Binding.DataDomain, Binding.DataType);
+		}
+
+		UpdatePreamble();
+	}
+	else if (InPropertyChangedEvent.GetMemberPropertyName() == OutputBindingsName)
+	{
+		RemoveAllPinsByPredicate(GetPins(), IsConnectableOutputPin);
+
+		for (const FOptimusParameterBinding& Binding: OutputBindingArray)
+		{
+			AddPin(Binding.Name, EOptimusNodePinDirection::Output, Binding.DataDomain, Binding.DataType);
+		}
+
+		UpdatePreamble();
+	}
+}
+
 void UOptimusNode_CustomComputeKernel::PropertyValueChanged(
 	const FPropertyChangedEvent& InPropertyChangedEvent
 	)
@@ -346,6 +423,38 @@ void UOptimusNode_CustomComputeKernel::PropertyValueChanged(
 		return false;
 	};
 	
+	auto UpdateName = [this](UObject* InNameScope, FOptimusParameterBinding& InBinding, UOptimusNodePin* InPin)
+	{
+		if (InPin->GetFName() != InBinding.Name)
+		{
+			InBinding.Name = Optimus::GetUniqueNameForScope(InNameScope, InBinding.Name);
+			SetPinName(InPin, InBinding.Name);
+		}
+	};
+
+	auto UpdatePinType = [this](UObject* , FOptimusParameterBinding& InBinding, UOptimusNodePin* InPin)
+	{
+		const FOptimusDataTypeHandle DataType = InBinding.DataType.Resolve();
+		if (InPin->GetDataType() != DataType)
+		{
+			SetPinDataType(InPin, DataType);
+		}
+	};
+
+	auto UpdatePinDataDomain = [this](UObject* , FOptimusParameterBinding& InBinding, UOptimusNodePin* InPin)
+	{
+		if (InPin->GetDataDomain() != InBinding.DataDomain)
+		{
+			SetPinDataDomain(InPin, InBinding.DataDomain);
+		}
+	};
+	
+	auto UpdatePin = [&](UObject* InNameScope, FOptimusParameterBinding& InBinding, UOptimusNodePin* InPin)
+	{
+		UpdateName(InNameScope, InBinding, InPin);
+		UpdatePinType(InNameScope, InBinding, InPin);
+		UpdatePinDataDomain(InNameScope, InBinding, InPin);
+	};
 	
 	if (InPropertyChangedEvent.GetMemberPropertyName() == GET_MEMBER_NAME_STRING_CHECKED(UOptimusNode_CustomComputeKernel, KernelName) && 
 		InPropertyChangedEvent.GetPropertyName() == GET_MEMBER_NAME_STRING_CHECKED(FOptimusValidatedName, Name))
@@ -375,15 +484,6 @@ void UOptimusNode_CustomComputeKernel::PropertyValueChanged(
 			}
 		}
 
-		auto UpdateName = [this](UObject* InNameScope, FOptimusParameterBinding& InBinding, UOptimusNodePin* InPin)
-		{
-			if (InPin->GetFName() != InBinding.Name)
-			{
-				InBinding.Name = Optimus::GetUniqueNameForScope(InNameScope, InBinding.Name);
-				SetPinName(InPin, InBinding.Name);
-			}
-		};
-
 		if (bUpdatePreamble || UpdateAllBindings(UpdateName))
 		{
 			UpdatePreamble();
@@ -391,15 +491,6 @@ void UOptimusNode_CustomComputeKernel::PropertyValueChanged(
 	}
 	else if (InPropertyChangedEvent.GetPropertyName() == GET_MEMBER_NAME_STRING_CHECKED(FOptimusValidatedName, Name))
 	{
-		auto UpdateName = [this](UObject* InNameScope, FOptimusParameterBinding& InBinding, UOptimusNodePin* InPin)
-		{
-			if (InPin->GetFName() != InBinding.Name)
-			{
-				InBinding.Name = Optimus::GetUniqueNameForScope(InNameScope, InBinding.Name);
-				SetPinName(InPin, InBinding.Name);
-			}
-		};
-
 		if (UpdateAllBindings(UpdateName))
 		{
 			UpdatePreamble();
@@ -408,15 +499,6 @@ void UOptimusNode_CustomComputeKernel::PropertyValueChanged(
 	}
 	else if (InPropertyChangedEvent.GetPropertyName() == GET_MEMBER_NAME_STRING_CHECKED(FOptimusDataTypeRef, TypeName))
 	{
-		auto UpdatePinType = [this](UObject* , FOptimusParameterBinding& InBinding, UOptimusNodePin* InPin)
-		{
-			const FOptimusDataTypeHandle DataType = InBinding.DataType.Resolve();
-			if (InPin->GetDataType() != DataType)
-			{
-				SetPinDataType(InPin, DataType);
-			}
-		};
-
 		if (UpdateAllBindings(UpdatePinType))
 		{
 			UpdatePreamble();
@@ -425,20 +507,21 @@ void UOptimusNode_CustomComputeKernel::PropertyValueChanged(
 	}
 	else if (InPropertyChangedEvent.GetPropertyName() == GET_MEMBER_NAME_STRING_CHECKED(FOptimusParameterBinding, DataDomain))
 	{
-		auto UpdatePinDataDomain = [this](UObject* , FOptimusParameterBinding& InBinding, UOptimusNodePin* InPin)
-		{
-			if (InPin->GetDataDomain() != InBinding.DataDomain)
-			{
-				SetPinDataDomain(InPin, InBinding.DataDomain);
-			}
-		};
-
 		if (UpdateAllBindings(UpdatePinDataDomain))
 		{
 			UpdatePreamble();
 			return;
 		}
 	}
+	else if ( InPropertyChangedEvent.GetPropertyName() == GET_MEMBER_NAME_STRING_CHECKED(FOptimusParameterBindingArray, InnerArray))
+	{
+		// Pasting into a specific binding
+		if (UpdateAllBindings(UpdatePin))
+		{
+			UpdatePreamble();
+			return;
+		}	
+	}
 }
 
 
diff --git a/Engine/Plugins/Animation/DeformerGraph/Source/OptimusCore/Private/Nodes/OptimusNode_CustomComputeKernel.h b/Engine/Plugins/Animation/DeformerGraph/Source/OptimusCore/Private/Nodes/OptimusNode_CustomComputeKernel.h
--- a/Engine/Plugins/Animation/DeformerGraph/Source/OptimusCore/Private/Nodes/OptimusNode_CustomComputeKernel.h
+++ b/Engine/Plugins/Animation/DeformerGraph/Source/OptimusCore/Private/Nodes/OptimusNode_CustomComputeKernel.h
@@ -172,6 +172,7 @@ protected:
 private:
 	
 #if WITH_EDITOR
+	void PropertyArrayPasted(const FPropertyChangedEvent& InPropertyChangedEvent);
 	void PropertyValueChanged(const FPropertyChangedEvent& InPropertyChangedEvent);
 	void PropertyArrayItemAdded(const FPropertyChangedEvent& InPropertyChangedEvent);
 	void PropertyArrayItemRemoved(const FPropertyChangedEvent& InPropertyChangedEvent);
