Optimus: fix preview mesh not getting saved with the asset if clicked apply

diff --git a/Engine/Plugins/Animation/DeformerGraph/Source/OptimusCore/Private/OptimusDeformer.cpp b/Engine/Plugins/Animation/DeformerGraph/Source/OptimusCore/Private/OptimusDeformer.cpp
--- a/Engine/Plugins/Animation/DeformerGraph/Source/OptimusCore/Private/OptimusDeformer.cpp
+++ b/Engine/Plugins/Animation/DeformerGraph/Source/OptimusCore/Private/OptimusDeformer.cpp
@@ -2359,6 +2359,10 @@ UMeshDeformerInstance* UOptimusDeformer::CreateInstance(
 
 void UOptimusDeformer::SetPreviewMesh(USkeletalMesh* PreviewMesh, bool bMarkAsDirty)
 {
+	if (bMarkAsDirty)
+	{
+		Modify();
+	}
 	Mesh = PreviewMesh;
 	
 	// FIXME: Notify upstream so the viewport can react.
