Optimus: Fix missing component bindings in combo box for blueprints with no native components.
Fix broken component icons in the same widget for non native components.

diff --git a/Engine/Plugins/Animation/DeformerGraph/Source/OptimusCore/Private/OptimusDeformerInstance.cpp b/Engine/Plugins/Animation/DeformerGraph/Source/OptimusCore/Private/OptimusDeformerInstance.cpp
--- a/Engine/Plugins/Animation/DeformerGraph/Source/OptimusCore/Private/OptimusDeformerInstance.cpp
+++ b/Engine/Plugins/Animation/DeformerGraph/Source/OptimusCore/Private/OptimusDeformerInstance.cpp
@@ -258,12 +258,6 @@ void UOptimusDeformerInstanceSettings::GetComponentBindings(
 	}
 }
 
-AActor* UOptimusDeformerInstanceSettings::GetActor() const
-{
-	// We should be owned by an actor at some point.
-	return GetTypedOuter<AActor>();
-}
-
 UOptimusComponentSourceBinding const* UOptimusDeformerInstanceSettings::GetComponentBindingByName(FName InBindingName) const
 {
 	if (const UOptimusDeformer* DeformerResolved = Deformer.Get())
diff --git a/Engine/Plugins/Animation/DeformerGraph/Source/OptimusCore/Public/OptimusDeformerInstance.h b/Engine/Plugins/Animation/DeformerGraph/Source/OptimusCore/Public/OptimusDeformerInstance.h
--- a/Engine/Plugins/Animation/DeformerGraph/Source/OptimusCore/Public/OptimusDeformerInstance.h
+++ b/Engine/Plugins/Animation/DeformerGraph/Source/OptimusCore/Public/OptimusDeformerInstance.h
@@ -123,9 +123,6 @@ public:
 protected:
 	friend class FOptimusDeformerInstanceComponentBindingCustomization;
 
-	/** Get the actor associated with this object. Used only by details customization. */
-	AActor* GetActor() const;
-
 	/** Get a full component source binding object by binding name. Used only by details customization. */
 	UOptimusComponentSourceBinding const* GetComponentBindingByName(FName InBindingName) const;
 };
diff --git a/Engine/Plugins/Animation/DeformerGraph/Source/OptimusEditor/Private/OptimusDetailsCustomization.cpp b/Engine/Plugins/Animation/DeformerGraph/Source/OptimusEditor/Private/OptimusDetailsCustomization.cpp
--- a/Engine/Plugins/Animation/DeformerGraph/Source/OptimusEditor/Private/OptimusDetailsCustomization.cpp
+++ b/Engine/Plugins/Animation/DeformerGraph/Source/OptimusEditor/Private/OptimusDetailsCustomization.cpp
@@ -1569,43 +1569,43 @@ void FOptimusDeformerInstanceComponentBindingCustomization::CustomizeHeader(
 	FName BindingName;
 	BindingPropertyHandle->GetValue(BindingName);
 
-	FName ComponentName;
-	ComponentPropertyHandle->GetValue(ComponentName);
+	FName SelectedComponentName;
+	ComponentPropertyHandle->GetValue(SelectedComponentName);
 	
 	TArray<UObject*> OuterObjects;
 	InPropertyHandle->GetOuterObjects(OuterObjects);
 
-	UOptimusDeformerInstanceSettings const* BindingProvider = Cast<UOptimusDeformerInstanceSettings>(OuterObjects[0]);
-	AActor const* OwningActor = BindingProvider ? BindingProvider->GetActor() : nullptr;
+	
 
 	const UOptimusComponentSourceBinding* Binding = nullptr;
 	FComponentHandle SelectedComponentHandle;
 
-	if (BindingProvider)
+	if (UOptimusDeformerInstanceSettings const* BindingProvider = Cast<UOptimusDeformerInstanceSettings>(OuterObjects[0]))
 	{
 		Binding = BindingProvider->GetComponentBindingByName(BindingName);
 		
-		if (OwningActor && Binding)
+		if (Binding != nullptr)
 		{
 			TArray<UActorComponent*> FilteredComponents;
 
-			// Add actor components
-			OwningActor->GetComponents(Binding->GetComponentSource()->GetComponentClass(), FilteredComponents);
 			// Add unset "auto" setting.
 			FilteredComponents.Add(nullptr);
 
-			// Add blueprint components
-			if (OwningActor->HasAnyFlags(RF_ArchetypeObject))
+			// Add actor components.
+			if (AActor const* OwningActor = BindingProvider->GetTypedOuter<AActor>())
+			{
+				OwningActor->GetComponents(Binding->GetComponentSource()->GetComponentClass(), FilteredComponents);
+			}
+
+			// Add blueprint components.
+			if (UBlueprintGeneratedClass const* OwningBlueprint = BindingProvider->GetTypedOuter<UBlueprintGeneratedClass>())
 			{
-				if (UBlueprintGeneratedClass* BlueprintGeneratedClass = Cast<UBlueprintGeneratedClass>(OwningActor->GetClass()))
+				const TArray<USCS_Node*>& ActorBlueprintNodes = OwningBlueprint->SimpleConstructionScript->GetAllNodes();
+				for (USCS_Node* Node : ActorBlueprintNodes)
 				{
-					const TArray<USCS_Node*>& ActorBlueprintNodes = BlueprintGeneratedClass->SimpleConstructionScript->GetAllNodes();
-					for (USCS_Node* Node : ActorBlueprintNodes)
+					if (Node->ComponentTemplate->IsA(Binding->GetComponentSource()->GetComponentClass()))
 					{
-						if (Node->ComponentTemplate->IsA(Binding->GetComponentSource()->GetComponentClass()))
-						{
-							FilteredComponents.Add(Node->ComponentTemplate);
-						}
+						FilteredComponents.Add(Node->ComponentTemplate);
 					}
 				}
 			}
@@ -1613,11 +1613,14 @@ void FOptimusDeformerInstanceComponentBindingCustomization::CustomizeHeader(
 			// Add filtered components to drop down
 			for (const UActorComponent* Component : FilteredComponents)
 			{
+				// Store ComponentNames for reverse lookup from Name->Handle used for getting component icon.
+				FName Name = Component ? FOptimusDeformerInstanceComponentBinding::GetSanitizedComponentName(Component) : FName();
+				ComponentNames.Add(Name);
+
 				FSoftObjectPath Path = Component ? FSoftObjectPath::GetOrCreateIDForObject(Component) : FSoftObjectPath();
 				ComponentHandles.Add(MakeShared<FSoftObjectPath>(Path));
-				
-				FName Name = Component ? FOptimusDeformerInstanceComponentBinding::GetSanitizedComponentName(Component) : FName();
-				if (Name.ToString() == ComponentName.ToString())
+
+				if (Name.ToString() == SelectedComponentName.ToString())
 				{
 					SelectedComponentHandle = ComponentHandles.Last();
 				}
@@ -1673,12 +1676,14 @@ void FOptimusDeformerInstanceComponentBindingCustomization::CustomizeHeader(
 			.AutoWidth()
 			[
 				SNew(SImage)
-				.Image_Lambda([ComponentPropertyHandle, OwningActor]()-> const FSlateBrush*
+				.Image_Lambda([this, ComponentPropertyHandle]()-> const FSlateBrush*
 				{
 					FName Name;
 					ComponentPropertyHandle->GetValue(Name);
-					UActorComponent* ComponentObject = (OwningActor == nullptr || Name.IsNone()) ? nullptr : FOptimusDeformerInstanceComponentBinding::GetActorComponent(OwningActor, Name.ToString()).Get();
-					return ComponentObject ? FSlateIconFinder::FindIconBrushForClass(ComponentObject->GetClass(), TEXT("SCS.Component")) : nullptr;
+
+					int32 Index = ComponentNames.Find(Name);
+					const UActorComponent* Component = (Index == INDEX_NONE) ? nullptr : Cast<UActorComponent>(ComponentHandles[Index]->ResolveObject());
+					return Component ? FSlateIconFinder::FindIconBrushForClass(Component->GetClass(), TEXT("SCS.Component")) : nullptr;
 				})
 				.ColorAndOpacity(FSlateColor::UseForeground())
 			]
diff --git a/Engine/Plugins/Animation/DeformerGraph/Source/OptimusEditor/Private/OptimusDetailsCustomization.h b/Engine/Plugins/Animation/DeformerGraph/Source/OptimusEditor/Private/OptimusDetailsCustomization.h
--- a/Engine/Plugins/Animation/DeformerGraph/Source/OptimusEditor/Private/OptimusDetailsCustomization.h
+++ b/Engine/Plugins/Animation/DeformerGraph/Source/OptimusEditor/Private/OptimusDetailsCustomization.h
@@ -413,6 +413,7 @@ protected:
 	void ComponentsReplaced(const TMap<UObject*, UObject*>& InReplacementMap);
 	
 private:
-	using FComponentHandle = TSharedPtr<FSoftObjectPath>; 
+	TArray<FName> ComponentNames;
+	using FComponentHandle = TSharedPtr<FSoftObjectPath>;
 	TArray<FComponentHandle> ComponentHandles;
 };
