Deformer Graph: Ensure that output pins on custom kernels are not parameters.

diff --git a/Engine/Plugins/Animation/DeformerGraph/Source/OptimusCore/Private/Nodes/OptimusNode_CustomComputeKernel.cpp b/Engine/Plugins/Animation/DeformerGraph/Source/OptimusCore/Private/Nodes/OptimusNode_CustomComputeKernel.cpp
--- a/Engine/Plugins/Animation/DeformerGraph/Source/OptimusCore/Private/Nodes/OptimusNode_CustomComputeKernel.cpp
+++ b/Engine/Plugins/Animation/DeformerGraph/Source/OptimusCore/Private/Nodes/OptimusNode_CustomComputeKernel.cpp
@@ -551,7 +551,29 @@ void UOptimusNode_CustomComputeKernel::PropertyArrayItemAdded(
 	{
 		InBinding.Name = Optimus::GetUniqueNameForScope(InGroupPin ? Cast<UObject>(InGroupPin) : this, InName);
 		InBinding.DataType = FOptimusDataTypeRegistry::Get().FindType(*FFloatProperty::StaticClass());
-		InBinding.DataDomain = FOptimusDataDomain();
+		if (InDirection == EOptimusNodePinDirection::Input)
+		{
+			// Default to parameters for input bindings.
+			InBinding.DataDomain = FOptimusDataDomain();
+		}
+		else
+		{
+			// Pick a suitable fallback for output pins.
+			FName DomainName = GetExecutionDomain();
+			if (DomainName.IsNone())
+			{
+				if (TArray<FName> ExecutionDomains = GetExecutionDomains(); !ExecutionDomains.IsEmpty())
+				{
+					DomainName = ExecutionDomains[0];
+				}
+				else
+				{
+					// FIXME: There should be a generic mechanism to get the most suitable default. 
+					DomainName = UOptimusSkeletalMeshComponentSource::Domains::Vertex;
+				}
+			}
+			InBinding.DataDomain = FOptimusDataDomain({DomainName});
+		}
 
 		AddPin(InBinding.Name, InDirection, InBinding.DataDomain, InBinding.DataType, InBeforePin, InGroupPin);
 
