ComputeFramework: Change to DataProvider validation API, and update for engine DataProviders.
Aim of change is to allow validation later in the pipe on the render thread, so that we can validate shader compilation etc.
Late validation of shader compilation will allow us to implement on demand compilation for compute framework shaders in the editor.

diff --git a/Engine/Plugins/Animation/DeformerGraph/Source/OptimusCore/Private/DataInterfaces/OptimusDataInterfaceAnimAttribute.cpp b/Engine/Plugins/Animation/DeformerGraph/Source/OptimusCore/Private/DataInterfaces/OptimusDataInterfaceAnimAttribute.cpp
--- a/Engine/Plugins/Animation/DeformerGraph/Source/OptimusCore/Private/DataInterfaces/OptimusDataInterfaceAnimAttribute.cpp
+++ b/Engine/Plugins/Animation/DeformerGraph/Source/OptimusCore/Private/DataInterfaces/OptimusDataInterfaceAnimAttribute.cpp
@@ -708,6 +708,20 @@ FOptimusAnimAttributeDataProviderProxy::FOptimusAnimAttributeDataProviderProxy(
 	AttributeArrayData.AddDefaulted(InTotalNumArrays);
 }
 
+bool FOptimusAnimAttributeDataProviderProxy::IsValid(FValidationData const& InValidationData) const
+{
+	if (InValidationData.ParameterStructSize != AttributeBuffer.Num())
+	{
+		return false;
+	}
+	if (AttributeArrayData.Num() != AttributeArrayMetadata.Num())
+	{
+		return false;
+	}
+
+	return true;
+}
+
 void FOptimusAnimAttributeDataProviderProxy::AllocateResources(FRDGBuilder& GraphBuilder)
 {
 	FComputeDataProviderRenderProxy::AllocateResources(GraphBuilder);
@@ -733,27 +747,17 @@ void FOptimusAnimAttributeDataProviderProxy::AllocateResources(FRDGBuilder& Grap
 	}
 }
 
-void FOptimusAnimAttributeDataProviderProxy::GatherDispatchData(FDispatchSetup const& InDispatchSetup, FCollectedDispatchData& InOutDispatchData)
+void FOptimusAnimAttributeDataProviderProxy::GatherDispatchData(FDispatchData const& InDispatchData)
 {
-	if (!ensure(InDispatchSetup.ParameterStructSizeForValidation == AttributeBuffer.Num()))
+	for (int32 InvocationIndex = 0; InvocationIndex < InDispatchData.NumInvocations; ++InvocationIndex)
 	{
-		return;
-	}
+		uint8* ParameterBuffer = (InDispatchData.ParameterBuffer + InDispatchData.ParameterBufferOffset + InDispatchData.ParameterBufferStride * InvocationIndex);
+		FMemory::Memcpy(ParameterBuffer, AttributeBuffer.GetData(), InDispatchData.ParameterStructSize);
 
-	for (int32 InvocationIndex = 0; InvocationIndex < InDispatchSetup.NumInvocations; ++InvocationIndex)
-	{
-		uint8* ParameterBuffer = (InOutDispatchData.ParameterBuffer + InDispatchSetup.ParameterBufferOffset + InDispatchSetup.ParameterBufferStride * InvocationIndex);
-		
-		FMemory::Memcpy(ParameterBuffer, AttributeBuffer.GetData(), InDispatchSetup.ParameterStructSizeForValidation);
-
-		if (ensure(AttributeArrayData.Num() == AttributeArrayMetadata.Num()))
+		for (int32 ArrayIndex = 0; ArrayIndex < AttributeArrayMetadata.Num(); ArrayIndex++)
 		{
-			for (int32 ArrayIndex = 0; ArrayIndex < AttributeArrayMetadata.Num(); ArrayIndex++)
-			{
-				const FArrayMetadata& ArrayMetadata = AttributeArrayMetadata[ArrayIndex];
-
-				*((FRDGBufferSRV**)(ParameterBuffer + ArrayMetadata.Offset)) = ArrayBufferSRVs[ArrayIndex];
-			}
+			const FArrayMetadata& ArrayMetadata = AttributeArrayMetadata[ArrayIndex];
+			*((FRDGBufferSRV**)(ParameterBuffer + ArrayMetadata.Offset)) = ArrayBufferSRVs[ArrayIndex];
 		}
 	}
 }
diff --git a/Engine/Plugins/Animation/DeformerGraph/Source/OptimusCore/Private/DataInterfaces/OptimusDataInterfaceAnimAttribute.h b/Engine/Plugins/Animation/DeformerGraph/Source/OptimusCore/Private/DataInterfaces/OptimusDataInterfaceAnimAttribute.h
--- a/Engine/Plugins/Animation/DeformerGraph/Source/OptimusCore/Private/DataInterfaces/OptimusDataInterfaceAnimAttribute.h
+++ b/Engine/Plugins/Animation/DeformerGraph/Source/OptimusCore/Private/DataInterfaces/OptimusDataInterfaceAnimAttribute.h
@@ -216,8 +216,9 @@ public:
 	);
 
 	//~ Begin FComputeDataProviderRenderProxy Interface
+	bool IsValid(FValidationData const& InValidationData) const override;
 	void AllocateResources(FRDGBuilder& GraphBuilder) override;
-	void GatherDispatchData(FDispatchSetup const& InDispatchSetup, FCollectedDispatchData& InOutDispatchData) override;
+	void GatherDispatchData(FDispatchData const& InDispatchData) override;
 	//~ End FComputeDataProviderRenderProxy Interface
 
 public:
diff --git a/Engine/Plugins/Animation/DeformerGraph/Source/OptimusCore/Private/DataInterfaces/OptimusDataInterfaceCloth.cpp b/Engine/Plugins/Animation/DeformerGraph/Source/OptimusCore/Private/DataInterfaces/OptimusDataInterfaceCloth.cpp
--- a/Engine/Plugins/Animation/DeformerGraph/Source/OptimusCore/Private/DataInterfaces/OptimusDataInterfaceCloth.cpp
+++ b/Engine/Plugins/Animation/DeformerGraph/Source/OptimusCore/Private/DataInterfaces/OptimusDataInterfaceCloth.cpp
@@ -130,12 +130,34 @@ FComputeDataProviderRenderProxy* UOptimusClothDataProvider::GetRenderProxy()
 
 FOptimusClothDataProviderProxy::FOptimusClothDataProviderProxy(USkinnedMeshComponent* SkinnedMeshComponent)
 {
-	SkeletalMeshObject = SkinnedMeshComponent->MeshObject;
+	if (SkinnedMeshComponent != nullptr && SkinnedMeshComponent->GetScene() != nullptr)
+	{
+		SkeletalMeshObject = SkinnedMeshComponent == nullptr ? nullptr : SkinnedMeshComponent->MeshObject;
+		FrameNumber = SkinnedMeshComponent == nullptr ? 0 : SkinnedMeshComponent->GetScene()->GetFrameNumber() + 1; // +1 matches the logic for FrameNumberToPrepare in FSkeletalMeshObjectGPUSkin::Update()
 
-	USkeletalMeshComponent* SkeletalMeshComponent = Cast<USkeletalMeshComponent>(SkinnedMeshComponent);
-	ClothBlendWeight = SkeletalMeshComponent ? SkeletalMeshComponent->ClothBlendWeight : 1.f;
+		if (USkeletalMeshComponent* SkeletalMeshComponent = Cast<USkeletalMeshComponent>(SkinnedMeshComponent))
+		{
+			ClothBlendWeight = SkeletalMeshComponent->ClothBlendWeight;
+		}
+	}
+}
+
+bool FOptimusClothDataProviderProxy::IsValid(FValidationData const& InValidationData) const
+{
+	if (InValidationData.ParameterStructSize != sizeof(FParameters))
+	{
+		return false;
+	}
+	if (SkeletalMeshObject == nullptr)
+	{
+		return false;
+	}
+	if (SkeletalMeshObject->GetSkeletalMeshRenderData().LODRenderData[SkeletalMeshObject->GetLOD()].RenderSections.Num() != InValidationData.NumInvocations)
+	{
+		return false;
+	}
 
-	FrameNumber = SkinnedMeshComponent->GetScene()->GetFrameNumber() + 1; // +1 matches the logic for FrameNumberToPrepare in FSkeletalMeshObjectGPUSkin::Update()
+	return true;
 }
 
 struct FClothDataInterfacePermutationIds
@@ -152,26 +174,33 @@ struct FClothDataInterfacePermutationIds
 	}
 };
 
-void FOptimusClothDataProviderProxy::GatherDispatchData(FDispatchSetup const& InDispatchSetup, FCollectedDispatchData& InOutDispatchData)
+void FOptimusClothDataProviderProxy::GatherPermutations(FPermutationData& InOutPermutationData) const
 {
-	if (!ensure(InDispatchSetup.ParameterStructSizeForValidation == sizeof(FClothDataInterfaceParameters)))
-	{
-		return;
-	}
-
 	const int32 LodIndex = SkeletalMeshObject->GetLOD();
 	FSkeletalMeshRenderData const& SkeletalMeshRenderData = SkeletalMeshObject->GetSkeletalMeshRenderData();
 	FSkeletalMeshLODRenderData const* LodRenderData = &SkeletalMeshRenderData.LODRenderData[LodIndex];
-	if (!ensure(LodRenderData->RenderSections.Num() == InDispatchSetup.NumInvocations))
+
+	FClothDataInterfacePermutationIds PermutationIds(InOutPermutationData.PermutationVector);
+	for (int32 InvocationIndex = 0; InvocationIndex < InOutPermutationData.NumInvocations; ++InvocationIndex)
 	{
-		return;
+		const bool bPreviousFrame = false;
+		FSkeletalMeshDeformerHelpers::FClothBuffers ClothBuffers = FSkeletalMeshDeformerHelpers::GetClothBuffersForReading(SkeletalMeshObject, LodIndex, InvocationIndex, FrameNumber, bPreviousFrame);
+		const bool bValidCloth = (ClothBuffers.ClothInfluenceBuffer != nullptr) && (ClothBuffers.ClothSimulatedPositionAndNormalBuffer != nullptr);
+
+		InOutPermutationData.PermutationIds[InvocationIndex] |= (bValidCloth ? PermutationIds.EnableDeformerCloth : 0);
 	}
+}
 
-	FClothDataInterfacePermutationIds PermutationIds(InDispatchSetup.PermutationVector);
+void FOptimusClothDataProviderProxy::GatherDispatchData(FDispatchData const& InDispatchData)
+{
+	const int32 LodIndex = SkeletalMeshObject->GetLOD();
+	FSkeletalMeshRenderData const& SkeletalMeshRenderData = SkeletalMeshObject->GetSkeletalMeshRenderData();
+	FSkeletalMeshLODRenderData const* LodRenderData = &SkeletalMeshRenderData.LODRenderData[LodIndex];
 
 	FRHIShaderResourceView* NullSRVBinding = GWhiteVertexBufferWithSRV->ShaderResourceViewRHI.GetReference();
 
-	for (int32 InvocationIndex = 0; InvocationIndex < InDispatchSetup.NumInvocations; ++InvocationIndex)
+	const TStridedView<FParameters> ParameterArray = MakeStridedParameterView<FParameters>(InDispatchData);
+	for (int32 InvocationIndex = 0; InvocationIndex < ParameterArray.Num(); ++InvocationIndex)
 	{
 		FSkelMeshRenderSection const& RenderSection = LodRenderData->RenderSections[InvocationIndex];
 
@@ -183,15 +212,13 @@ void FOptimusClothDataProviderProxy::GatherDispatchData(FDispatchSetup const& In
 		FSkeletalMeshDeformerHelpers::FClothBuffers ClothBuffers = FSkeletalMeshDeformerHelpers::GetClothBuffersForReading(SkeletalMeshObject, LodIndex, InvocationIndex, FrameNumber, bPreviousFrame);
 		const bool bValidCloth = (ClothBuffers.ClothInfluenceBuffer != nullptr) && (ClothBuffers.ClothSimulatedPositionAndNormalBuffer != nullptr);
 
-		FClothDataInterfaceParameters* Parameters = (FClothDataInterfaceParameters*)(InOutDispatchData.ParameterBuffer + InDispatchSetup.ParameterBufferOffset + InDispatchSetup.ParameterBufferStride * InvocationIndex);
-		Parameters->NumVertices = RenderSection.NumVertices;
-		Parameters->InputStreamStart = ClothBuffers.ClothInfluenceBufferOffset;
-		Parameters->ClothBlendWeight = bValidCloth ? ClothBlendWeight : 0.f;
-		Parameters->NumInfluencesPerVertex = bValidCloth ? NumClothInfluencesPerVertex : 0;
-		Parameters->ClothToLocal = ClothBuffers.ClothToLocal;
-		Parameters->ClothBuffer = bValidCloth ? ClothBuffers.ClothInfluenceBuffer : NullSRVBinding;
-		Parameters->ClothPositionsAndNormalsBuffer = bValidCloth ? ClothBuffers.ClothSimulatedPositionAndNormalBuffer : NullSRVBinding;
-
-		InOutDispatchData.PermutationId[InvocationIndex] |= (bValidCloth ? PermutationIds.EnableDeformerCloth : 0);
+		FParameters& Parameters = ParameterArray[InvocationIndex];
+		Parameters.NumVertices = RenderSection.NumVertices;
+		Parameters.InputStreamStart = ClothBuffers.ClothInfluenceBufferOffset;
+		Parameters.ClothBlendWeight = bValidCloth ? ClothBlendWeight : 0.f;
+		Parameters.NumInfluencesPerVertex = bValidCloth ? NumClothInfluencesPerVertex : 0;
+		Parameters.ClothToLocal = ClothBuffers.ClothToLocal;
+		Parameters.ClothBuffer = bValidCloth ? ClothBuffers.ClothInfluenceBuffer : NullSRVBinding;
+		Parameters.ClothPositionsAndNormalsBuffer = bValidCloth ? ClothBuffers.ClothSimulatedPositionAndNormalBuffer : NullSRVBinding;
 	}
 }
diff --git a/Engine/Plugins/Animation/DeformerGraph/Source/OptimusCore/Private/DataInterfaces/OptimusDataInterfaceCloth.h b/Engine/Plugins/Animation/DeformerGraph/Source/OptimusCore/Private/DataInterfaces/OptimusDataInterfaceCloth.h
--- a/Engine/Plugins/Animation/DeformerGraph/Source/OptimusCore/Private/DataInterfaces/OptimusDataInterfaceCloth.h
+++ b/Engine/Plugins/Animation/DeformerGraph/Source/OptimusCore/Private/DataInterfaces/OptimusDataInterfaceCloth.h
@@ -6,6 +6,7 @@
 #include "ComputeFramework/ComputeDataProvider.h"
 #include "OptimusDataInterfaceCloth.generated.h"
 
+class FClothDataInterfaceParameters;
 class FSkeletalMeshObject;
 class USkinnedMeshComponent;
 
@@ -55,11 +56,15 @@ public:
 	FOptimusClothDataProviderProxy(USkinnedMeshComponent* SkinnedMeshComponent);
 
 	//~ Begin FComputeDataProviderRenderProxy Interface
-	void GatherDispatchData(FDispatchSetup const& InDispatchSetup, FCollectedDispatchData& InOutDispatchData);
+	bool IsValid(FValidationData const& InValidationData) const override;
+	void GatherPermutations(FPermutationData& InOutPermutationData) const override;
+	void GatherDispatchData(FDispatchData const& InDispatchData);
 	//~ End FComputeDataProviderRenderProxy Interface
 
 private:
+	using FParameters = FClothDataInterfaceParameters;
+
 	FSkeletalMeshObject* SkeletalMeshObject = nullptr;
-	float ClothBlendWeight = 0.0f;
+	float ClothBlendWeight = 1.0f;
 	uint32 FrameNumber = 0;
 };
diff --git a/Engine/Plugins/Animation/DeformerGraph/Source/OptimusCore/Private/DataInterfaces/OptimusDataInterfaceConnectivity.cpp b/Engine/Plugins/Animation/DeformerGraph/Source/OptimusCore/Private/DataInterfaces/OptimusDataInterfaceConnectivity.cpp
--- a/Engine/Plugins/Animation/DeformerGraph/Source/OptimusCore/Private/DataInterfaces/OptimusDataInterfaceConnectivity.cpp
+++ b/Engine/Plugins/Animation/DeformerGraph/Source/OptimusCore/Private/DataInterfaces/OptimusDataInterfaceConnectivity.cpp
@@ -236,11 +236,29 @@ FComputeDataProviderRenderProxy* UOptimusConnectivityDataProvider::GetRenderProx
 
 
 FOptimusConnectivityDataProviderProxy::FOptimusConnectivityDataProviderProxy(USkinnedMeshComponent* SkinnedMeshComponent, TArray< TArray<uint32> >& InAdjacencyBufferPerLod)
-	: SkeletalMeshObject(SkinnedMeshComponent->MeshObject)
+	: SkeletalMeshObject(SkinnedMeshComponent != nullptr ? SkinnedMeshComponent->MeshObject : nullptr)
 	, AdjacencyBufferPerLod(InAdjacencyBufferPerLod)
 {
 }
 
+bool FOptimusConnectivityDataProviderProxy::IsValid(FValidationData const& InValidationData) const
+{
+	if (InValidationData.ParameterStructSize != sizeof(FParameters))
+	{
+		return false;
+	}
+	if (SkeletalMeshObject == nullptr || AdjacencyBufferPerLod.Num() == 0)
+	{
+		return false;
+	}
+	if (SkeletalMeshObject->GetSkeletalMeshRenderData().LODRenderData[SkeletalMeshObject->GetLOD()].RenderSections.Num() != InValidationData.NumInvocations)
+	{
+		return false;
+	}
+
+	return true;
+}
+
 void FOptimusConnectivityDataProviderProxy::AllocateResources(FRDGBuilder& GraphBuilder)
 {
 	const int32 LodIndex = SkeletalMeshObject->GetLOD();
@@ -252,31 +270,23 @@ void FOptimusConnectivityDataProviderProxy::AllocateResources(FRDGBuilder& Graph
 	GraphBuilder.QueueBufferUpload(ConnectivityBuffer, ConnectivityData.GetData(), ConnectivityData.Num() * sizeof(uint32), ERDGInitialDataFlags::None);
 }
 
-void FOptimusConnectivityDataProviderProxy::GatherDispatchData(FDispatchSetup const& InDispatchSetup, FCollectedDispatchData& InOutDispatchData)
+void FOptimusConnectivityDataProviderProxy::GatherDispatchData(FDispatchData const& InDispatchData)
 {
-	if (!ensure(InDispatchSetup.ParameterStructSizeForValidation == sizeof(FConnectivityDataInterfaceParameters)))
-	{
-		return;
-	}
-
 	const int32 LodIndex = SkeletalMeshObject->GetLOD();
 	FSkeletalMeshRenderData const& SkeletalMeshRenderData = SkeletalMeshObject->GetSkeletalMeshRenderData();
 	FSkeletalMeshLODRenderData const* LodRenderData = &SkeletalMeshRenderData.LODRenderData[LodIndex];
-	if (!ensure(LodRenderData->RenderSections.Num() == InDispatchSetup.NumInvocations))
-	{
-		return;
-	}
 
 	FRHIShaderResourceView* NullSRVBinding = GWhiteVertexBufferWithSRV->ShaderResourceViewRHI.GetReference();
 
-	for (int32 InvocationIndex = 0; InvocationIndex < InDispatchSetup.NumInvocations; ++InvocationIndex)
+	const TStridedView<FParameters> ParameterArray = MakeStridedParameterView<FParameters>(InDispatchData);
+	for (int32 InvocationIndex = 0; InvocationIndex < ParameterArray.Num(); ++InvocationIndex)
 	{
 		FSkelMeshRenderSection const& RenderSection = LodRenderData->RenderSections[InvocationIndex];
 
-		FConnectivityDataInterfaceParameters* Parameters = (FConnectivityDataInterfaceParameters*)(InOutDispatchData.ParameterBuffer + InDispatchSetup.ParameterBufferOffset + InDispatchSetup.ParameterBufferStride * InvocationIndex);
-		Parameters->NumVertices = RenderSection.NumVertices;
-		Parameters->InputStreamStart = RenderSection.BaseVertexIndex;
-		Parameters->MaxConnectedVertexCount = UOptimusConnectivityDataInterface::MaxConnectedVertexCount;
-		Parameters->ConnectivityBuffer = ConnectivityBufferSRV;
+		FParameters& Parameters = ParameterArray[InvocationIndex];
+		Parameters.NumVertices = RenderSection.NumVertices;
+		Parameters.InputStreamStart = RenderSection.BaseVertexIndex;
+		Parameters.MaxConnectedVertexCount = UOptimusConnectivityDataInterface::MaxConnectedVertexCount;
+		Parameters.ConnectivityBuffer = ConnectivityBufferSRV;
 	}
 }
diff --git a/Engine/Plugins/Animation/DeformerGraph/Source/OptimusCore/Private/DataInterfaces/OptimusDataInterfaceConnectivity.h b/Engine/Plugins/Animation/DeformerGraph/Source/OptimusCore/Private/DataInterfaces/OptimusDataInterfaceConnectivity.h
--- a/Engine/Plugins/Animation/DeformerGraph/Source/OptimusCore/Private/DataInterfaces/OptimusDataInterfaceConnectivity.h
+++ b/Engine/Plugins/Animation/DeformerGraph/Source/OptimusCore/Private/DataInterfaces/OptimusDataInterfaceConnectivity.h
@@ -6,6 +6,7 @@
 #include "ComputeFramework/ComputeDataProvider.h"
 #include "OptimusDataInterfaceConnectivity.generated.h"
 
+class FConnectivityDataInterfaceParameters;
 class FRDGBuffer;
 class FRDGBufferSRV;
 class FSkeletalMeshObject;
@@ -62,11 +63,14 @@ public:
 	FOptimusConnectivityDataProviderProxy(USkinnedMeshComponent* SkinnedMeshComponent, TArray< TArray<uint32> >& InAdjacencyBufferPerLod);
 
 	//~ Begin FComputeDataProviderRenderProxy Interface
+	bool IsValid(FValidationData const& InValidationData) const override;
 	void AllocateResources(FRDGBuilder& GraphBuilder) override;
-	void GatherDispatchData(FDispatchSetup const& InDispatchSetup, FCollectedDispatchData& InOutDispatchData);
+	void GatherDispatchData(FDispatchData const& InDispatchData);
 	//~ End FComputeDataProviderRenderProxy Interface
 
 private:
+	using FParameters = FConnectivityDataInterfaceParameters;
+
 	FSkeletalMeshObject* SkeletalMeshObject = nullptr;
 	TArray< TArray<uint32> > const& AdjacencyBufferPerLod;
 
diff --git a/Engine/Plugins/Animation/DeformerGraph/Source/OptimusCore/Private/DataInterfaces/OptimusDataInterfaceDebugDraw.cpp b/Engine/Plugins/Animation/DeformerGraph/Source/OptimusCore/Private/DataInterfaces/OptimusDataInterfaceDebugDraw.cpp
--- a/Engine/Plugins/Animation/DeformerGraph/Source/OptimusCore/Private/DataInterfaces/OptimusDataInterfaceDebugDraw.cpp
+++ b/Engine/Plugins/Animation/DeformerGraph/Source/OptimusCore/Private/DataInterfaces/OptimusDataInterfaceDebugDraw.cpp
@@ -116,10 +116,10 @@ FComputeDataProviderRenderProxy* UOptimusDebugDrawDataProvider::GetRenderProxy()
 FOptimusDebugDrawDataProviderProxy::FOptimusDebugDrawDataProviderProxy(UPrimitiveComponent* InPrimitiveComponent, FOptimusDebugDrawParameters const& InDebugDrawParameters)
 	: Setup(FIntRect(0, 0, 1920, 1080))
 {
-	Scene = InPrimitiveComponent->GetScene();
+	Scene = InPrimitiveComponent != nullptr ? InPrimitiveComponent->GetScene() : nullptr;
 
 	// Split LocalToWorld into a pre-translation and transform for large world coordinate support.
-	FMatrix RenderMatrix = InPrimitiveComponent->GetRenderMatrix();
+	FMatrix RenderMatrix = InPrimitiveComponent != nullptr ? InPrimitiveComponent->GetRenderMatrix() : FMatrix();
 	FVector PreViewTranslation = -RenderMatrix.GetOrigin();
 	LocalToWorld = FMatrix44f(RenderMatrix.ConcatTranslation(PreViewTranslation));
 
@@ -138,6 +138,20 @@ FOptimusDebugDrawDataProviderProxy::FOptimusDebugDrawDataProviderProxy(UPrimitiv
 	}
 }
 
+bool FOptimusDebugDrawDataProviderProxy::IsValid(FValidationData const& InValidationData) const
+{
+	if (InValidationData.ParameterStructSize != sizeof(FParameters))
+	{
+		return false;
+	}
+	if (Scene == nullptr)
+	{
+		return false;
+	}
+
+	return true;
+}
+
 void FOptimusDebugDrawDataProviderProxy::AllocateResources(FRDGBuilder& GraphBuilder)
 {
 	// Allocate ShaderPrint output buffers.
@@ -157,26 +171,22 @@ void FOptimusDebugDrawDataProviderProxy::AllocateResources(FRDGBuilder& GraphBui
 	}
 }
 
-void FOptimusDebugDrawDataProviderProxy::GatherDispatchData(FDispatchSetup const& InDispatchSetup, FCollectedDispatchData& InOutDispatchData)
+void FOptimusDebugDrawDataProviderProxy::GatherDispatchData(FDispatchData const& InDispatchData)
 {
-	if (!ensure(InDispatchSetup.ParameterStructSizeForValidation == sizeof(FDebugDrawDataInterfaceParameters)))
-	{
-		return;
-	}
-
-	for (int32 InvocationIndex = 0; InvocationIndex < InDispatchSetup.NumInvocations; ++InvocationIndex)
+	const TStridedView<FParameters> ParameterArray = MakeStridedParameterView<FParameters>(InDispatchData);
+	for (int32 InvocationIndex = 0; InvocationIndex < ParameterArray.Num(); ++InvocationIndex)
 	{
-		FDebugDrawDataInterfaceParameters* Parameters = (FDebugDrawDataInterfaceParameters*)(InOutDispatchData.ParameterBuffer + InDispatchSetup.ParameterBufferOffset + InDispatchSetup.ParameterBufferStride * InvocationIndex);
-		Parameters->LocalToWorld = LocalToWorld;
-		Parameters->Resolution = ConfigParameters.Resolution;
-		Parameters->FontSize = ConfigParameters.FontSize;
-		Parameters->FontSpacing = ConfigParameters.FontSpacing;
-		Parameters->MaxValueCount = ConfigParameters.MaxValueCount;
-		Parameters->MaxSymbolCount = ConfigParameters.MaxSymbolCount;
-		Parameters->MaxStateCount = ConfigParameters.MaxStateCount;
-		Parameters->MaxLineCount = ConfigParameters.MaxLineCount;
-		Parameters->MaxTriangleCount = ConfigParameters.MaxTriangleCount;
-		Parameters->StateBuffer = CachedParameters.ShaderPrint_StateBuffer;
-		Parameters->RWEntryBuffer = CachedParameters.ShaderPrint_RWEntryBuffer;
+		FParameters& Parameters = ParameterArray[InvocationIndex];
+		Parameters.LocalToWorld = LocalToWorld;
+		Parameters.Resolution = ConfigParameters.Resolution;
+		Parameters.FontSize = ConfigParameters.FontSize;
+		Parameters.FontSpacing = ConfigParameters.FontSpacing;
+		Parameters.MaxValueCount = ConfigParameters.MaxValueCount;
+		Parameters.MaxSymbolCount = ConfigParameters.MaxSymbolCount;
+		Parameters.MaxStateCount = ConfigParameters.MaxStateCount;
+		Parameters.MaxLineCount = ConfigParameters.MaxLineCount;
+		Parameters.MaxTriangleCount = ConfigParameters.MaxTriangleCount;
+		Parameters.StateBuffer = CachedParameters.ShaderPrint_StateBuffer;
+		Parameters.RWEntryBuffer = CachedParameters.ShaderPrint_RWEntryBuffer;
 	}
 }
diff --git a/Engine/Plugins/Animation/DeformerGraph/Source/OptimusCore/Private/DataInterfaces/OptimusDataInterfaceDebugDraw.h b/Engine/Plugins/Animation/DeformerGraph/Source/OptimusCore/Private/DataInterfaces/OptimusDataInterfaceDebugDraw.h
--- a/Engine/Plugins/Animation/DeformerGraph/Source/OptimusCore/Private/DataInterfaces/OptimusDataInterfaceDebugDraw.h
+++ b/Engine/Plugins/Animation/DeformerGraph/Source/OptimusCore/Private/DataInterfaces/OptimusDataInterfaceDebugDraw.h
@@ -8,6 +8,7 @@
 
 #include "OptimusDataInterfaceDebugDraw.generated.h"
 
+class FDebugDrawDataInterfaceParameters;
 class UPrimitiveComponent;
 
 /* User controllable debug draw settings. */
@@ -89,11 +90,14 @@ public:
 	FOptimusDebugDrawDataProviderProxy(UPrimitiveComponent* InPrimitiveComponent, FOptimusDebugDrawParameters const& InDebugDrawParameters);
 
 	//~ Begin FComputeDataProviderRenderProxy Interface
+	bool IsValid(FValidationData const& InValidationData) const override;
 	void AllocateResources(FRDGBuilder& GraphBuilder) override;
-	void GatherDispatchData(FDispatchSetup const& InDispatchSetup, FCollectedDispatchData& InOutDispatchData) override;
+	void GatherDispatchData(FDispatchData const& InDispatchData) override;
 	//~ End FComputeDataProviderRenderProxy Interface
 
 private:
+	using FParameters = FDebugDrawDataInterfaceParameters;
+
 	FSceneInterface const* Scene;
 	FMatrix44f LocalToWorld;
 	ShaderPrint::FShaderPrintSetup Setup;
diff --git a/Engine/Plugins/Animation/DeformerGraph/Source/OptimusCore/Private/DataInterfaces/OptimusDataInterfaceDuplicateVertices.cpp b/Engine/Plugins/Animation/DeformerGraph/Source/OptimusCore/Private/DataInterfaces/OptimusDataInterfaceDuplicateVertices.cpp
--- a/Engine/Plugins/Animation/DeformerGraph/Source/OptimusCore/Private/DataInterfaces/OptimusDataInterfaceDuplicateVertices.cpp
+++ b/Engine/Plugins/Animation/DeformerGraph/Source/OptimusCore/Private/DataInterfaces/OptimusDataInterfaceDuplicateVertices.cpp
@@ -119,7 +119,25 @@ FComputeDataProviderRenderProxy* UOptimusDuplicateVerticesDataProvider::GetRende
 
 FOptimusDuplicateVerticesDataProviderProxy::FOptimusDuplicateVerticesDataProviderProxy(USkinnedMeshComponent* SkinnedMeshComponent)
 {
-	SkeletalMeshObject = SkinnedMeshComponent->MeshObject;
+	SkeletalMeshObject = SkinnedMeshComponent != nullptr ? SkinnedMeshComponent->MeshObject : nullptr;
+}
+
+bool FOptimusDuplicateVerticesDataProviderProxy::IsValid(FValidationData const& InValidationData) const
+{
+	if (InValidationData.ParameterStructSize != sizeof(FParameters))
+	{
+		return false;
+	}
+	if (SkeletalMeshObject == nullptr)
+	{
+		return false;
+	}
+	if (SkeletalMeshObject->GetSkeletalMeshRenderData().LODRenderData[SkeletalMeshObject->GetLOD()].RenderSections.Num() != InValidationData.NumInvocations)
+	{
+		return false;
+	}
+
+	return true;
 }
 
 struct FDuplicateVerticesDataInterfacePermutationIds
@@ -136,26 +154,35 @@ struct FDuplicateVerticesDataInterfacePermutationIds
 	}
 };
 
-void FOptimusDuplicateVerticesDataProviderProxy::GatherDispatchData(FDispatchSetup const& InDispatchSetup, FCollectedDispatchData& InOutDispatchData)
+void FOptimusDuplicateVerticesDataProviderProxy::GatherPermutations(FPermutationData& InOutPermutationData) const
 {
-	if (!ensure(InDispatchSetup.ParameterStructSizeForValidation == sizeof(FDuplicateVerticesDataInterfaceParameters)))
-	{
-		return;
-	}
-
 	const int32 LodIndex = SkeletalMeshObject->GetLOD();
 	FSkeletalMeshRenderData const& SkeletalMeshRenderData = SkeletalMeshObject->GetSkeletalMeshRenderData();
 	FSkeletalMeshLODRenderData const* LodRenderData = &SkeletalMeshRenderData.LODRenderData[LodIndex];
-	if (!ensure(LodRenderData->RenderSections.Num() == InDispatchSetup.NumInvocations))
+
+	FDuplicateVerticesDataInterfacePermutationIds PermutationIds(InOutPermutationData.PermutationVector);
+	for (int32 InvocationIndex = 0; InvocationIndex < InOutPermutationData.NumInvocations; ++InvocationIndex)
 	{
-		return;
+		FSkelMeshRenderSection const& RenderSection = LodRenderData->RenderSections[InvocationIndex];
+
+		FRHIShaderResourceView* DuplicateVertexStartAndLengthSRV = RenderSection.DuplicatedVerticesBuffer.LengthAndIndexDuplicatedVerticesIndexBuffer.VertexBufferSRV;
+		FRHIShaderResourceView* DuplicateVerticesSRV = RenderSection.DuplicatedVerticesBuffer.DuplicatedVerticesIndexBuffer.VertexBufferSRV;
+		const bool bValidDuplicateIndices = (DuplicateVertexStartAndLengthSRV != nullptr) && (DuplicateVerticesSRV != nullptr);
+
+		InOutPermutationData.PermutationIds[InvocationIndex] |= (bValidDuplicateIndices ? PermutationIds.EnableDuplicateVertices : 0);
 	}
+}
 
-	FDuplicateVerticesDataInterfacePermutationIds PermutationIds(InDispatchSetup.PermutationVector);
+void FOptimusDuplicateVerticesDataProviderProxy::GatherDispatchData(FDispatchData const& InDispatchData)
+{
+	const int32 LodIndex = SkeletalMeshObject->GetLOD();
+	FSkeletalMeshRenderData const& SkeletalMeshRenderData = SkeletalMeshObject->GetSkeletalMeshRenderData();
+	FSkeletalMeshLODRenderData const* LodRenderData = &SkeletalMeshRenderData.LODRenderData[LodIndex];
 
 	FRHIShaderResourceView* NullSRVBinding = GWhiteVertexBufferWithSRV->ShaderResourceViewRHI.GetReference();
 
-	for (int32 InvocationIndex = 0; InvocationIndex < InDispatchSetup.NumInvocations; ++InvocationIndex)
+	const TStridedView<FParameters> ParameterArray = MakeStridedParameterView<FParameters>(InDispatchData);
+	for (int32 InvocationIndex = 0; InvocationIndex < ParameterArray.Num(); ++InvocationIndex)
 	{
 		FSkelMeshRenderSection const& RenderSection = LodRenderData->RenderSections[InvocationIndex];
 
@@ -163,13 +190,11 @@ void FOptimusDuplicateVerticesDataProviderProxy::GatherDispatchData(FDispatchSet
 		FRHIShaderResourceView* DuplicateVerticesSRV = RenderSection.DuplicatedVerticesBuffer.DuplicatedVerticesIndexBuffer.VertexBufferSRV;
 		const bool bValidDuplicateIndices = (DuplicateVertexStartAndLengthSRV != nullptr) && (DuplicateVerticesSRV != nullptr);
 
-		FDuplicateVerticesDataInterfaceParameters* Parameters = (FDuplicateVerticesDataInterfaceParameters*)(InOutDispatchData.ParameterBuffer + InDispatchSetup.ParameterBufferOffset + InDispatchSetup.ParameterBufferStride * InvocationIndex);
-		Parameters->NumVertices = RenderSection.NumVertices;
-		Parameters->NumDuplicateVertices = bValidDuplicateIndices ? RenderSection.DuplicatedVerticesBuffer.DupVertData.Num() : 0;
-		Parameters->InputStreamStart = RenderSection.BaseVertexIndex;
-		Parameters->DuplicateVertexStartAndLength = bValidDuplicateIndices ? DuplicateVertexStartAndLengthSRV : NullSRVBinding;
-		Parameters->DuplicateVertices = bValidDuplicateIndices ? DuplicateVerticesSRV : NullSRVBinding;
-
-		InOutDispatchData.PermutationId[InvocationIndex] |= (bValidDuplicateIndices ? PermutationIds.EnableDuplicateVertices : 0);
+		FParameters& Parameters = ParameterArray[InvocationIndex];
+		Parameters.NumVertices = RenderSection.NumVertices;
+		Parameters.NumDuplicateVertices = bValidDuplicateIndices ? RenderSection.DuplicatedVerticesBuffer.DupVertData.Num() : 0;
+		Parameters.InputStreamStart = RenderSection.BaseVertexIndex;
+		Parameters.DuplicateVertexStartAndLength = bValidDuplicateIndices ? DuplicateVertexStartAndLengthSRV : NullSRVBinding;
+		Parameters.DuplicateVertices = bValidDuplicateIndices ? DuplicateVerticesSRV : NullSRVBinding;
 	}
 }
diff --git a/Engine/Plugins/Animation/DeformerGraph/Source/OptimusCore/Private/DataInterfaces/OptimusDataInterfaceDuplicateVertices.h b/Engine/Plugins/Animation/DeformerGraph/Source/OptimusCore/Private/DataInterfaces/OptimusDataInterfaceDuplicateVertices.h
--- a/Engine/Plugins/Animation/DeformerGraph/Source/OptimusCore/Private/DataInterfaces/OptimusDataInterfaceDuplicateVertices.h
+++ b/Engine/Plugins/Animation/DeformerGraph/Source/OptimusCore/Private/DataInterfaces/OptimusDataInterfaceDuplicateVertices.h
@@ -7,6 +7,7 @@
 
 #include "OptimusDataInterfaceDuplicateVertices.generated.h"
 
+class FDuplicateVerticesDataInterfaceParameters;
 class FSkeletalMeshObject;
 class USkinnedMeshComponent;
 
@@ -59,9 +60,13 @@ public:
 	FOptimusDuplicateVerticesDataProviderProxy(USkinnedMeshComponent* SkinnedMeshComponent);
 
 	//~ Begin FComputeDataProviderRenderProxy Interface
-	void GatherDispatchData(FDispatchSetup const& InDispatchSetup, FCollectedDispatchData& InOutDispatchData) override;
+	bool IsValid(FValidationData const& InValidationData) const override;
+	void GatherPermutations(FPermutationData& InOutPermutationData) const override;
+	void GatherDispatchData(FDispatchData const& InDispatchData) override;
 	//~ End FComputeDataProviderRenderProxy Interface
 
 private:
+	using FParameters = FDuplicateVerticesDataInterfaceParameters;
+
 	FSkeletalMeshObject* SkeletalMeshObject;
 };
diff --git a/Engine/Plugins/Animation/DeformerGraph/Source/OptimusCore/Private/DataInterfaces/OptimusDataInterfaceGraph.cpp b/Engine/Plugins/Animation/DeformerGraph/Source/OptimusCore/Private/DataInterfaces/OptimusDataInterfaceGraph.cpp
--- a/Engine/Plugins/Animation/DeformerGraph/Source/OptimusCore/Private/DataInterfaces/OptimusDataInterfaceGraph.cpp
+++ b/Engine/Plugins/Animation/DeformerGraph/Source/OptimusCore/Private/DataInterfaces/OptimusDataInterfaceGraph.cpp
@@ -162,22 +162,27 @@ FOptimusGraphDataProviderProxy::FOptimusGraphDataProviderProxy(UOptimusDeformerI
 	}
 }
 
-void FOptimusGraphDataProviderProxy::GatherDispatchData(FDispatchSetup const& InDispatchSetup, FCollectedDispatchData& InOutDispatchData)
+bool FOptimusGraphDataProviderProxy::IsValid(FValidationData const& InValidationData) const
 {
 	if (ParameterData.Num() == 0)
 	{
 		// todo[CF]: Why can we end up here? Remove this condition if possible.
-		return;
+		return false;
 	}
 
-	if (!ensure(ParameterData.Num() == InDispatchSetup.ParameterStructSizeForValidation))
+	if (InValidationData.ParameterStructSize != ParameterData.Num())
 	{
-		return;
+		return false;
 	}
 
-	for (int32 InvocationIndex = 0; InvocationIndex < InDispatchSetup.NumInvocations; ++InvocationIndex)
+	return true;
+}
+
+void FOptimusGraphDataProviderProxy::GatherDispatchData(FDispatchData const& InDispatchData)
+{
+	for (int32 InvocationIndex = 0; InvocationIndex < InDispatchData.NumInvocations; ++InvocationIndex)
 	{
-		void* ParameterBuffer = (void*)(InOutDispatchData.ParameterBuffer + InDispatchSetup.ParameterBufferOffset + InDispatchSetup.ParameterBufferStride * InvocationIndex);
+		void* ParameterBuffer = (void*)(InDispatchData.ParameterBuffer + InDispatchData.ParameterBufferOffset + InDispatchData.ParameterBufferStride * InvocationIndex);
 		FMemory::Memcpy(ParameterBuffer, ParameterData.GetData(), ParameterData.Num());
 	}
 }
diff --git a/Engine/Plugins/Animation/DeformerGraph/Source/OptimusCore/Private/DataInterfaces/OptimusDataInterfaceGraph.h b/Engine/Plugins/Animation/DeformerGraph/Source/OptimusCore/Private/DataInterfaces/OptimusDataInterfaceGraph.h
--- a/Engine/Plugins/Animation/DeformerGraph/Source/OptimusCore/Private/DataInterfaces/OptimusDataInterfaceGraph.h
+++ b/Engine/Plugins/Animation/DeformerGraph/Source/OptimusCore/Private/DataInterfaces/OptimusDataInterfaceGraph.h
@@ -85,7 +85,8 @@ public:
 	FOptimusGraphDataProviderProxy(UOptimusDeformerInstance const* DeformerInstance, TArray<FOptimusGraphVariableDescription> const& Variables, int32 ParameterBufferSize);
 
 	//~ Begin FComputeDataProviderRenderProxy Interface
-	void GatherDispatchData(FDispatchSetup const& InDispatchSetup, FCollectedDispatchData& InOutDispatchData) override;
+	bool IsValid(FValidationData const& InValidationData) const override;
+	void GatherDispatchData(FDispatchData const& InDispatchData) override;
 	//~ End FComputeDataProviderRenderProxy Interface
 
 private:
diff --git a/Engine/Plugins/Animation/DeformerGraph/Source/OptimusCore/Private/DataInterfaces/OptimusDataInterfaceMorphTarget.cpp b/Engine/Plugins/Animation/DeformerGraph/Source/OptimusCore/Private/DataInterfaces/OptimusDataInterfaceMorphTarget.cpp
--- a/Engine/Plugins/Animation/DeformerGraph/Source/OptimusCore/Private/DataInterfaces/OptimusDataInterfaceMorphTarget.cpp
+++ b/Engine/Plugins/Animation/DeformerGraph/Source/OptimusCore/Private/DataInterfaces/OptimusDataInterfaceMorphTarget.cpp
@@ -29,13 +29,11 @@ TArray<FOptimusCDIPinDefinition> UOptimusMorphTargetDataInterface::GetPinDefinit
 	return Defs;
 }
 
-
 TSubclassOf<UActorComponent> UOptimusMorphTargetDataInterface::GetRequiredComponentClass() const
 {
 	return USkinnedMeshComponent::StaticClass();
 }
 
-
 void UOptimusMorphTargetDataInterface::GetSupportedInputs(TArray<FShaderFunctionDefinition>& OutFunctions) const
 {
 	OutFunctions.AddDefaulted_GetRef()
@@ -109,8 +107,29 @@ FComputeDataProviderRenderProxy* UOptimusMorphTargetDataProvider::GetRenderProxy
 
 FOptimusMorphTargetDataProviderProxy::FOptimusMorphTargetDataProviderProxy(USkinnedMeshComponent* SkinnedMeshComponent)
 {
-	SkeletalMeshObject = SkinnedMeshComponent->MeshObject;
-	FrameNumber = SkinnedMeshComponent->GetScene()->GetFrameNumber() + 1; // +1 matches the logic for FrameNumberToPrepare in FSkeletalMeshObjectGPUSkin::Update()
+	if (SkinnedMeshComponent != nullptr && SkinnedMeshComponent->GetScene() != nullptr)
+	{
+		SkeletalMeshObject = SkinnedMeshComponent->MeshObject;
+		FrameNumber = SkinnedMeshComponent->GetScene()->GetFrameNumber() + 1; // +1 matches the logic for FrameNumberToPrepare in FSkeletalMeshObjectGPUSkin::Update()
+	}
+}
+
+bool FOptimusMorphTargetDataProviderProxy::IsValid(FValidationData const& InValidationData) const
+{
+	if (InValidationData.ParameterStructSize != sizeof(FParameters))
+	{
+		return false;
+	}
+	if (SkeletalMeshObject == nullptr)
+	{
+		return false;
+	}
+	if (SkeletalMeshObject->GetSkeletalMeshRenderData().LODRenderData[SkeletalMeshObject->GetLOD()].RenderSections.Num() != InValidationData.NumInvocations)
+	{
+		return false;
+	}
+
+	return true;
 }
 
 struct FMorphTargetDataInterfacePermutationIds
@@ -127,38 +146,44 @@ struct FMorphTargetDataInterfacePermutationIds
 	}
 };
 
-void FOptimusMorphTargetDataProviderProxy::GatherDispatchData(FDispatchSetup const& InDispatchSetup, FCollectedDispatchData& InOutDispatchData)
+void FOptimusMorphTargetDataProviderProxy::GatherPermutations(FPermutationData& InOutPermutationData) const
 {
-	if (!ensure(InDispatchSetup.ParameterStructSizeForValidation == sizeof(FMorphTargetDataInterfaceParameters)))
-	{
-		return;
-	}
-
 	const int32 LodIndex = SkeletalMeshObject->GetLOD();
 	FSkeletalMeshRenderData const& SkeletalMeshRenderData = SkeletalMeshObject->GetSkeletalMeshRenderData();
 	FSkeletalMeshLODRenderData const* LodRenderData = &SkeletalMeshRenderData.LODRenderData[LodIndex];
-	if (!ensure(LodRenderData->RenderSections.Num() == InDispatchSetup.NumInvocations))
+
+	FMorphTargetDataInterfacePermutationIds PermutationIds(InOutPermutationData.PermutationVector);
+	for (int32 InvocationIndex = 0; InvocationIndex < InOutPermutationData.NumInvocations; ++InvocationIndex)
 	{
-		return;
-	}
+		FSkelMeshRenderSection const& RenderSection = LodRenderData->RenderSections[InvocationIndex];
+
+		const bool bPreviousFrame = false;
+		FRHIShaderResourceView* MorphBufferSRV = FSkeletalMeshDeformerHelpers::GetMorphTargetBufferForReading(SkeletalMeshObject, LodIndex, InvocationIndex, FrameNumber, bPreviousFrame);
+		const bool bValidMorph = MorphBufferSRV != nullptr;
 
-	FMorphTargetDataInterfacePermutationIds PermutationIds(InDispatchSetup.PermutationVector);
+		InOutPermutationData.PermutationIds[InvocationIndex] |= (bValidMorph ? PermutationIds.EnableDeformerMorphTarget : 0);
+	}
+}
 
+void FOptimusMorphTargetDataProviderProxy::GatherDispatchData(FDispatchData const& InDispatchData)
+{
 	FRHIShaderResourceView* NullSRVBinding = GWhiteVertexBufferWithSRV->ShaderResourceViewRHI.GetReference();
 
-	for (int32 InvocationIndex = 0; InvocationIndex < InDispatchSetup.NumInvocations; ++InvocationIndex)
+	const int32 LodIndex = SkeletalMeshObject->GetLOD();
+	FSkeletalMeshRenderData const& SkeletalMeshRenderData = SkeletalMeshObject->GetSkeletalMeshRenderData();
+	FSkeletalMeshLODRenderData const* LodRenderData = &SkeletalMeshRenderData.LODRenderData[LodIndex];
+
+	const TStridedView<FParameters> ParameterArray = MakeStridedParameterView<FParameters>(InDispatchData);
+	for (int32 InvocationIndex = 0; InvocationIndex < ParameterArray.Num(); ++InvocationIndex)
 	{
 		FSkelMeshRenderSection const& RenderSection = LodRenderData->RenderSections[InvocationIndex];
 
 		const bool bPreviousFrame = false;
 		FRHIShaderResourceView* MorphBufferSRV = FSkeletalMeshDeformerHelpers::GetMorphTargetBufferForReading(SkeletalMeshObject, LodIndex, InvocationIndex, FrameNumber, bPreviousFrame);
-		const bool bValidMorph = MorphBufferSRV != nullptr;
-
-		FMorphTargetDataInterfaceParameters* Parameters = (FMorphTargetDataInterfaceParameters*)(InOutDispatchData.ParameterBuffer + InDispatchSetup.ParameterBufferOffset + InDispatchSetup.ParameterBufferStride * InvocationIndex);
-		Parameters->NumVertices = RenderSection.NumVertices;
-		Parameters->InputStreamStart = RenderSection.BaseVertexIndex;
-		Parameters->MorphBuffer = MorphBufferSRV != nullptr ? MorphBufferSRV : NullSRVBinding;
 
-		InOutDispatchData.PermutationId[InvocationIndex] |= (bValidMorph ? PermutationIds.EnableDeformerMorphTarget : 0);
+		FParameters& Parameters = ParameterArray[InvocationIndex];
+		Parameters.NumVertices = RenderSection.NumVertices;
+		Parameters.InputStreamStart = RenderSection.BaseVertexIndex;
+		Parameters.MorphBuffer = MorphBufferSRV != nullptr ? MorphBufferSRV : NullSRVBinding;
 	}
 }
diff --git a/Engine/Plugins/Animation/DeformerGraph/Source/OptimusCore/Private/DataInterfaces/OptimusDataInterfaceMorphTarget.h b/Engine/Plugins/Animation/DeformerGraph/Source/OptimusCore/Private/DataInterfaces/OptimusDataInterfaceMorphTarget.h
--- a/Engine/Plugins/Animation/DeformerGraph/Source/OptimusCore/Private/DataInterfaces/OptimusDataInterfaceMorphTarget.h
+++ b/Engine/Plugins/Animation/DeformerGraph/Source/OptimusCore/Private/DataInterfaces/OptimusDataInterfaceMorphTarget.h
@@ -7,6 +7,7 @@
 
 #include "OptimusDataInterfaceMorphTarget.generated.h"
 
+class FMorphTargetDataInterfaceParameters;
 class FSkeletalMeshObject;
 class USkinnedMeshComponent;
 
@@ -56,10 +57,14 @@ public:
 	FOptimusMorphTargetDataProviderProxy(USkinnedMeshComponent* SkinnedMeshComponent);
 
 	//~ Begin FComputeDataProviderRenderProxy Interface
-	void GatherDispatchData(FDispatchSetup const& InDispatchSetup, FCollectedDispatchData& InOutDispatchData);
+	bool IsValid(FValidationData const& InValidationData) const override;
+	void GatherPermutations(FPermutationData& InOutPermutationData) const override;
+	void GatherDispatchData(FDispatchData const& InDispatchData);
 	//~ End FComputeDataProviderRenderProxy Interface
 
 private:
-	FSkeletalMeshObject* SkeletalMeshObject;
+	using FParameters = FMorphTargetDataInterfaceParameters;
+
+	FSkeletalMeshObject* SkeletalMeshObject = nullptr;
 	uint32 FrameNumber = 0;
 };
diff --git a/Engine/Plugins/Animation/DeformerGraph/Source/OptimusCore/Private/DataInterfaces/OptimusDataInterfaceRawBuffer.cpp b/Engine/Plugins/Animation/DeformerGraph/Source/OptimusCore/Private/DataInterfaces/OptimusDataInterfaceRawBuffer.cpp
--- a/Engine/Plugins/Animation/DeformerGraph/Source/OptimusCore/Private/DataInterfaces/OptimusDataInterfaceRawBuffer.cpp
+++ b/Engine/Plugins/Animation/DeformerGraph/Source/OptimusCore/Private/DataInterfaces/OptimusDataInterfaceRawBuffer.cpp
@@ -355,14 +355,26 @@ FComputeDataProviderRenderProxy* UOptimusPersistentBufferDataProvider::GetRender
 FOptimusTransientBufferDataProviderProxy::FOptimusTransientBufferDataProviderProxy(
 	TArray<int32> InInvocationElementCounts,
 	int32 InElementStride,
-	int32 InRawStride
-	) :
-	InvocationElementCounts(InInvocationElementCounts),
-	ElementStride(InElementStride),
-	RawStride(InRawStride)
+	int32 InRawStride) 
+	: InvocationElementCounts(InInvocationElementCounts)
+	, ElementStride(InElementStride)
+	, RawStride(InRawStride)
 {
 }
 
+bool FOptimusTransientBufferDataProviderProxy::IsValid(FValidationData const& InValidationData) const
+{
+	if (InValidationData.ParameterStructSize != sizeof(FParameters))
+	{
+		return false;
+	}
+	if (InValidationData.NumInvocations != InvocationElementCounts.Num())
+	{
+		return false;
+	}
+
+	return true;
+}
 
 void FOptimusTransientBufferDataProviderProxy::AllocateResources(FRDGBuilder& GraphBuilder)
 {
@@ -379,52 +391,51 @@ void FOptimusTransientBufferDataProviderProxy::AllocateResources(FRDGBuilder& Gr
 	}
 }
 
-void FOptimusTransientBufferDataProviderProxy::GatherDispatchData(FDispatchSetup const& InDispatchSetup, FCollectedDispatchData& InOutDispatchData)
+void FOptimusTransientBufferDataProviderProxy::GatherDispatchData(FDispatchData const& InDispatchData)
 {
-	if (!ensure(InDispatchSetup.ParameterStructSizeForValidation == sizeof(FTransientBufferDataInterfaceParameters)))
+	TStridedView<FParameters> ParameterArray = MakeStridedParameterView<FParameters>(InDispatchData);
+	for (int32 InvocationIndex = 0; InvocationIndex < ParameterArray.Num(); ++InvocationIndex)
 	{
-		return;
-	}
-
-	if (!ensure(InDispatchSetup.NumInvocations == Buffer.Num()))
-	{
-		return;
-	}
-
-	for (int32 InvocationIndex = 0; InvocationIndex < InvocationElementCounts.Num(); ++InvocationIndex)
-	{
-		FTransientBufferDataInterfaceParameters* Parameters =
-			reinterpret_cast<FTransientBufferDataInterfaceParameters*>(InOutDispatchData.ParameterBuffer + InDispatchSetup.ParameterBufferOffset + InDispatchSetup.ParameterBufferStride * InvocationIndex);
-		Parameters->StartOffset = 0;
-		Parameters->BufferSize = InvocationElementCounts[InvocationIndex];
-		Parameters->BufferSRV = BufferSRV[InvocationIndex];
-		Parameters->BufferUAV = BufferUAV[InvocationIndex];
+		FParameters& Parameters = ParameterArray[InvocationIndex];
+		Parameters.StartOffset = 0;
+		Parameters.BufferSize = InvocationElementCounts[InvocationIndex];
+		Parameters.BufferSRV = BufferSRV[InvocationIndex];
+		Parameters.BufferUAV = BufferUAV[InvocationIndex];
 	}
 }
 
 
-
 FOptimusPersistentBufferDataProviderProxy::FOptimusPersistentBufferDataProviderProxy(
 	TArray<int32> InInvocationElementCounts,
 	int32 InElementStride,
 	int32 InRawStride,
 	TSharedPtr<FOptimusPersistentBufferPool> InBufferPool,
 	FName InResourceName,
-	int32 InLODIndex
-	) :
-	InvocationElementCounts(InInvocationElementCounts),
-	ElementStride(InElementStride),
-	RawStride(InRawStride),
-	BufferPool(InBufferPool),
-	ResourceName(InResourceName),
-	LODIndex(InLODIndex)
+	int32 InLODIndex)
+	: InvocationElementCounts(InInvocationElementCounts)
+	, ElementStride(InElementStride)
+	, RawStride(InRawStride)
+	, BufferPool(InBufferPool)
+	, ResourceName(InResourceName)
+	, LODIndex(InLODIndex)
 {
 }
 
+bool FOptimusPersistentBufferDataProviderProxy::IsValid(FValidationData const& InValidationData) const
+{
+	if (InValidationData.ParameterStructSize != sizeof(FParameters))
+	{
+		return false;
+	}
+	if (InValidationData.NumInvocations != InvocationElementCounts.Num())
+	{
+		return false;
+	}
+
+	return true;
+}
 
-void FOptimusPersistentBufferDataProviderProxy::AllocateResources(
-	FRDGBuilder& GraphBuilder
-	)
+void FOptimusPersistentBufferDataProviderProxy::AllocateResources(FRDGBuilder& GraphBuilder)
 {
 	BufferPool->GetResourceBuffers(GraphBuilder, ResourceName, LODIndex, ElementStride, RawStride, InvocationElementCounts, Buffers);
 	ensure(Buffers.Num() == InvocationElementCounts.Num());
@@ -435,31 +446,14 @@ void FOptimusPersistentBufferDataProviderProxy::AllocateResources(
 	}
 }
 
-
-void FOptimusPersistentBufferDataProviderProxy::GatherDispatchData(
-	FDispatchSetup const& InDispatchSetup,
-	FCollectedDispatchData& InOutDispatchData
-	)
+void FOptimusPersistentBufferDataProviderProxy::GatherDispatchData(FDispatchData const& InDispatchData)
 {
-	if (!ensure(InDispatchSetup.ParameterStructSizeForValidation == sizeof(FPersistentBufferDataInterfaceParameters)))
-	{
-		return;
-	}
-
-	if (!ensure(InDispatchSetup.NumInvocations == Buffers.Num()))
+	TStridedView<FParameters> ParameterArray = MakeStridedParameterView<FParameters>(InDispatchData);
+	for (int32 InvocationIndex = 0; InvocationIndex < ParameterArray.Num(); ++InvocationIndex)
 	{
-		return;
+		FParameters& Parameters = ParameterArray[InvocationIndex];
+		Parameters.StartOffset = 0;
+		Parameters.BufferSize = InvocationElementCounts[InvocationIndex];
+		Parameters.BufferUAV = BufferUAVs[InvocationIndex];
 	}
-
-	for (int32 InvocationIndex = 0; InvocationIndex < InvocationElementCounts.Num(); ++InvocationIndex)
-	{
-		FPersistentBufferDataInterfaceParameters* Parameters =
-			reinterpret_cast<FPersistentBufferDataInterfaceParameters*>(InOutDispatchData.ParameterBuffer + InDispatchSetup.ParameterBufferOffset + InDispatchSetup.ParameterBufferStride * InvocationIndex);
-		
-		Parameters->StartOffset = 0;
-		Parameters->BufferSize = InvocationElementCounts[InvocationIndex];
-		Parameters->BufferUAV = BufferUAVs[InvocationIndex];
-	}
-
-	FComputeDataProviderRenderProxy::GatherDispatchData(InDispatchSetup, InOutDispatchData);
 }
diff --git a/Engine/Plugins/Animation/DeformerGraph/Source/OptimusCore/Private/DataInterfaces/OptimusDataInterfaceRawBuffer.h b/Engine/Plugins/Animation/DeformerGraph/Source/OptimusCore/Private/DataInterfaces/OptimusDataInterfaceRawBuffer.h
--- a/Engine/Plugins/Animation/DeformerGraph/Source/OptimusCore/Private/DataInterfaces/OptimusDataInterfaceRawBuffer.h
+++ b/Engine/Plugins/Animation/DeformerGraph/Source/OptimusCore/Private/DataInterfaces/OptimusDataInterfaceRawBuffer.h
@@ -9,13 +9,15 @@
 
 #include "OptimusDataInterfaceRawBuffer.generated.h"
 
+class FOptimusPersistentBufferPool;
+class FPersistentBufferDataInterfaceParameters;
 class FRDGBuffer;
 class FRDGBufferSRV;
 class FRDGBufferUAV;
-class FOptimusPersistentBufferPool;
+class FTransientBufferDataInterfaceParameters;
+class UOptimusComponentSource;
 class UOptimusComponentSourceBinding;
 class UOptimusRawBufferDataProvider;
-class UOptimusComponentSource;
 
 /** Write to buffer operation types. */
 UENUM()
@@ -218,11 +220,14 @@ public:
 		);
 
 	//~ Begin FComputeDataProviderRenderProxy Interface
+	bool IsValid(FValidationData const& InValidationData) const override;
 	void AllocateResources(FRDGBuilder& GraphBuilder) override;
-	void GatherDispatchData(FDispatchSetup const& InDispatchSetup, FCollectedDispatchData& InOutDispatchData) override;
+	void GatherDispatchData(FDispatchData const& InDispatchData) override;
 	//~ End FComputeDataProviderRenderProxy Interface
 
 private:
+	using FParameters = FTransientBufferDataInterfaceParameters;
+
 	const TArray<int32> InvocationElementCounts;
 	const int32 ElementStride;
 	const int32 RawStride;
@@ -247,11 +252,14 @@ public:
 	);
 
 	//~ Begin FComputeDataProviderRenderProxy Interface
+	bool IsValid(FValidationData const& InValidationData) const override;
 	void AllocateResources(FRDGBuilder& GraphBuilder) override;
-	void GatherDispatchData(FDispatchSetup const& InDispatchSetup, FCollectedDispatchData& InOutDispatchData) override;
+	void GatherDispatchData(FDispatchData const& InDispatchData) override;
 	//~ End FComputeDataProviderRenderProxy Interface
 
 private:
+	using FParameters = FPersistentBufferDataInterfaceParameters;
+
 	const TArray<int32> InvocationElementCounts;
 	const int32 ElementStride;
 	const int32 RawStride;
diff --git a/Engine/Plugins/Animation/DeformerGraph/Source/OptimusCore/Private/DataInterfaces/OptimusDataInterfaceScene.cpp b/Engine/Plugins/Animation/DeformerGraph/Source/OptimusCore/Private/DataInterfaces/OptimusDataInterfaceScene.cpp
--- a/Engine/Plugins/Animation/DeformerGraph/Source/OptimusCore/Private/DataInterfaces/OptimusDataInterfaceScene.cpp
+++ b/Engine/Plugins/Animation/DeformerGraph/Source/OptimusCore/Private/DataInterfaces/OptimusDataInterfaceScene.cpp
@@ -93,12 +93,14 @@ FComputeDataProviderRenderProxy* UOptimusSceneDataProvider::GetRenderProxy()
 
 FOptimusSceneDataProviderProxy::FOptimusSceneDataProviderProxy(USceneComponent* SceneComponent)
 {
-	bool bUseSceneTime = SceneComponent != nullptr;
+	bIsValid = SceneComponent != nullptr;
+	
+	bool bUseSceneTime = bIsValid;
 #if WITH_EDITOR
 	// Don't tick time in Editor unless in PIE.
-	if (GIsEditor && bUseSceneTime)
+	if (GIsEditor && bIsValid)
 	{
-		bUseSceneTime &= (bUseSceneTime && SceneComponent->GetWorld() != nullptr && SceneComponent->GetWorld()->WorldType != EWorldType::Editor);
+		bIsValid &= (bUseSceneTime && SceneComponent->GetWorld() != nullptr && SceneComponent->GetWorld()->WorldType != EWorldType::Editor);
 	}
 #endif
 	GameTime = bUseSceneTime ? SceneComponent->GetWorld()->TimeSeconds : 0;
@@ -106,18 +108,24 @@ FOptimusSceneDataProviderProxy::FOptimusSceneDataProviderProxy(USceneComponent*
 	FrameNumber = bUseSceneTime ? SceneComponent->GetScene()->GetFrameNumber() : 0;
 }
 
-void FOptimusSceneDataProviderProxy::GatherDispatchData(FDispatchSetup const& InDispatchSetup, FCollectedDispatchData& InOutDispatchData)
+bool FOptimusSceneDataProviderProxy::IsValid(FValidationData const& InValidationData) const
 {
-	if (!ensure(InDispatchSetup.ParameterStructSizeForValidation == sizeof(FSceneDataInterfaceParameters)))
+	if (InValidationData.ParameterStructSize != sizeof(FParameters))
 	{
-		return;
+		return false;
 	}
 
-	for (int32 InvocationIndex = 0; InvocationIndex < InDispatchSetup.NumInvocations; ++InvocationIndex)
+	return bIsValid;
+}
+
+void FOptimusSceneDataProviderProxy::GatherDispatchData(FDispatchData const& InDispatchData)
+{
+	const TStridedView<FParameters> ParameterArray = MakeStridedParameterView<FParameters>(InDispatchData);
+	for (int32 InvocationIndex = 0; InvocationIndex < ParameterArray.Num(); ++InvocationIndex)
 	{
-		FSceneDataInterfaceParameters* Parameters = (FSceneDataInterfaceParameters*)(InOutDispatchData.ParameterBuffer + InDispatchSetup.ParameterBufferOffset + InDispatchSetup.ParameterBufferStride * InvocationIndex);
-		Parameters->GameTime = GameTime;
-		Parameters->GameTimeDelta = GameTimeDelta;
-		Parameters->FrameNumber = FrameNumber;
+		FParameters& Parameters = ParameterArray[InvocationIndex];
+		Parameters.GameTime = GameTime;
+		Parameters.GameTimeDelta = GameTimeDelta;
+		Parameters.FrameNumber = FrameNumber;
 	}
 }
diff --git a/Engine/Plugins/Animation/DeformerGraph/Source/OptimusCore/Private/DataInterfaces/OptimusDataInterfaceScene.h b/Engine/Plugins/Animation/DeformerGraph/Source/OptimusCore/Private/DataInterfaces/OptimusDataInterfaceScene.h
--- a/Engine/Plugins/Animation/DeformerGraph/Source/OptimusCore/Private/DataInterfaces/OptimusDataInterfaceScene.h
+++ b/Engine/Plugins/Animation/DeformerGraph/Source/OptimusCore/Private/DataInterfaces/OptimusDataInterfaceScene.h
@@ -6,6 +6,7 @@
 #include "ComputeFramework/ComputeDataProvider.h"
 #include "OptimusDataInterfaceScene.generated.h"
 
+class FSceneDataInterfaceParameters;
 class USceneComponent;
 
 /** Compute Framework Data Interface for reading general scene data. */
@@ -52,10 +53,14 @@ public:
 	FOptimusSceneDataProviderProxy(USceneComponent* SceneComponent);
 
 	//~ Begin FComputeDataProviderRenderProxy Interface
-	void GatherDispatchData(FDispatchSetup const& InDispatchSetup, FCollectedDispatchData& InOutDispatchData) override;
+	bool IsValid(FValidationData const& InValidationData) const override;
+	void GatherDispatchData(FDispatchData const& InDispatchData) override;
 	//~ End FComputeDataProviderRenderProxy Interface
 
 private:
+	using FParameters = FSceneDataInterfaceParameters;
+
+	bool bIsValid;
 	float GameTime;
 	float GameTimeDelta;
 	uint32 FrameNumber;
diff --git a/Engine/Plugins/Animation/DeformerGraph/Source/OptimusCore/Private/DataInterfaces/OptimusDataInterfaceSkeleton.cpp b/Engine/Plugins/Animation/DeformerGraph/Source/OptimusCore/Private/DataInterfaces/OptimusDataInterfaceSkeleton.cpp
--- a/Engine/Plugins/Animation/DeformerGraph/Source/OptimusCore/Private/DataInterfaces/OptimusDataInterfaceSkeleton.cpp
+++ b/Engine/Plugins/Animation/DeformerGraph/Source/OptimusCore/Private/DataInterfaces/OptimusDataInterfaceSkeleton.cpp
@@ -133,8 +133,26 @@ FComputeDataProviderRenderProxy* UOptimusSkeletonDataProvider::GetRenderProxy()
 
 FOptimusSkeletonDataProviderProxy::FOptimusSkeletonDataProviderProxy(USkinnedMeshComponent* SkinnedMeshComponent)
 {
-	SkeletalMeshObject = SkinnedMeshComponent->MeshObject;
-	BoneRevisionNumber = SkinnedMeshComponent->GetBoneTransformRevisionNumber();
+	SkeletalMeshObject = SkinnedMeshComponent != nullptr ? SkinnedMeshComponent->MeshObject : nullptr;
+	BoneRevisionNumber = SkinnedMeshComponent != nullptr ? SkinnedMeshComponent->GetBoneTransformRevisionNumber() : 0;
+}
+
+bool FOptimusSkeletonDataProviderProxy::IsValid(FValidationData const& InValidationData) const
+{
+	if (InValidationData.ParameterStructSize != sizeof(FParameters))
+	{
+		return false;
+	}
+	if (SkeletalMeshObject == nullptr)
+	{
+		return false;
+	}
+	if (SkeletalMeshObject->GetSkeletalMeshRenderData().LODRenderData[SkeletalMeshObject->GetLOD()].RenderSections.Num() != InValidationData.NumInvocations)
+	{
+		return false;
+	}
+
+	return true;
 }
 
 struct FSkeletonDataInterfacePermutationIds
@@ -169,26 +187,46 @@ struct FSkeletonDataInterfacePermutationIds
 	}
 };
 
-void FOptimusSkeletonDataProviderProxy::GatherDispatchData(FDispatchSetup const& InDispatchSetup, FCollectedDispatchData& InOutDispatchData)
+void FOptimusSkeletonDataProviderProxy::GatherPermutations(FPermutationData& InOutPermutationData) const
 {
-	if (!ensure(InDispatchSetup.ParameterStructSizeForValidation == sizeof(FSkeletonDataInterfaceParameters)))
-	{
-		return;
-	}
-
 	const int32 LodIndex = SkeletalMeshObject->GetLOD();
 	FSkeletalMeshRenderData const& SkeletalMeshRenderData = SkeletalMeshObject->GetSkeletalMeshRenderData();
 	FSkeletalMeshLODRenderData const* LodRenderData = &SkeletalMeshRenderData.LODRenderData[LodIndex];
-	if (!ensure(LodRenderData->RenderSections.Num() == InDispatchSetup.NumInvocations))
+
+	FSkeletonDataInterfacePermutationIds PermutationIds(InOutPermutationData.PermutationVector);
+	for (int32 InvocationIndex = 0; InvocationIndex < InOutPermutationData.NumInvocations; ++InvocationIndex)
 	{
-		return;
+		FSkelMeshRenderSection const& RenderSection = LodRenderData->RenderSections[InvocationIndex];
+
+		const bool bPreviousFrame = false;
+		FRHIShaderResourceView* BoneBufferSRV = FSkeletalMeshDeformerHelpers::GetBoneBufferForReading(SkeletalMeshObject, LodIndex, InvocationIndex, bPreviousFrame);
+
+		FSkinWeightVertexBuffer const* WeightBuffer = LodRenderData->GetSkinWeightVertexBuffer();
+		check(WeightBuffer != nullptr);
+		FRHIShaderResourceView* SkinWeightBufferSRV = WeightBuffer->GetDataVertexBuffer()->GetSRV();
+		const bool bUnlimitedBoneInfluences = WeightBuffer->GetBoneInfluenceType() == GPUSkinBoneInfluenceType::UnlimitedBoneInfluence;
+		FRHIShaderResourceView* InputWeightLookupStreamSRV = bUnlimitedBoneInfluences ? WeightBuffer->GetLookupVertexBuffer()->GetSRV() : nullptr;
+		const bool bValidBones = (BoneBufferSRV != nullptr) && (SkinWeightBufferSRV != nullptr) && (!bUnlimitedBoneInfluences || InputWeightLookupStreamSRV != nullptr);
+		const bool bUse16BitBoneIndex = WeightBuffer->Use16BitBoneIndex();
+		const bool bUse16BitBoneWeights = WeightBuffer->Use16BitBoneWeight();
+
+		InOutPermutationData.PermutationIds[InvocationIndex] |= (bValidBones ? PermutationIds.EnableDeformerBones : 0);
+		InOutPermutationData.PermutationIds[InvocationIndex] |= (bUnlimitedBoneInfluences ? PermutationIds.UnlimitedBoneInfluence : 0);
+		InOutPermutationData.PermutationIds[InvocationIndex] |= (bUse16BitBoneIndex ? PermutationIds.BoneIndexUint16 : 0);
+		InOutPermutationData.PermutationIds[InvocationIndex] |= (bUse16BitBoneWeights ? PermutationIds.BoneWeightsUint16 : 0);
 	}
+}
 
-	FSkeletonDataInterfacePermutationIds PermutationIds(InDispatchSetup.PermutationVector);
+void FOptimusSkeletonDataProviderProxy::GatherDispatchData(FDispatchData const& InDispatchData)
+{
+	const int32 LodIndex = SkeletalMeshObject->GetLOD();
+	FSkeletalMeshRenderData const& SkeletalMeshRenderData = SkeletalMeshObject->GetSkeletalMeshRenderData();
+	FSkeletalMeshLODRenderData const* LodRenderData = &SkeletalMeshRenderData.LODRenderData[LodIndex];
 
 	FRHIShaderResourceView* NullSRVBinding = GWhiteVertexBufferWithSRV->ShaderResourceViewRHI.GetReference();
 
-	for (int32 InvocationIndex = 0; InvocationIndex < InDispatchSetup.NumInvocations; ++InvocationIndex)
+	const TStridedView<FParameters> ParameterArray = MakeStridedParameterView<FParameters>(InDispatchData);
+	for (int32 InvocationIndex = 0; InvocationIndex < ParameterArray.Num(); ++InvocationIndex)
 	{
 		FSkelMeshRenderSection const& RenderSection = LodRenderData->RenderSections[InvocationIndex];
 
@@ -200,24 +238,16 @@ void FOptimusSkeletonDataProviderProxy::GatherDispatchData(FDispatchSetup const&
 		FRHIShaderResourceView* SkinWeightBufferSRV = WeightBuffer->GetDataVertexBuffer()->GetSRV();
 		const bool bUnlimitedBoneInfluences = WeightBuffer->GetBoneInfluenceType() == GPUSkinBoneInfluenceType::UnlimitedBoneInfluence;
 		FRHIShaderResourceView* InputWeightLookupStreamSRV = bUnlimitedBoneInfluences ? WeightBuffer->GetLookupVertexBuffer()->GetSRV() : nullptr;
-		const bool bValidBones = (BoneBufferSRV != nullptr) && (SkinWeightBufferSRV != nullptr) && (!bUnlimitedBoneInfluences || InputWeightLookupStreamSRV != nullptr);
-		const bool bUse16BitBoneIndex = WeightBuffer->Use16BitBoneIndex();
-		const bool bUse16BitBoneWeights = WeightBuffer->Use16BitBoneWeight();
 
-		FSkeletonDataInterfaceParameters* Parameters = (FSkeletonDataInterfaceParameters*)(InOutDispatchData.ParameterBuffer + InDispatchSetup.ParameterBufferOffset + InDispatchSetup.ParameterBufferStride * InvocationIndex);
-		Parameters->NumVertices = RenderSection.NumVertices;
-		Parameters->InputStreamStart = RenderSection.BaseVertexIndex;
-		Parameters->NumBoneInfluences = WeightBuffer->GetMaxBoneInfluences();
-		Parameters->InputWeightStart = (WeightBuffer->GetConstantInfluencesVertexStride() * RenderSection.GetVertexBufferIndex()) / sizeof(float);
-		Parameters->InputWeightStride = WeightBuffer->GetConstantInfluencesVertexStride();
-		Parameters->InputWeightIndexSize = WeightBuffer->GetBoneIndexByteSize() | (WeightBuffer->GetBoneWeightByteSize() << 8);
-		Parameters->BoneMatrices = BoneBufferSRV != nullptr ? BoneBufferSRV : NullSRVBinding;
-		Parameters->InputWeightStream = SkinWeightBufferSRV != nullptr ? SkinWeightBufferSRV : NullSRVBinding;
-		Parameters->InputWeightLookupStream = InputWeightLookupStreamSRV != nullptr ? InputWeightLookupStreamSRV : NullSRVBinding;
-
-		InOutDispatchData.PermutationId[InvocationIndex] |= (bValidBones ? PermutationIds.EnableDeformerBones : 0);
-		InOutDispatchData.PermutationId[InvocationIndex] |= (bUnlimitedBoneInfluences ? PermutationIds.UnlimitedBoneInfluence : 0);
-		InOutDispatchData.PermutationId[InvocationIndex] |= (bUse16BitBoneIndex ? PermutationIds.BoneIndexUint16 : 0);
-		InOutDispatchData.PermutationId[InvocationIndex] |= (bUse16BitBoneWeights ? PermutationIds.BoneWeightsUint16 : 0);
+		FParameters& Parameters = ParameterArray[InvocationIndex];
+		Parameters.NumVertices = RenderSection.NumVertices;
+		Parameters.InputStreamStart = RenderSection.BaseVertexIndex;
+		Parameters.NumBoneInfluences = WeightBuffer->GetMaxBoneInfluences();
+		Parameters.InputWeightStart = (WeightBuffer->GetConstantInfluencesVertexStride() * RenderSection.GetVertexBufferIndex()) / sizeof(float);
+		Parameters.InputWeightStride = WeightBuffer->GetConstantInfluencesVertexStride();
+		Parameters.InputWeightIndexSize = WeightBuffer->GetBoneIndexByteSize() | (WeightBuffer->GetBoneWeightByteSize() << 8);
+		Parameters.BoneMatrices = BoneBufferSRV != nullptr ? BoneBufferSRV : NullSRVBinding;
+		Parameters.InputWeightStream = SkinWeightBufferSRV != nullptr ? SkinWeightBufferSRV : NullSRVBinding;
+		Parameters.InputWeightLookupStream = InputWeightLookupStreamSRV != nullptr ? InputWeightLookupStreamSRV : NullSRVBinding;
 	}
 }
diff --git a/Engine/Plugins/Animation/DeformerGraph/Source/OptimusCore/Private/DataInterfaces/OptimusDataInterfaceSkeleton.h b/Engine/Plugins/Animation/DeformerGraph/Source/OptimusCore/Private/DataInterfaces/OptimusDataInterfaceSkeleton.h
--- a/Engine/Plugins/Animation/DeformerGraph/Source/OptimusCore/Private/DataInterfaces/OptimusDataInterfaceSkeleton.h
+++ b/Engine/Plugins/Animation/DeformerGraph/Source/OptimusCore/Private/DataInterfaces/OptimusDataInterfaceSkeleton.h
@@ -8,6 +8,7 @@
 #include "OptimusDataInterfaceSkeleton.generated.h"
 
 class FSkeletalMeshObject;
+class FSkeletonDataInterfaceParameters;
 class USkinnedMeshComponent;
 
 /** Compute Framework Data Interface for skeletal data. */
@@ -56,10 +57,14 @@ public:
 	FOptimusSkeletonDataProviderProxy(USkinnedMeshComponent* SkinnedMeshComponent);
 
 	//~ Begin FComputeDataProviderRenderProxy Interface
-	void GatherDispatchData(FDispatchSetup const& InDispatchSetup, FCollectedDispatchData& InOutDispatchData) override;
+	bool IsValid(FValidationData const& InValidationData) const override;
+	void GatherPermutations(FPermutationData& InOutPermutationData) const override;
+	void GatherDispatchData(FDispatchData const& InDispatchData) override;
 	//~ End FComputeDataProviderRenderProxy Interface
 
 private:
+	using FParameters = FSkeletonDataInterfaceParameters;
+
 	FSkeletalMeshObject* SkeletalMeshObject;
 	uint32 BoneRevisionNumber = 0;
 };
diff --git a/Engine/Plugins/Animation/DeformerGraph/Source/OptimusCore/Private/DataInterfaces/OptimusDataInterfaceSkinnedMesh.cpp b/Engine/Plugins/Animation/DeformerGraph/Source/OptimusCore/Private/DataInterfaces/OptimusDataInterfaceSkinnedMesh.cpp
--- a/Engine/Plugins/Animation/DeformerGraph/Source/OptimusCore/Private/DataInterfaces/OptimusDataInterfaceSkinnedMesh.cpp
+++ b/Engine/Plugins/Animation/DeformerGraph/Source/OptimusCore/Private/DataInterfaces/OptimusDataInterfaceSkinnedMesh.cpp
@@ -149,27 +149,37 @@ FComputeDataProviderRenderProxy* UOptimusSkinnedMeshDataProvider::GetRenderProxy
 
 FOptimusSkinnedMeshDataProviderProxy::FOptimusSkinnedMeshDataProviderProxy(USkinnedMeshComponent* SkinnedMeshComponent)
 {
-	SkeletalMeshObject = SkinnedMeshComponent->MeshObject;
+	SkeletalMeshObject = SkinnedMeshComponent != nullptr ? SkinnedMeshComponent->MeshObject : nullptr;
 }
 
-void FOptimusSkinnedMeshDataProviderProxy::GatherDispatchData(FDispatchSetup const& InDispatchSetup, FCollectedDispatchData& InOutDispatchData)
+bool FOptimusSkinnedMeshDataProviderProxy::IsValid(FValidationData const& InValidationData) const
 {
-	if (!ensure(InDispatchSetup.ParameterStructSizeForValidation == sizeof(FSkinnedMeshDataInterfaceParameters)))
+	if (InValidationData.ParameterStructSize != sizeof(FParameters))
 	{
-		return;
+		return false;
+	}
+	if (SkeletalMeshObject == nullptr)
+	{
+		return false;
+	}
+	if (SkeletalMeshObject->GetSkeletalMeshRenderData().LODRenderData[SkeletalMeshObject->GetLOD()].RenderSections.Num() != InValidationData.NumInvocations)
+	{
+		return false;
 	}
 
+	return true;
+}
+
+void FOptimusSkinnedMeshDataProviderProxy::GatherDispatchData(FDispatchData const& InDispatchData)
+{
 	const int32 LodIndex = SkeletalMeshObject->GetLOD();
 	FSkeletalMeshRenderData const& SkeletalMeshRenderData = SkeletalMeshObject->GetSkeletalMeshRenderData();
 	FSkeletalMeshLODRenderData const* LodRenderData = &SkeletalMeshRenderData.LODRenderData[LodIndex];
-	if (!ensure(LodRenderData->RenderSections.Num() == InDispatchSetup.NumInvocations))
-	{
-		return;
-	}
 
 	FRHIShaderResourceView* NullSRVBinding = GWhiteVertexBufferWithSRV->ShaderResourceViewRHI.GetReference();
 
-	for (int32 InvocationIndex = 0; InvocationIndex < InDispatchSetup.NumInvocations; ++InvocationIndex)
+	const TStridedView<FParameters> ParameterArray = MakeStridedParameterView<FParameters>(InDispatchData);
+	for (int32 InvocationIndex = 0; InvocationIndex < ParameterArray.Num(); ++InvocationIndex)
 	{
 		FSkelMeshRenderSection const& RenderSection = LodRenderData->RenderSections[InvocationIndex];
 
@@ -179,16 +189,16 @@ void FOptimusSkinnedMeshDataProviderProxy::GatherDispatchData(FDispatchSetup con
 		FRHIShaderResourceView* MeshUVBufferSRV = LodRenderData->StaticVertexBuffers.StaticMeshVertexBuffer.GetTexCoordsSRV();
 		FRHIShaderResourceView* MeshColorBufferSRV = LodRenderData->StaticVertexBuffers.ColorVertexBuffer.GetColorComponentsSRV();
 
-		FSkinnedMeshDataInterfaceParameters* Parameters = (FSkinnedMeshDataInterfaceParameters*)(InOutDispatchData.ParameterBuffer + InDispatchSetup.ParameterBufferOffset + InDispatchSetup.ParameterBufferStride * InvocationIndex);
-		Parameters->NumVertices = RenderSection.NumVertices;
-		Parameters->NumTriangles = RenderSection.NumTriangles;
-		Parameters->NumUVChannels = LodRenderData->StaticVertexBuffers.StaticMeshVertexBuffer.GetNumTexCoords();
-		Parameters->IndexBufferStart = RenderSection.BaseIndex;
-		Parameters->InputStreamStart = RenderSection.BaseVertexIndex;
-		Parameters->IndexBuffer = IndexBufferSRV != nullptr ? IndexBufferSRV : NullSRVBinding;
-		Parameters->PositionInputBuffer = MeshVertexBufferSRV != nullptr ? MeshVertexBufferSRV : NullSRVBinding;
-		Parameters->TangentInputBuffer = MeshTangentBufferSRV != nullptr ? MeshTangentBufferSRV : NullSRVBinding;
-		Parameters->UVInputBuffer = MeshUVBufferSRV != nullptr ? MeshUVBufferSRV : NullSRVBinding;
-		Parameters->ColorInputBuffer = MeshColorBufferSRV != nullptr ? MeshColorBufferSRV : NullSRVBinding;
+		FParameters& Parameters = ParameterArray[InvocationIndex];
+		Parameters.NumVertices = RenderSection.NumVertices;
+		Parameters.NumTriangles = RenderSection.NumTriangles;
+		Parameters.NumUVChannels = LodRenderData->StaticVertexBuffers.StaticMeshVertexBuffer.GetNumTexCoords();
+		Parameters.IndexBufferStart = RenderSection.BaseIndex;
+		Parameters.InputStreamStart = RenderSection.BaseVertexIndex;
+		Parameters.IndexBuffer = IndexBufferSRV != nullptr ? IndexBufferSRV : NullSRVBinding;
+		Parameters.PositionInputBuffer = MeshVertexBufferSRV != nullptr ? MeshVertexBufferSRV : NullSRVBinding;
+		Parameters.TangentInputBuffer = MeshTangentBufferSRV != nullptr ? MeshTangentBufferSRV : NullSRVBinding;
+		Parameters.UVInputBuffer = MeshUVBufferSRV != nullptr ? MeshUVBufferSRV : NullSRVBinding;
+		Parameters.ColorInputBuffer = MeshColorBufferSRV != nullptr ? MeshColorBufferSRV : NullSRVBinding;
 	}
 }
diff --git a/Engine/Plugins/Animation/DeformerGraph/Source/OptimusCore/Private/DataInterfaces/OptimusDataInterfaceSkinnedMesh.h b/Engine/Plugins/Animation/DeformerGraph/Source/OptimusCore/Private/DataInterfaces/OptimusDataInterfaceSkinnedMesh.h
--- a/Engine/Plugins/Animation/DeformerGraph/Source/OptimusCore/Private/DataInterfaces/OptimusDataInterfaceSkinnedMesh.h
+++ b/Engine/Plugins/Animation/DeformerGraph/Source/OptimusCore/Private/DataInterfaces/OptimusDataInterfaceSkinnedMesh.h
@@ -8,6 +8,7 @@
 #include "OptimusDataInterfaceSkinnedMesh.generated.h"
 
 class FSkeletalMeshObject;
+class FSkinnedMeshDataInterfaceParameters;
 class USkinnedMeshComponent;
 
 /** Compute Framework Data Interface for reading skeletal mesh. */
@@ -55,9 +56,12 @@ public:
 	FOptimusSkinnedMeshDataProviderProxy(USkinnedMeshComponent* SkinnedMeshComponent);
 
 	//~ Begin FComputeDataProviderRenderProxy Interface
-	void GatherDispatchData(FDispatchSetup const& InDispatchSetup, FCollectedDispatchData& InOutDispatchData) override;
+	bool IsValid(FValidationData const& InValidationData) const override;
+	void GatherDispatchData(FDispatchData const& InDispatchData) override;
 	//~ End FComputeDataProviderRenderProxy Interface
 
 private:
+	using FParameters = FSkinnedMeshDataInterfaceParameters;
+
 	FSkeletalMeshObject* SkeletalMeshObject;
 };
diff --git a/Engine/Plugins/Animation/DeformerGraph/Source/OptimusCore/Private/DataInterfaces/OptimusDataInterfaceSkinnedMeshExec.cpp b/Engine/Plugins/Animation/DeformerGraph/Source/OptimusCore/Private/DataInterfaces/OptimusDataInterfaceSkinnedMeshExec.cpp
--- a/Engine/Plugins/Animation/DeformerGraph/Source/OptimusCore/Private/DataInterfaces/OptimusDataInterfaceSkinnedMeshExec.cpp
+++ b/Engine/Plugins/Animation/DeformerGraph/Source/OptimusCore/Private/DataInterfaces/OptimusDataInterfaceSkinnedMeshExec.cpp
@@ -94,18 +94,23 @@ FComputeDataProviderRenderProxy* UOptimusSkinnedMeshExecDataProvider::GetRenderP
 
 FOptimusSkinnedMeshExecDataProviderProxy::FOptimusSkinnedMeshExecDataProviderProxy(USkinnedMeshComponent* InSkinnedMeshComponent, EOptimusSkinnedMeshExecDomain InDomain)
 {
-	SkeletalMeshObject = InSkinnedMeshComponent->MeshObject;
+	SkeletalMeshObject = InSkinnedMeshComponent != nullptr ? InSkinnedMeshComponent->MeshObject : nullptr;
 	Domain = InDomain;
 }
 
 int32 FOptimusSkinnedMeshExecDataProviderProxy::GetDispatchThreadCount(TArray<FIntVector>& ThreadCounts) const
 {
+	ThreadCounts.Reset();
+	if (SkeletalMeshObject == nullptr)
+	{
+		return 0;
+	}
+
 	const int32 LodIndex = SkeletalMeshObject->GetLOD();
 	FSkeletalMeshRenderData const& SkeletalMeshRenderData = SkeletalMeshObject->GetSkeletalMeshRenderData();
 	FSkeletalMeshLODRenderData const* LodRenderData = &SkeletalMeshRenderData.LODRenderData[LodIndex];
 	const int32 NumInvocations = LodRenderData->RenderSections.Num();
 
-	ThreadCounts.Reset();
 	ThreadCounts.Reserve(NumInvocations);
 	for (int32 InvocationIndex = 0; InvocationIndex < NumInvocations; ++InvocationIndex)
 	{
@@ -117,29 +122,37 @@ int32 FOptimusSkinnedMeshExecDataProviderProxy::GetDispatchThreadCount(TArray<FI
 	return NumInvocations;
 }
 
-void FOptimusSkinnedMeshExecDataProviderProxy::GatherDispatchData(FDispatchSetup const& InDispatchSetup, FCollectedDispatchData& InOutDispatchData)
+bool FOptimusSkinnedMeshExecDataProviderProxy::IsValid(FValidationData const& InValidationData) const
 {
-	if (!ensure(InDispatchSetup.ParameterStructSizeForValidation == sizeof(FSkinedMeshExecDataInterfaceParameters)))
+	if (InValidationData.ParameterStructSize != sizeof(FParameters))
+	{
+		return false;
+	}
+	if (SkeletalMeshObject == nullptr)
+	{
+		return false;
+	}
+	if (SkeletalMeshObject->GetSkeletalMeshRenderData().LODRenderData[SkeletalMeshObject->GetLOD()].RenderSections.Num() != InValidationData.NumInvocations)
 	{
-		return;
+		return false;
 	}
 
+	return true;
+}
+
+void FOptimusSkinnedMeshExecDataProviderProxy::GatherDispatchData(FDispatchData const& InDispatchData)
+{
 	const int32 LodIndex = SkeletalMeshObject->GetLOD();
 	FSkeletalMeshRenderData const& SkeletalMeshRenderData = SkeletalMeshObject->GetSkeletalMeshRenderData();
 	FSkeletalMeshLODRenderData const* LodRenderData = &SkeletalMeshRenderData.LODRenderData[LodIndex];
 
-	const int32 NumInvocations = LodRenderData->RenderSections.Num();
-	if (!ensure(NumInvocations == InDispatchSetup.NumInvocations))
-	{
-		return;
-	}
-
-	for (int32 InvocationIndex = 0; InvocationIndex < NumInvocations; ++InvocationIndex)
+	const TStridedView<FParameters> ParameterArray = MakeStridedParameterView<FParameters>(InDispatchData);
+	for (int32 InvocationIndex = 0; InvocationIndex < ParameterArray.Num(); ++InvocationIndex)
 	{
 		FSkelMeshRenderSection const& RenderSection = LodRenderData->RenderSections[InvocationIndex];
 		const int32 NumThreads = Domain == EOptimusSkinnedMeshExecDomain::Vertex ? RenderSection.NumVertices : RenderSection.NumTriangles;
 
-		FSkinedMeshExecDataInterfaceParameters* Parameters = (FSkinedMeshExecDataInterfaceParameters*)(InOutDispatchData.ParameterBuffer + InDispatchSetup.ParameterBufferOffset + InDispatchSetup.ParameterBufferStride * InvocationIndex);
-		Parameters->NumThreads = FIntVector(NumThreads, 1, 1);
+		FParameters& Parameters = ParameterArray[InvocationIndex];
+		Parameters.NumThreads = FIntVector(NumThreads, 1, 1);
 	}
 }
diff --git a/Engine/Plugins/Animation/DeformerGraph/Source/OptimusCore/Private/DataInterfaces/OptimusDataInterfaceSkinnedMeshExec.h b/Engine/Plugins/Animation/DeformerGraph/Source/OptimusCore/Private/DataInterfaces/OptimusDataInterfaceSkinnedMeshExec.h
--- a/Engine/Plugins/Animation/DeformerGraph/Source/OptimusCore/Private/DataInterfaces/OptimusDataInterfaceSkinnedMeshExec.h
+++ b/Engine/Plugins/Animation/DeformerGraph/Source/OptimusCore/Private/DataInterfaces/OptimusDataInterfaceSkinnedMeshExec.h
@@ -3,15 +3,15 @@
 #pragma once
 
 #include "OptimusComputeDataInterface.h"
-
 #include "ComputeFramework/ComputeDataProvider.h"
 
 #include "OptimusDataInterfaceSkinnedMeshExec.generated.h"
 
-class USkinnedMeshComponent;
-class FSkeletalMeshObject;
 class FRDGBuffer;
 class FRDGBufferUAV;
+class FSkeletalMeshObject;
+class FSkinedMeshExecDataInterfaceParameters;
+class USkinnedMeshComponent;
 
 UENUM()
 enum class EOptimusSkinnedMeshExecDomain : uint8
@@ -77,10 +77,13 @@ public:
 
 	//~ Begin FComputeDataProviderRenderProxy Interface
 	int32 GetDispatchThreadCount(TArray<FIntVector>& ThreadCounts) const override;
-	void GatherDispatchData(FDispatchSetup const& InDispatchSetup, FCollectedDispatchData& InOutDispatchData) override;
+	bool IsValid(FValidationData const& InValidationData) const override;
+	void GatherDispatchData(FDispatchData const& InDispatchData) override;
 	//~ End FComputeDataProviderRenderProxy Interface
 
 private:
+	using FParameters = FSkinedMeshExecDataInterfaceParameters;
+
 	FSkeletalMeshObject* SkeletalMeshObject = nullptr;
 	EOptimusSkinnedMeshExecDomain Domain = EOptimusSkinnedMeshExecDomain::Vertex;
 };
diff --git a/Engine/Plugins/Animation/DeformerGraph/Source/OptimusCore/Private/DataInterfaces/OptimusDataInterfaceSkinnedMeshWrite.cpp b/Engine/Plugins/Animation/DeformerGraph/Source/OptimusCore/Private/DataInterfaces/OptimusDataInterfaceSkinnedMeshWrite.cpp
--- a/Engine/Plugins/Animation/DeformerGraph/Source/OptimusCore/Private/DataInterfaces/OptimusDataInterfaceSkinnedMeshWrite.cpp
+++ b/Engine/Plugins/Animation/DeformerGraph/Source/OptimusCore/Private/DataInterfaces/OptimusDataInterfaceSkinnedMeshWrite.cpp
@@ -128,10 +128,28 @@ FComputeDataProviderRenderProxy* UOptimusSkinnedMeshWriteDataProvider::GetRender
 
 FOptimusSkinnedMeshWriteDataProviderProxy::FOptimusSkinnedMeshWriteDataProviderProxy(USkinnedMeshComponent* InSkinnedMeshComponent, uint64 InOutputMask)
 {
-	SkeletalMeshObject = InSkinnedMeshComponent->MeshObject;
+	SkeletalMeshObject = InSkinnedMeshComponent != nullptr ? InSkinnedMeshComponent->MeshObject : nullptr;
 	OutputMask = InOutputMask;
 }
 
+bool FOptimusSkinnedMeshWriteDataProviderProxy::IsValid(FValidationData const& InValidationData) const
+{
+	if (InValidationData.ParameterStructSize != sizeof(FParameters))
+	{
+		return false;
+	}
+	if (SkeletalMeshObject == nullptr)
+	{
+		return false;
+	}
+	if (SkeletalMeshObject->GetSkeletalMeshRenderData().LODRenderData[SkeletalMeshObject->GetLOD()].RenderSections.Num() != InValidationData.NumInvocations)
+	{
+		return false;
+	}
+
+	return true;
+}
+
 void FOptimusSkinnedMeshWriteDataProviderProxy::AllocateResources(FRDGBuilder& GraphBuilder)
 {
 	// Allocate required buffers
@@ -177,30 +195,22 @@ void FOptimusSkinnedMeshWriteDataProviderProxy::AllocateResources(FRDGBuilder& G
 	FSkeletalMeshDeformerHelpers::UpdateVertexFactoryBufferOverrides(SkeletalMeshObject, LodIndex);
 }
 
-void FOptimusSkinnedMeshWriteDataProviderProxy::GatherDispatchData(FDispatchSetup const& InDispatchSetup, FCollectedDispatchData& InOutDispatchData)
+void FOptimusSkinnedMeshWriteDataProviderProxy::GatherDispatchData(FDispatchData const& InDispatchData)
 {
-	if (!ensure(InDispatchSetup.ParameterStructSizeForValidation == sizeof(FSkinedMeshWriteDataInterfaceParameters)))
-	{
-		return;
-	}
-
 	const int32 LodIndex = SkeletalMeshObject->GetLOD();
 	FSkeletalMeshRenderData const& SkeletalMeshRenderData = SkeletalMeshObject->GetSkeletalMeshRenderData();
 	FSkeletalMeshLODRenderData const* LodRenderData = &SkeletalMeshRenderData.LODRenderData[LodIndex];
-	if (!ensure(LodRenderData->RenderSections.Num() == InDispatchSetup.NumInvocations))
-	{
-		return;
-	}
 
-	for (int32 InvocationIndex = 0; InvocationIndex < InDispatchSetup.NumInvocations; ++InvocationIndex)
+	const TStridedView<FParameters> ParameterArray = MakeStridedParameterView<FParameters>(InDispatchData);
+	for (int32 InvocationIndex = 0; InvocationIndex < ParameterArray.Num(); ++InvocationIndex)
 	{
 		FSkelMeshRenderSection const& RenderSection = LodRenderData->RenderSections[InvocationIndex];
 
-		FSkinedMeshWriteDataInterfaceParameters* Parameters = (FSkinedMeshWriteDataInterfaceParameters*)(InOutDispatchData.ParameterBuffer + InDispatchSetup.ParameterBufferOffset + InDispatchSetup.ParameterBufferStride * InvocationIndex);
-		Parameters->NumVertices = RenderSection.GetNumVertices();
-		Parameters->OutputStreamStart = RenderSection.GetVertexBufferIndex();
-		Parameters->PositionBufferUAV = PositionBufferUAV;
-		Parameters->TangentBufferUAV = TangentBufferUAV;
-		Parameters->ColorBufferUAV = ColorBufferUAV;
+		FParameters& Parameters = ParameterArray[InvocationIndex];
+		Parameters.NumVertices = RenderSection.GetNumVertices();
+		Parameters.OutputStreamStart = RenderSection.GetVertexBufferIndex();
+		Parameters.PositionBufferUAV = PositionBufferUAV;
+		Parameters.TangentBufferUAV = TangentBufferUAV;
+		Parameters.ColorBufferUAV = ColorBufferUAV;
 	}
 }
diff --git a/Engine/Plugins/Animation/DeformerGraph/Source/OptimusCore/Private/DataInterfaces/OptimusDataInterfaceSkinnedMeshWrite.h b/Engine/Plugins/Animation/DeformerGraph/Source/OptimusCore/Private/DataInterfaces/OptimusDataInterfaceSkinnedMeshWrite.h
--- a/Engine/Plugins/Animation/DeformerGraph/Source/OptimusCore/Private/DataInterfaces/OptimusDataInterfaceSkinnedMeshWrite.h
+++ b/Engine/Plugins/Animation/DeformerGraph/Source/OptimusCore/Private/DataInterfaces/OptimusDataInterfaceSkinnedMeshWrite.h
@@ -7,10 +7,11 @@
 
 #include "OptimusDataInterfaceSkinnedMeshWrite.generated.h"
 
-class USkinnedMeshComponent;
-class FSkeletalMeshObject;
 class FRDGBuffer;
 class FRDGBufferUAV;
+class FSkeletalMeshObject;
+class FSkinedMeshWriteDataInterfaceParameters;
+class USkinnedMeshComponent;
 
 /** Compute Framework Data Interface for writing skinned mesh. */
 UCLASS(Category = ComputeFramework)
@@ -61,11 +62,14 @@ public:
 	FOptimusSkinnedMeshWriteDataProviderProxy(USkinnedMeshComponent* InSkinnedMeshComponent, uint64 InOutputMask);
 
 	//~ Begin FComputeDataProviderRenderProxy Interface
+	bool IsValid(FValidationData const& InValidationData) const override;
 	void AllocateResources(FRDGBuilder& GraphBuilder) override;
-	void GatherDispatchData(FDispatchSetup const& InDispatchSetup, FCollectedDispatchData& InOutDispatchData) override;
+	void GatherDispatchData(FDispatchData const& InDispatchData) override;
 	//~ End FComputeDataProviderRenderProxy Interface
 
 private:
+	using FParameters = FSkinedMeshWriteDataInterfaceParameters;
+
 	FSkeletalMeshObject* SkeletalMeshObject;
 	uint64 OutputMask;
 
diff --git a/Engine/Plugins/Animation/MLDeformer/MLDeformerFramework/Source/MLDeformerFramework/Private/MLDeformerGraphDataInterface.cpp b/Engine/Plugins/Animation/MLDeformer/MLDeformerFramework/Source/MLDeformerFramework/Private/MLDeformerGraphDataInterface.cpp
--- a/Engine/Plugins/Animation/MLDeformer/MLDeformerFramework/Source/MLDeformerFramework/Private/MLDeformerGraphDataInterface.cpp
+++ b/Engine/Plugins/Animation/MLDeformer/MLDeformerFramework/Source/MLDeformerFramework/Private/MLDeformerGraphDataInterface.cpp
@@ -97,40 +97,43 @@ namespace UE::MLDeformer
 	{
 		using namespace UE::MLDeformer;
 
-		const UMLDeformerAsset* DeformerAsset = DeformerComponent->GetDeformerAsset();
-		if (DeformerAsset)
+		const UMLDeformerAsset* DeformerAsset = DeformerComponent != nullptr ? DeformerComponent->GetDeformerAsset() : nullptr;
+		const UMLDeformerModel* Model = DeformerAsset != nullptr ? DeformerAsset->GetModel() : nullptr;
+		const UMLDeformerModelInstance* ModelInstance = DeformerComponent != nullptr ? DeformerComponent->GetModelInstance() : nullptr;
+		if (Model != nullptr && ModelInstance != nullptr)
 		{
-			const UMLDeformerModel* Model = DeformerAsset->GetModel();
-			const UMLDeformerModelInstance* ModelInstance = DeformerComponent->GetModelInstance();
-		
 			SkeletalMeshObject = ModelInstance->GetSkeletalMeshComponent()->MeshObject;
 			NeuralNetwork = Model->GetNeuralNetwork();
 			NeuralNetworkInferenceHandle = ModelInstance->GetNeuralNetworkInferenceHandle();
-			bCanRunNeuralNet = ModelInstance->IsCompatible() && Model->IsNeuralNetworkOnGPU();
+			bCanRunNeuralNet = ModelInstance->IsCompatible() && Model->IsNeuralNetworkOnGPU() && DeformerComponent->GetModelInstance()->IsValidForDataProvider();
 			Weight = DeformerComponent->GetWeight();
 			VertexMapBufferSRV = Model->GetVertexMapBuffer().ShaderResourceViewRHI;
 		}
 	}
 
-	void FMLDeformerGraphDataProviderProxy::AllocateResources(FRDGBuilder& GraphBuilder)
+	bool FMLDeformerGraphDataProviderProxy::IsValid(FValidationData const& InValidationData) const
 	{
-		if (bCanRunNeuralNet)
+		if (InValidationData.ParameterStructSize != sizeof(FMLDeformerGraphDataInterfaceParameters))
 		{
-			check(NeuralNetwork != nullptr);
-			Buffer = GraphBuilder.RegisterExternalBuffer(NeuralNetwork->GetOutputTensorForContext(NeuralNetworkInferenceHandle).GetPooledBuffer());
+			return false;
 		}
-		else
+		if (!bCanRunNeuralNet || SkeletalMeshObject == nullptr || NeuralNetwork == nullptr || VertexMapBufferSRV == nullptr)
 		{
-			// TODO: use an actual buffer that is of the right size, and filled with zero's.
-			Buffer = GraphBuilder.RegisterExternalBuffer(GWhiteVertexBufferWithRDG->Buffer);
+			return false;
 		}
 
+		return true;
+	}
+
+	void FMLDeformerGraphDataProviderProxy::AllocateResources(FRDGBuilder& GraphBuilder)
+	{
+		Buffer = GraphBuilder.RegisterExternalBuffer(NeuralNetwork->GetOutputTensorForContext(NeuralNetworkInferenceHandle).GetPooledBuffer());
 		BufferSRV = GraphBuilder.CreateSRV(Buffer, PF_R32_FLOAT);
 	}
 
-	void FMLDeformerGraphDataProviderProxy::GatherDispatchData(FDispatchSetup const& InDispatchSetup, FCollectedDispatchData& InOutDispatchData)
+	void FMLDeformerGraphDataProviderProxy::GatherDispatchData(FDispatchData const& InDispatchData)
 	{
-		MLDEFORMER_GRAPH_DISPATCH_START(FMLDeformerGraphDataInterfaceParameters, InDispatchSetup, InOutDispatchData)
+		MLDEFORMER_GRAPH_DISPATCH_START(FMLDeformerGraphDataInterfaceParameters, InDispatchData)
 		MLDEFORMER_GRAPH_DISPATCH_DEFAULT_PARAMETERS()
 		MLDEFORMER_GRAPH_DISPATCH_END()
 	}
diff --git a/Engine/Plugins/Animation/MLDeformer/MLDeformerFramework/Source/MLDeformerFramework/Private/MLDeformerGraphDebugDataInterface.cpp b/Engine/Plugins/Animation/MLDeformer/MLDeformerFramework/Source/MLDeformerFramework/Private/MLDeformerGraphDebugDataInterface.cpp
--- a/Engine/Plugins/Animation/MLDeformer/MLDeformerFramework/Source/MLDeformerFramework/Private/MLDeformerGraphDebugDataInterface.cpp
+++ b/Engine/Plugins/Animation/MLDeformer/MLDeformerFramework/Source/MLDeformerFramework/Private/MLDeformerGraphDebugDataInterface.cpp
@@ -120,19 +120,19 @@ bool UMLDeformerGraphDebugDataProvider::IsValid() const
 FComputeDataProviderRenderProxy* UMLDeformerGraphDebugDataProvider::GetRenderProxy()
 {
 #if WITH_EDITORONLY_DATA
-	UE::MLDeformer::FMLDeformerGraphDebugDataProviderProxy* Proxy = new UE::MLDeformer::FMLDeformerGraphDebugDataProviderProxy(DeformerComponent, DeformerAsset, this);
-	if (DeformerAsset)
+	if (DeformerComponent && DeformerAsset)
 	{
+		UE::MLDeformer::FMLDeformerGraphDebugDataProviderProxy* Proxy = new UE::MLDeformer::FMLDeformerGraphDebugDataProviderProxy(DeformerComponent, DeformerAsset, this);
 		const float SampleTime = DeformerComponent->GetModelInstance()->GetSkeletalMeshComponent()->GetPosition();
 		UMLDeformerModel* Model = DeformerAsset->GetModel();
 		Model->SampleGroundTruthPositions(SampleTime, Proxy->GetGroundTruthPositions());
 		Proxy->HandleZeroGroundTruthPositions();
 		return Proxy;
 	}
-	return nullptr;
-#else
-	return nullptr;
 #endif
+
+	// Return default invalid proxy.
+	return new FComputeDataProviderRenderProxy();
 }
 
 #if WITH_EDITORONLY_DATA
@@ -143,7 +143,7 @@ namespace UE::MLDeformer
 	{
 		Provider = InProvider;
 
-		if (DeformerAsset)
+		if (DeformerComponent && DeformerAsset)
 		{
 			const UMLDeformerModel* Model = DeformerAsset->GetModel();	
 			const UMLDeformerVizSettings* VizSettings = Model->GetVizSettings();
@@ -175,6 +175,20 @@ namespace UE::MLDeformer
 		}
 	}
 
+	bool FMLDeformerGraphDebugDataProviderProxy::IsValid(FValidationData const& InValidationData) const
+	{
+		if (InValidationData.ParameterStructSize != sizeof(FMLDeformerGraphDebugDataInterfaceParameters))
+		{
+			return false;
+		}
+		if (SkeletalMeshObject == nullptr || VertexMapBufferSRV == nullptr)
+		{
+			return false;
+		}
+
+		return true;
+	}
+
 	void FMLDeformerGraphDebugDataProviderProxy::AllocateResources(FRDGBuilder& GraphBuilder)
 	{
 		GroundTruthBuffer = GraphBuilder.CreateBuffer(FRDGBufferDesc::CreateStructuredDesc(sizeof(float), 3 * GroundTruthPositions.Num()), TEXT("MLDeformer.GroundTruthPositions"));
@@ -182,9 +196,9 @@ namespace UE::MLDeformer
 		GraphBuilder.QueueBufferUpload(GroundTruthBuffer, GroundTruthPositions.GetData(), sizeof(FVector3f) * GroundTruthPositions.Num(), ERDGInitialDataFlags::None);
 	}
 
-	void FMLDeformerGraphDebugDataProviderProxy::GatherDispatchData(FDispatchSetup const& InDispatchSetup, FCollectedDispatchData& InOutDispatchData)
+	void FMLDeformerGraphDebugDataProviderProxy::GatherDispatchData(FDispatchData const& InDispatchData)
 	{
-		MLDEFORMER_GRAPH_DISPATCH_START(FMLDeformerGraphDebugDataInterfaceParameters, InDispatchSetup, InOutDispatchData)
+		MLDEFORMER_GRAPH_DISPATCH_START(FMLDeformerGraphDebugDataInterfaceParameters, InDispatchData)
 		MLDEFORMER_GRAPH_DISPATCH_DEFAULT_DEBUG_PARAMETERS()
 		MLDEFORMER_GRAPH_DISPATCH_END()
 	}
diff --git a/Engine/Plugins/Animation/MLDeformer/MLDeformerFramework/Source/MLDeformerFramework/Public/MLDeformerGraphDataInterface.h b/Engine/Plugins/Animation/MLDeformer/MLDeformerFramework/Source/MLDeformerFramework/Public/MLDeformerGraphDataInterface.h
--- a/Engine/Plugins/Animation/MLDeformer/MLDeformerFramework/Source/MLDeformerFramework/Public/MLDeformerGraphDataInterface.h
+++ b/Engine/Plugins/Animation/MLDeformer/MLDeformerFramework/Source/MLDeformerFramework/Public/MLDeformerGraphDataInterface.h
@@ -27,24 +27,21 @@ class UMLDeformerModel;
 	SHADER_PARAMETER_RDG_BUFFER_SRV(Buffer<float>, PositionDeltaBuffer) \
 	SHADER_PARAMETER_SRV(Buffer<uint>, VertexMapBuffer)
 
-#define MLDEFORMER_GRAPH_DISPATCH_START(ParameterStructType, InDispatchSetup, InOutDispatchData) \
-	if (!ensure(InDispatchSetup.ParameterStructSizeForValidation == sizeof(ParameterStructType))) \
-	{ \
-		return; \
-	} \
+#define MLDEFORMER_GRAPH_DISPATCH_START(ParameterStructType, InDispatchData) \
 	const FSkeletalMeshRenderData& SkeletalMeshRenderData = SkeletalMeshObject->GetSkeletalMeshRenderData(); \
 	const FSkeletalMeshLODRenderData* LodRenderData = SkeletalMeshRenderData.GetPendingFirstLOD(0); \
-	for (int32 InvocationIndex = 0; InvocationIndex < InDispatchSetup.NumInvocations; ++InvocationIndex) \
+	const TStridedView<ParameterStructType> ParameterArray = MakeStridedParameterView<ParameterStructType>(InDispatchData); \
+	for (int32 InvocationIndex = 0; InvocationIndex < ParameterArray.Num(); ++InvocationIndex) \
 	{ \
 		const FSkelMeshRenderSection& RenderSection = LodRenderData->RenderSections[InvocationIndex]; \
-		ParameterStructType* Parameters = (ParameterStructType*)(InOutDispatchData.ParameterBuffer + InDispatchSetup.ParameterBufferOffset + InDispatchSetup.ParameterBufferStride * InvocationIndex);
+		ParameterStructType& Parameters = ParameterArray[InvocationIndex]; \
 
 #define MLDEFORMER_GRAPH_DISPATCH_DEFAULT_PARAMETERS() \
-	Parameters->NumVertices = 0; \
-	Parameters->InputStreamStart = RenderSection.BaseVertexIndex; \
-	Parameters->Weight = Weight; \
-	Parameters->PositionDeltaBuffer = BufferSRV; \
-	Parameters->VertexMapBuffer = VertexMapBufferSRV;
+	Parameters.NumVertices = 0; \
+	Parameters.InputStreamStart = RenderSection.BaseVertexIndex; \
+	Parameters.Weight = Weight; \
+	Parameters.PositionDeltaBuffer = BufferSRV; \
+	Parameters.VertexMapBuffer = VertexMapBufferSRV;
 
 #define MLDEFORMER_GRAPH_DISPATCH_END() }
 
@@ -105,8 +102,9 @@ namespace UE::MLDeformer
 		FMLDeformerGraphDataProviderProxy(UMLDeformerComponent* DeformerComponent);
 
 		// FComputeDataProviderRenderProxy overrides.
-		virtual void AllocateResources(FRDGBuilder& GraphBuilder) override;
-		virtual void GatherDispatchData(FDispatchSetup const& InDispatchSetup, FCollectedDispatchData& InOutDispatchData) override;
+		bool IsValid(FValidationData const& InValidationData) const override;
+		void AllocateResources(FRDGBuilder& GraphBuilder) override;
+		void GatherDispatchData(FDispatchData const& InDispatchData) override;
 		// ~END FComputeDataProviderRenderProxy overrides.
 
 	protected:
diff --git a/Engine/Plugins/Animation/MLDeformer/MLDeformerFramework/Source/MLDeformerFramework/Public/MLDeformerGraphDebugDataInterface.h b/Engine/Plugins/Animation/MLDeformer/MLDeformerFramework/Source/MLDeformerFramework/Public/MLDeformerGraphDebugDataInterface.h
--- a/Engine/Plugins/Animation/MLDeformer/MLDeformerFramework/Source/MLDeformerFramework/Public/MLDeformerGraphDebugDataInterface.h
+++ b/Engine/Plugins/Animation/MLDeformer/MLDeformerFramework/Source/MLDeformerFramework/Public/MLDeformerGraphDebugDataInterface.h
@@ -27,14 +27,14 @@ class USkeletalMeshComponent;
 	SHADER_PARAMETER_SRV(Buffer<uint>, VertexMapBuffer)
 
 #define MLDEFORMER_GRAPH_DISPATCH_DEFAULT_DEBUG_PARAMETERS() \
-	Parameters->NumVertices = 0; \
-	Parameters->InputStreamStart = RenderSection.BaseVertexIndex; \
-	Parameters->HeatMapMode = HeatMapMode; \
-	Parameters->HeatMapMax = HeatMapMax; \
-	Parameters->GroundTruthLerp = GroundTruthLerp; \
-	Parameters->GroundTruthBufferSize = GroundTruthPositions.Num(); \
-	Parameters->PositionGroundTruthBuffer = GroundTruthBufferSRV; \
-	Parameters->VertexMapBuffer = VertexMapBufferSRV;
+	Parameters.NumVertices = 0; \
+	Parameters.InputStreamStart = RenderSection.BaseVertexIndex; \
+	Parameters.HeatMapMode = HeatMapMode; \
+	Parameters.HeatMapMax = HeatMapMax; \
+	Parameters.GroundTruthLerp = GroundTruthLerp; \
+	Parameters.GroundTruthBufferSize = GroundTruthPositions.Num(); \
+	Parameters.PositionGroundTruthBuffer = GroundTruthBufferSRV; \
+	Parameters.VertexMapBuffer = VertexMapBufferSRV;
 
 /** 
  * Compute Framework Data Interface for MLDeformer debugging data. 
@@ -98,15 +98,16 @@ namespace UE::MLDeformer
 		virtual void HandleZeroGroundTruthPositions();
 
 		// FComputeDataProviderRenderProxy overrides.
-		virtual void AllocateResources(FRDGBuilder& GraphBuilder) override;
-		virtual void GatherDispatchData(FDispatchSetup const& InDispatchSetup, FCollectedDispatchData& InOutDispatchData) override;
+		bool IsValid(FValidationData const& InValidationData) const override;
+		void AllocateResources(FRDGBuilder& GraphBuilder) override;
+		void GatherDispatchData(FDispatchData const& InDispatchData) override;
 		// ~END FComputeDataProviderRenderProxy overrides.
 
 		TArray<FVector3f>& GetGroundTruthPositions() { return GroundTruthPositions; }
 
 	protected:
 		TObjectPtr<UMLDeformerGraphDebugDataProvider> Provider = nullptr;
-		FSkeletalMeshObject* SkeletalMeshObject;
+		FSkeletalMeshObject* SkeletalMeshObject = nullptr;
 		TArray<FVector3f> GroundTruthPositions;
 		FRHIShaderResourceView* VertexMapBufferSRV = nullptr;
 		FRDGBuffer* GroundTruthBuffer = nullptr;
diff --git a/Engine/Plugins/Runtime/ComputeFramework/Source/ComputeFramework/Private/ComputeGraphWorker.cpp b/Engine/Plugins/Runtime/ComputeFramework/Source/ComputeFramework/Private/ComputeGraphWorker.cpp
--- a/Engine/Plugins/Runtime/ComputeFramework/Source/ComputeFramework/Private/ComputeGraphWorker.cpp
+++ b/Engine/Plugins/Runtime/ComputeFramework/Source/ComputeFramework/Private/ComputeGraphWorker.cpp
@@ -39,6 +39,64 @@ void FComputeGraphTaskWorker::SubmitWork(FRDGBuilder& GraphBuilder, FName InExec
 			FGraphInvocation const& GraphInvocation = GraphInvocations[GraphIndex];
 			FComputeGraphRenderProxy const* GraphRenderProxy = GraphInvocation.GraphRenderProxy;
 
+			TArray<TShaderRef<FComputeKernelShader>> Shaders;
+
+			// Validation phase.
+			// Check if all DataInterfaces are valid.
+			// At the same time gather the permutation id so that we can validate if shader is compiled.
+			bool bIsValid = true;
+			for (int32 KernelIndex = 0; bIsValid && KernelIndex < GraphRenderProxy->KernelInvocations.Num(); ++KernelIndex)
+			{
+				FComputeGraphRenderProxy::FKernelInvocation const& KernelInvocation = GraphRenderProxy->KernelInvocations[KernelIndex];
+
+				TArray<FIntVector> ThreadCounts;
+				const int32 NumSubInvocations = GraphInvocation.DataProviderRenderProxies[KernelInvocation.ExecutionProviderIndex]->GetDispatchThreadCount(ThreadCounts);
+
+				// Iterate shader parameter members to fill the dispatch data structures.
+				// We assume that the members were filled out with a single data interface per member, and that the
+				// order is the same one defined in the KernelInvocation.BoundProviderIndices.
+				TArray<FShaderParametersMetadata::FMember> const& ParamMembers = KernelInvocation.ShaderParameterMetadata->GetMembers();
+
+				FComputeDataProviderRenderProxy::FPermutationData PermutationData{ NumSubInvocations, GraphRenderProxy->ShaderPermutationVectors[KernelIndex] };
+				PermutationData.PermutationIds.AddZeroed(NumSubInvocations);
+
+				for (int32 MemberIndex = 0; bIsValid && MemberIndex < ParamMembers.Num(); ++MemberIndex)
+				{
+					FShaderParametersMetadata::FMember const& Member = ParamMembers[MemberIndex];
+					if (ensure(Member.GetBaseType() == EUniformBufferBaseType::UBMT_NESTED_STRUCT))
+					{
+						const int32 DataProviderIndex = KernelInvocation.BoundProviderIndices[MemberIndex];
+						FComputeDataProviderRenderProxy* DataProvider = GraphInvocation.DataProviderRenderProxies[DataProviderIndex];
+						if (ensure(DataProvider != nullptr))
+						{
+							FComputeDataProviderRenderProxy::FValidationData ValidationData{ NumSubInvocations, Member.GetStructMetadata()->GetSize() };
+							bIsValid &= DataProvider->IsValid(ValidationData);
+
+							if (bIsValid)
+							{
+								DataProvider->GatherPermutations(PermutationData);
+							}
+						}
+					}
+				}
+
+				// Get shader. This can fail if compilation is pending.
+				Shaders.Reserve(Shaders.Num() + NumSubInvocations);
+				for (int32 SubInvocationIndex = 0; bIsValid && SubInvocationIndex < NumSubInvocations; ++SubInvocationIndex)
+				{
+					TShaderRef<FComputeKernelShader> Shader = KernelInvocation.KernelResource->GetShader(PermutationData.PermutationIds[SubInvocationIndex]);
+					bIsValid &= Shader.IsValid();
+					Shaders.Add(Shader);
+				}
+			}
+
+			// If we can't run the graph for any reason, back out now.
+			if (!bIsValid)
+			{
+				continue;
+			}
+
+			// From here on we are committed to submitting the work to the GPU.
 			RDG_EVENT_SCOPE(GraphBuilder, "%s:%s", *GraphInvocation.OwnerName.ToString(), *GraphRenderProxy->GraphName.ToString());
 
 			// Do resource allocation for all the data providers in the graph.
@@ -62,18 +120,12 @@ void FComputeGraphTaskWorker::SubmitWork(FRDGBuilder& GraphBuilder, FName InExec
 				const int32 NumSubInvocations = GraphInvocation.DataProviderRenderProxies[KernelInvocation.ExecutionProviderIndex]->GetDispatchThreadCount(ThreadCounts);
 
 				TStridedView<FComputeKernelShader::FParameters> ParameterArray = GraphBuilder.AllocParameters<FComputeKernelShader::FParameters>(KernelInvocation.ShaderParameterMetadata, NumSubInvocations);
+				FComputeDataProviderRenderProxy::FDispatchData DispatchData{ NumSubInvocations, 0, 0, ParameterArray.GetStride(), reinterpret_cast<uint8*>(&ParameterArray[0]) };
 
 				// Iterate shader parameter members to fill the dispatch data structures.
 				// We assume that the members were filled out with a single data interface per member, and that the
 				// order is the same one defined in the KernelInvocation.BoundProviderIndices.
 				TArray<FShaderParametersMetadata::FMember> const& ParamMembers = KernelInvocation.ShaderParameterMetadata->GetMembers();
-
-				FComputeDataProviderRenderProxy::FCollectedDispatchData DispatchData;
-				DispatchData.ParameterBuffer = reinterpret_cast<uint8*>(&ParameterArray[0]);
-				DispatchData.PermutationId.AddZeroed(NumSubInvocations);
-
-				FComputeDataProviderRenderProxy::FDispatchSetup DispatchSetup{ NumSubInvocations, 0, ParameterArray.GetStride(), 0, GraphRenderProxy->ShaderPermutationVectors[KernelIndex]};
-
 				for (int32 MemberIndex = 0; MemberIndex < ParamMembers.Num(); ++MemberIndex)
 				{
 					FShaderParametersMetadata::FMember const& Member = ParamMembers[MemberIndex];
@@ -83,10 +135,9 @@ void FComputeGraphTaskWorker::SubmitWork(FRDGBuilder& GraphBuilder, FName InExec
 						FComputeDataProviderRenderProxy* DataProvider = GraphInvocation.DataProviderRenderProxies[DataProviderIndex];
 						if (ensure(DataProvider != nullptr))
 						{
-							DispatchSetup.ParameterBufferOffset = Member.GetOffset();
-							DispatchSetup.ParameterStructSizeForValidation = Member.GetStructMetadata()->GetSize();
-
-							DataProvider->GatherDispatchData(DispatchSetup, DispatchData);
+							DispatchData.ParameterStructSize = Member.GetStructMetadata()->GetSize();
+							DispatchData.ParameterBufferOffset = Member.GetOffset();
+							DataProvider->GatherDispatchData(DispatchData);
 						}
 					}
 				}
@@ -94,7 +145,7 @@ void FComputeGraphTaskWorker::SubmitWork(FRDGBuilder& GraphBuilder, FName InExec
 				// Dispatch work to the render graph.
 				for (int32 SubInvocationIndex = 0; SubInvocationIndex < NumSubInvocations; ++SubInvocationIndex)
 				{
-					TShaderRef<FComputeKernelShader> Shader = KernelInvocation.KernelResource->GetShader(DispatchData.PermutationId[SubInvocationIndex]);
+					TShaderRef<FComputeKernelShader> Shader = Shaders[KernelIndex * NumSubInvocations + SubInvocationIndex];
 					const FIntVector GroupCount = FComputeShaderUtils::GetGroupCount(ThreadCounts[SubInvocationIndex], KernelInvocation.KernelGroupSize);
 
 					FComputeShaderUtils::AddPass(
diff --git a/Engine/Plugins/Runtime/ComputeFramework/Source/ComputeFramework/Public/ComputeFramework/ComputeDataProvider.h b/Engine/Plugins/Runtime/ComputeFramework/Source/ComputeFramework/Public/ComputeFramework/ComputeDataProvider.h
--- a/Engine/Plugins/Runtime/ComputeFramework/Source/ComputeFramework/Public/ComputeFramework/ComputeDataProvider.h
+++ b/Engine/Plugins/Runtime/ComputeFramework/Source/ComputeFramework/Public/ComputeFramework/ComputeDataProvider.h
@@ -3,6 +3,7 @@
 #pragma once
 
 #include "CoreMinimal.h"
+#include "Containers/StridedView.h"
 #include "ComputeDataProvider.generated.h"
 
 class FComputeDataProviderRenderProxy;
@@ -33,7 +34,7 @@ public:
 
 /**
  * Compute Framework Data Provider Proxy. 
- * A concrete instance of this is created by the UComputeDataProvider and used for the render thread gathering of data for a Compute Kernel. 
+ * A concrete instance of this is created by the UComputeDataProvider gathering of data for a Compute Kernel on the render thread. 
  */
 class COMPUTEFRAMEWORK_API FComputeDataProviderRenderProxy
 {
@@ -46,25 +47,82 @@ public:
 	*/
 	virtual int32 GetDispatchThreadCount(TArray<FIntVector>& ThreadCounts) const { return 0; };
 
+	/** Data needed for validation. */
+	struct FValidationData
+	{
+		int32 NumInvocations;
+		int32 ParameterStructSize;
+	};
+
+	/**
+	 * Validates that we are OK to dispatch work.
+	 * Default implementation returns false.
+	 */
+	virtual bool IsValid(FValidationData const& InValidationData) const { return false; }
+
+	/** Data needed for setting permuations. */
+	struct FPermutationData
+	{
+		const int32 NumInvocations;
+		const FComputeKernelPermutationVector& PermutationVector;
+		TArray<int32> PermutationIds;
+	};
+
+	/** 
+	 * Gathers permutation bits for each invocation.
+	 * This is called before any calls to AllocateResources() because we validate all requested shaders before doing any further work.
+	 */
+	virtual void GatherPermutations(FPermutationData& InOutPermutationData) const {}
+
 	/* Called once before any calls to GatherDispatchData() to allow any RDG resource allocation. */
 	virtual void AllocateResources(FRDGBuilder& GraphBuilder) {}
 
-	/** */
+	/** Setup needed to gather dispatch data. */
+	struct FDispatchData
+	{
+		int32 NumInvocations;
+		int32 ParameterStructSize;
+		int32 ParameterBufferOffset;
+		int32 ParameterBufferStride;
+		uint8* ParameterBuffer;
+	};
+
+	/** Collect parameter data required to dispatch work. */
+	virtual void GatherDispatchData(FDispatchData const& InDispatchData) {}
+
+protected:
+	/** Helper for making an FStridedView over the FDispatchData. */
+	template <typename ElementType>
+	TStridedView<ElementType> MakeStridedParameterView(FDispatchData const& InDispatchData)
+	{
+		return TStridedView<ElementType>(InDispatchData.ParameterBufferStride, (ElementType*)(InDispatchData.ParameterBuffer + InDispatchData.ParameterBufferOffset), InDispatchData.NumInvocations);
+	}
+
+protected:
+	//~ Start deprecation of old GatherDispatchData() interface.
 	struct FDispatchSetup
 	{
+		UE_DEPRECATED(5.2, "Convert to using the new GatherDispatchData() that takes FDispatchData.")
 		int32 NumInvocations;
+		UE_DEPRECATED(5.2, "Convert to using the new GatherDispatchData() that takes FDispatchData.")
 		int32 ParameterBufferOffset;
+		UE_DEPRECATED(5.2, "Convert to using the new GatherDispatchData() that takes FDispatchData.")
 		int32 ParameterBufferStride;
+		UE_DEPRECATED(5.2, "Convert to using the new GatherDispatchData() that takes FDispatchData.")
 		int32 ParameterStructSizeForValidation;
+		UE_DEPRECATED(5.2, "Convert to using the new GatherDispatchData() that takes FDispatchData.")
 		FComputeKernelPermutationVector const& PermutationVector;
 	};
 
-	/** Collected data during kernel dispatch. */
 	struct FCollectedDispatchData
 	{
+		UE_DEPRECATED(5.2, "Convert to using the new GatherDispatchData() that takes FDispatchData.")
 		uint8* ParameterBuffer;
+		UE_DEPRECATED(5.2, "Convert to using the new GatherDispatchData() that takes FDispatchData.")
 		TArray<int32> PermutationId;
 	};
 
+	UE_DEPRECATED(5.2, "Convert to using the new GatherDispatchData() that takes FDispatchData.")
 	virtual void GatherDispatchData(FDispatchSetup const& InDispatchSetup, FCollectedDispatchData& InOutDispatchData) {}
+	//~ End deprecation of old GatherDispatchData() interface.
 };
diff --git a/Engine/Plugins/Runtime/HairStrands/Source/HairStrandsCore/Private/Deformer/DeformerDataInterfaceGroom.cpp b/Engine/Plugins/Runtime/HairStrands/Source/HairStrandsCore/Private/Deformer/DeformerDataInterfaceGroom.cpp
--- a/Engine/Plugins/Runtime/HairStrands/Source/HairStrandsCore/Private/Deformer/DeformerDataInterfaceGroom.cpp
+++ b/Engine/Plugins/Runtime/HairStrands/Source/HairStrandsCore/Private/Deformer/DeformerDataInterfaceGroom.cpp
@@ -148,6 +148,24 @@ FOptimusGroomDataProviderProxy::FOptimusGroomDataProviderProxy(UGroomComponent*
 	GroomComponent = InGroomComponent;
 }
 
+bool FOptimusGroomDataProviderProxy::IsValid(FValidationData const& InValidationData) const
+{
+	if (InValidationData.ParameterStructSize != sizeof(FParameters))
+	{
+		return false;
+	}
+	if (GroomComponent == nullptr)
+	{
+		return false;
+	}
+	if (InValidationData.NumInvocations != GroomComponent->GetGroupCount())
+	{
+		return false;
+	}
+	
+	return true;
+}
+
 void FOptimusGroomDataProviderProxy::AllocateResources(FRDGBuilder& GraphBuilder)
 {
 	Resources.Empty();
@@ -167,17 +185,14 @@ void FOptimusGroomDataProviderProxy::AllocateResources(FRDGBuilder& GraphBuilder
 }
 
 FHairGroupPublicData::FVertexFactoryInput ComputeHairStrandsVertexInputData(const FHairGroupInstance* Instance, EGroomViewMode ViewMode);
-void FOptimusGroomDataProviderProxy::GatherDispatchData(FDispatchSetup const& InDispatchSetup, FCollectedDispatchData& InOutDispatchData)
+void FOptimusGroomDataProviderProxy::GatherDispatchData(FDispatchData const& InDispatchData)
 {
-	if (!ensure(InDispatchSetup.ParameterStructSizeForValidation == sizeof(FGroomDataInterfaceParameters)))
-	{
-		return;
-	}
-
 	const uint32 InstanceCount = GroomComponent ? GroomComponent->GetGroupCount() : 0;
-	check(InDispatchSetup.NumInvocations == InstanceCount);
+	check(InDispatchData.NumInvocations == InstanceCount);
 	check(uint32(Resources.Num()) == InstanceCount);
-	for (int32 InvocationIndex = 0; InvocationIndex < InDispatchSetup.NumInvocations; ++InvocationIndex)
+
+	const TStridedView<FParameters> ParameterArray = MakeStridedParameterView<FParameters>(InDispatchData);
+	for (int32 InvocationIndex = 0; InvocationIndex < ParameterArray.Num(); ++InvocationIndex)
 	{
 		if (FHairGroupInstance* Instance = GroomComponent->GetGroupInstance(InvocationIndex))
 		{
@@ -187,25 +202,25 @@ void FOptimusGroomDataProviderProxy::GatherDispatchData(FDispatchSetup const& In
 
 			const FHairGroupPublicData::FVertexFactoryInput VFInput = ComputeHairStrandsVertexInputData(Instance, EGroomViewMode::None);
 			
-			FGroomDataInterfaceParameters* Parameters = (FGroomDataInterfaceParameters*)(InOutDispatchData.ParameterBuffer + InDispatchSetup.ParameterBufferOffset + InDispatchSetup.ParameterBufferStride * InvocationIndex);
-			Parameters->NumControlPoints = NumControlPoints;
-			Parameters->NumCurves = NumCurves;
-			Parameters->VF_HairRadius = VFInput.Strands.HairRadius;
-			Parameters->VF_HairLength = VFInput.Strands.HairLength;
+			FParameters& Parameters = ParameterArray[InvocationIndex];
+			Parameters.NumControlPoints = NumControlPoints;
+			Parameters.NumCurves = NumCurves;
+			Parameters.VF_HairRadius = VFInput.Strands.HairRadius;
+			Parameters.VF_HairLength = VFInput.Strands.HairLength;
 
 			if (bIsSRVValid)
 			{
-				Parameters->PositionBuffer = Resources[InvocationIndex].PositionSRV;
-				Parameters->Attribute0Buffer = Resources[InvocationIndex].Attribute0SRV;
-				Parameters->Attribute1Buffer = Resources[InvocationIndex].Attribute1SRV;
-				Parameters->CurveBuffer = Resources[InvocationIndex].CurveSRV;
+				Parameters.PositionBuffer = Resources[InvocationIndex].PositionSRV;
+				Parameters.Attribute0Buffer = Resources[InvocationIndex].Attribute0SRV;
+				Parameters.Attribute1Buffer = Resources[InvocationIndex].Attribute1SRV;
+				Parameters.CurveBuffer = Resources[InvocationIndex].CurveSRV;
 			}
 			else
 			{
-				Parameters->PositionBuffer   = Resources[InvocationIndex].FallbackSRV;
-				Parameters->Attribute0Buffer = Resources[InvocationIndex].FallbackSRV;
-				Parameters->Attribute1Buffer = Resources[InvocationIndex].FallbackSRV;
-				Parameters->CurveBuffer		 = Resources[InvocationIndex].FallbackSRV;
+				Parameters.PositionBuffer   = Resources[InvocationIndex].FallbackSRV;
+				Parameters.Attribute0Buffer = Resources[InvocationIndex].FallbackSRV;
+				Parameters.Attribute1Buffer = Resources[InvocationIndex].FallbackSRV;
+				Parameters.CurveBuffer		 = Resources[InvocationIndex].FallbackSRV;
 			}
 		}
 	}
diff --git a/Engine/Plugins/Runtime/HairStrands/Source/HairStrandsCore/Private/Deformer/DeformerDataInterfaceGroom.h b/Engine/Plugins/Runtime/HairStrands/Source/HairStrandsCore/Private/Deformer/DeformerDataInterfaceGroom.h
--- a/Engine/Plugins/Runtime/HairStrands/Source/HairStrandsCore/Private/Deformer/DeformerDataInterfaceGroom.h
+++ b/Engine/Plugins/Runtime/HairStrands/Source/HairStrandsCore/Private/Deformer/DeformerDataInterfaceGroom.h
@@ -7,6 +7,7 @@
 #include "ComputeFramework/ComputeDataProvider.h"
 #include "DeformerDataInterfaceGroom.generated.h"
 
+class FGroomDataInterfaceParameters;
 class UGroomComponent;
 
 /** Compute Framework Data Interface for reading groom. */
@@ -55,11 +56,13 @@ public:
 	FOptimusGroomDataProviderProxy(UGroomComponent* InGroomComponent);
 
 	//~ Begin FComputeDataProviderRenderProxy Interface
+	bool IsValid(FValidationData const& InValidationData) const override;
 	void AllocateResources(FRDGBuilder& GraphBuilder) override;
-	void GatherDispatchData(FDispatchSetup const& InDispatchSetup, FCollectedDispatchData& InOutDispatchData) override;
+	void GatherDispatchData(FDispatchData const& InDispatchData) override;
 	//~ End FComputeDataProviderRenderProxy Interface
 
 private:
+	using FParameters = FGroomDataInterfaceParameters;
 
 	UGroomComponent* GroomComponent = nullptr; // Should it be HairInstance instead?
 	struct FResources
diff --git a/Engine/Plugins/Runtime/HairStrands/Source/HairStrandsCore/Private/Deformer/DeformerDataInterfaceGroomExec.cpp b/Engine/Plugins/Runtime/HairStrands/Source/HairStrandsCore/Private/Deformer/DeformerDataInterfaceGroomExec.cpp
--- a/Engine/Plugins/Runtime/HairStrands/Source/HairStrandsCore/Private/Deformer/DeformerDataInterfaceGroomExec.cpp
+++ b/Engine/Plugins/Runtime/HairStrands/Source/HairStrandsCore/Private/Deformer/DeformerDataInterfaceGroomExec.cpp
@@ -109,27 +109,31 @@ int32 FOptimusGroomExecDataProviderProxy::GetDispatchThreadCount(TArray<FIntVect
 	return NumInvocations;
 }
 
-void FOptimusGroomExecDataProviderProxy::GatherDispatchData(FDispatchSetup const& InDispatchSetup, FCollectedDispatchData& InOutDispatchData)
+bool FOptimusGroomExecDataProviderProxy::IsValid(FValidationData const& InValidationData) const
 {
-	if (!ensure(InDispatchSetup.ParameterStructSizeForValidation == sizeof(FGroomExecDataInterfaceParameters)))
+	if (InValidationData.ParameterStructSize != sizeof(FParameters))
 	{
-		return;
+		return false;
 	}
-
-	const int32 NumInvocations = GroomComponent ? GroomComponent->GetGroupCount() : 0;
-	if (!ensure(NumInvocations == InDispatchSetup.NumInvocations))
+	if (GroomComponent == nullptr || GroomComponent->GetGroupCount() == 0)
 	{
-		return;
+		return false;
 	}
 
-	for (int32 InvocationIndex = 0; InvocationIndex < NumInvocations; ++InvocationIndex)
+	return true;
+}
+
+void FOptimusGroomExecDataProviderProxy::GatherDispatchData(FDispatchData const& InDispatchData)
+{
+	const TStridedView<FParameters> ParameterArray = MakeStridedParameterView<FParameters>(InDispatchData);
+	for (int32 InvocationIndex = 0; InvocationIndex < ParameterArray.Num(); ++InvocationIndex)
 	{
 		FHairGroupInstance* Instance = GroomComponent->GetGroupInstance(InvocationIndex);
 		const int32 NumControlPoints = Instance->Strands.Data->PointCount;
 		const int32 NumCurves = Instance->Strands.Data->CurveCount;
 		const int32 NumThreads = Domain == EOptimusGroomExecDomain::ControlPoint ? NumControlPoints : NumCurves;
 
-		FGroomExecDataInterfaceParameters* Parameters = (FGroomExecDataInterfaceParameters*)(InOutDispatchData.ParameterBuffer + InDispatchSetup.ParameterBufferOffset + InDispatchSetup.ParameterBufferStride * InvocationIndex);
-		Parameters->NumThreads = FIntVector(NumThreads, 1, 1);
+		FParameters& Parameters = ParameterArray[InvocationIndex];
+		Parameters.NumThreads = FIntVector(NumThreads, 1, 1);
 	}
 }
diff --git a/Engine/Plugins/Runtime/HairStrands/Source/HairStrandsCore/Private/Deformer/DeformerDataInterfaceGroomExec.h b/Engine/Plugins/Runtime/HairStrands/Source/HairStrandsCore/Private/Deformer/DeformerDataInterfaceGroomExec.h
--- a/Engine/Plugins/Runtime/HairStrands/Source/HairStrandsCore/Private/Deformer/DeformerDataInterfaceGroomExec.h
+++ b/Engine/Plugins/Runtime/HairStrands/Source/HairStrandsCore/Private/Deformer/DeformerDataInterfaceGroomExec.h
@@ -6,6 +6,7 @@
 #include "ComputeFramework/ComputeDataProvider.h"
 #include "DeformerDataInterfaceGroomExec.generated.h"
 
+class FGroomExecDataInterfaceParameters;
 class UGroomComponent;
 
 UENUM()
@@ -72,10 +73,13 @@ public:
 
 	//~ Begin FComputeDataProviderRenderProxy Interface
 	int32 GetDispatchThreadCount(TArray<FIntVector>& ThreadCounts) const override;
-	void GatherDispatchData(FDispatchSetup const& InDispatchSetup, FCollectedDispatchData& InOutDispatchData) override;
+	bool IsValid(FValidationData const& InValidationData) const override;
+	void GatherDispatchData(FDispatchData const& InDispatchData) override;
 	//~ End FComputeDataProviderRenderProxy Interface
 
 private:
+	using FParameters = FGroomExecDataInterfaceParameters;
+
 	UGroomComponent* GroomComponent = nullptr;
 	EOptimusGroomExecDomain Domain = EOptimusGroomExecDomain::ControlPoint;
 };
diff --git a/Engine/Plugins/Runtime/HairStrands/Source/HairStrandsCore/Private/Deformer/DeformerDataInterfaceGroomWrite.cpp b/Engine/Plugins/Runtime/HairStrands/Source/HairStrandsCore/Private/Deformer/DeformerDataInterfaceGroomWrite.cpp
--- a/Engine/Plugins/Runtime/HairStrands/Source/HairStrandsCore/Private/Deformer/DeformerDataInterfaceGroomWrite.cpp
+++ b/Engine/Plugins/Runtime/HairStrands/Source/HairStrandsCore/Private/Deformer/DeformerDataInterfaceGroomWrite.cpp
@@ -111,6 +111,20 @@ FOptimusGroomWriteDataProviderProxy::FOptimusGroomWriteDataProviderProxy(UGroomC
 	}
 }
 
+bool FOptimusGroomWriteDataProviderProxy::IsValid(FValidationData const& InValidationData) const
+{
+	if (InValidationData.ParameterStructSize != sizeof(FParameters))
+	{
+		return false;
+	}
+	if (InValidationData.NumInvocations != Instances.Num())
+	{
+		return false;
+	}
+
+	return true;
+}
+
 void FOptimusGroomWriteDataProviderProxy::AllocateResources(FRDGBuilder& GraphBuilder)
 {
 	for (FHairGroupInstance* GroomInstance : Instances)
@@ -140,16 +154,10 @@ void FOptimusGroomWriteDataProviderProxy::AllocateResources(FRDGBuilder& GraphBu
 	}
 }
 
-void FOptimusGroomWriteDataProviderProxy::GatherDispatchData(FDispatchSetup const& InDispatchSetup, FCollectedDispatchData& InOutDispatchData)
+void FOptimusGroomWriteDataProviderProxy::GatherDispatchData(FDispatchData const& InDispatchData)
 {
-	if (!ensure(InDispatchSetup.ParameterStructSizeForValidation == sizeof(FGroomWriteDataInterfaceParameters)))
-	{
-		return;
-	}
-
-	check(InDispatchSetup.NumInvocations == Instances.Num());
-
-	for (int32 InvocationIndex = 0; InvocationIndex < InDispatchSetup.NumInvocations; ++InvocationIndex)
+	const TStridedView<FParameters> ParameterArray = MakeStridedParameterView<FParameters>(InDispatchData);
+	for (int32 InvocationIndex = 0; InvocationIndex < ParameterArray.Num(); ++InvocationIndex)
 	{
 		if (FHairGroupInstance* GroomInstance = Instances[InvocationIndex])
 		{
@@ -159,12 +167,12 @@ void FOptimusGroomWriteDataProviderProxy::GatherDispatchData(FDispatchSetup cons
 			const int32 NumCurves = bValid ? GroomInstance->Strands.Data->CurveCount : 0;
 			const int32 NumControlPoints = bValid ? GroomInstance->Strands.Data->PointCount : 0;
 
-			FGroomWriteDataInterfaceParameters* Parameters = (FGroomWriteDataInterfaceParameters*)(InOutDispatchData.ParameterBuffer + InDispatchSetup.ParameterBufferOffset + InDispatchSetup.ParameterBufferStride * InvocationIndex);
-			Parameters->NumControlPoints = NumControlPoints;
-			Parameters->NumCurves = NumCurves;
-			Parameters->OutputStreamStart = 0;
-			Parameters->PositionBufferSRV = bValid ? Resource.PositionBufferSRV : Resource.PositionBufferSRV_fallback;
-			Parameters->PositionBufferUAV = bValid ? Resource.PositionBufferUAV : Resource.PositionBufferUAV_fallback;
+			FParameters& Parameters = ParameterArray[InvocationIndex];
+			Parameters.NumControlPoints = NumControlPoints;
+			Parameters.NumCurves = NumCurves;
+			Parameters.OutputStreamStart = 0;
+			Parameters.PositionBufferSRV = bValid ? Resource.PositionBufferSRV : Resource.PositionBufferSRV_fallback;
+			Parameters.PositionBufferUAV = bValid ? Resource.PositionBufferUAV : Resource.PositionBufferUAV_fallback;
 		}
 	}
 }
\ No newline at end of file
diff --git a/Engine/Plugins/Runtime/HairStrands/Source/HairStrandsCore/Private/Deformer/DeformerDataInterfaceGroomWrite.h b/Engine/Plugins/Runtime/HairStrands/Source/HairStrandsCore/Private/Deformer/DeformerDataInterfaceGroomWrite.h
--- a/Engine/Plugins/Runtime/HairStrands/Source/HairStrandsCore/Private/Deformer/DeformerDataInterfaceGroomWrite.h
+++ b/Engine/Plugins/Runtime/HairStrands/Source/HairStrandsCore/Private/Deformer/DeformerDataInterfaceGroomWrite.h
@@ -6,11 +6,12 @@
 #include "ComputeFramework/ComputeDataProvider.h"
 #include "DeformerDataInterfaceGroomWrite.generated.h"
 
-class UGroomComponent;
+class FGroomWriteDataInterfaceParameters;
+struct FHairGroupInstance;
 class FRDGBuffer;
-class FRDGBufferUAV;
 class FRDGBufferSRV;
-struct FHairGroupInstance;
+class FRDGBufferUAV;
+class UGroomComponent;
 
 /** Compute Framework Data Interface for writing skinned mesh. */
 UCLASS(Category = ComputeFramework)
@@ -61,11 +62,14 @@ public:
 	FOptimusGroomWriteDataProviderProxy(UGroomComponent* InGroomComponent, uint64 InOutputMask);
 
 	//~ Begin FComputeDataProviderRenderProxy Interface
+	bool IsValid(FValidationData const& InValidationData) const override;
 	void AllocateResources(FRDGBuilder& GraphBuilder) override;
-	void GatherDispatchData(FDispatchSetup const& InDispatchSetup, FCollectedDispatchData& InOutDispatchData) override;
+	void GatherDispatchData(FDispatchData const& InDispatchData) override;
 	//~ End FComputeDataProviderRenderProxy Interface
 
 private:
+	using FParameters = FGroomWriteDataInterfaceParameters;
+
 	TArray<FHairGroupInstance*> Instances;
 	uint64 OutputMask;
 
