Optimus: Deprecate UComputeDataProvider::IsValid() and remove concrete implementations that are no longer called.

diff --git a/Engine/Plugins/Animation/DeformerGraph/Source/OptimusCore/Private/DataInterfaces/OptimusDataInterfaceAnimAttribute.cpp b/Engine/Plugins/Animation/DeformerGraph/Source/OptimusCore/Private/DataInterfaces/OptimusDataInterfaceAnimAttribute.cpp
--- a/Engine/Plugins/Animation/DeformerGraph/Source/OptimusCore/Private/DataInterfaces/OptimusDataInterfaceAnimAttribute.cpp
+++ b/Engine/Plugins/Animation/DeformerGraph/Source/OptimusCore/Private/DataInterfaces/OptimusDataInterfaceAnimAttribute.cpp
@@ -616,11 +616,6 @@ void UOptimusAnimAttributeDataProvider::Init(
 	}
 }
 
-bool UOptimusAnimAttributeDataProvider::IsValid() const
-{
-	return SkeletalMesh != nullptr;
-}
-
 FComputeDataProviderRenderProxy* UOptimusAnimAttributeDataProvider::GetRenderProxy()
 {
 	FOptimusAnimAttributeDataProviderProxy* Proxy = new FOptimusAnimAttributeDataProviderProxy(AttributeBufferSize, TotalNumArrays);
diff --git a/Engine/Plugins/Animation/DeformerGraph/Source/OptimusCore/Private/DataInterfaces/OptimusDataInterfaceAnimAttribute.h b/Engine/Plugins/Animation/DeformerGraph/Source/OptimusCore/Private/DataInterfaces/OptimusDataInterfaceAnimAttribute.h
--- a/Engine/Plugins/Animation/DeformerGraph/Source/OptimusCore/Private/DataInterfaces/OptimusDataInterfaceAnimAttribute.h
+++ b/Engine/Plugins/Animation/DeformerGraph/Source/OptimusCore/Private/DataInterfaces/OptimusDataInterfaceAnimAttribute.h
@@ -196,7 +196,6 @@ public:
 	int32 TotalNumArrays = 0;
 	
 	//~ Begin UComputeDataProvider Interface
-	bool IsValid() const override;
 	FComputeDataProviderRenderProxy* GetRenderProxy() override;
 	//~ End UComputeDataProvider Interface
 };
diff --git a/Engine/Plugins/Animation/DeformerGraph/Source/OptimusCore/Private/DataInterfaces/OptimusDataInterfaceCloth.cpp b/Engine/Plugins/Animation/DeformerGraph/Source/OptimusCore/Private/DataInterfaces/OptimusDataInterfaceCloth.cpp
--- a/Engine/Plugins/Animation/DeformerGraph/Source/OptimusCore/Private/DataInterfaces/OptimusDataInterfaceCloth.cpp
+++ b/Engine/Plugins/Animation/DeformerGraph/Source/OptimusCore/Private/DataInterfaces/OptimusDataInterfaceCloth.cpp
@@ -115,13 +115,6 @@ UComputeDataProvider* UOptimusClothDataInterface::CreateDataProvider(TObjectPtr<
 }
 
 
-bool UOptimusClothDataProvider::IsValid() const
-{
-	return
-		SkinnedMesh != nullptr &&
-		SkinnedMesh->MeshObject != nullptr;
-}
-
 FComputeDataProviderRenderProxy* UOptimusClothDataProvider::GetRenderProxy()
 {
 	return new FOptimusClothDataProviderProxy(SkinnedMesh);
diff --git a/Engine/Plugins/Animation/DeformerGraph/Source/OptimusCore/Private/DataInterfaces/OptimusDataInterfaceCloth.h b/Engine/Plugins/Animation/DeformerGraph/Source/OptimusCore/Private/DataInterfaces/OptimusDataInterfaceCloth.h
--- a/Engine/Plugins/Animation/DeformerGraph/Source/OptimusCore/Private/DataInterfaces/OptimusDataInterfaceCloth.h
+++ b/Engine/Plugins/Animation/DeformerGraph/Source/OptimusCore/Private/DataInterfaces/OptimusDataInterfaceCloth.h
@@ -45,7 +45,6 @@ public:
 	TObjectPtr<USkinnedMeshComponent> SkinnedMesh = nullptr;
 
 	//~ Begin UComputeDataProvider Interface
-	bool IsValid() const override;
 	FComputeDataProviderRenderProxy* GetRenderProxy() override;
 	//~ End UComputeDataProvider Interface
 };
diff --git a/Engine/Plugins/Animation/DeformerGraph/Source/OptimusCore/Private/DataInterfaces/OptimusDataInterfaceConnectivity.cpp b/Engine/Plugins/Animation/DeformerGraph/Source/OptimusCore/Private/DataInterfaces/OptimusDataInterfaceConnectivity.cpp
--- a/Engine/Plugins/Animation/DeformerGraph/Source/OptimusCore/Private/DataInterfaces/OptimusDataInterfaceConnectivity.cpp
+++ b/Engine/Plugins/Animation/DeformerGraph/Source/OptimusCore/Private/DataInterfaces/OptimusDataInterfaceConnectivity.cpp
@@ -221,14 +221,6 @@ UComputeDataProvider* UOptimusConnectivityDataInterface::CreateDataProvider(TObj
 }
 
 
-bool UOptimusConnectivityDataProvider::IsValid() const
-{
-	return
-		SkinnedMesh != nullptr &&
-		SkinnedMesh->MeshObject != nullptr &&
-		AdjacencyBufferPerLod.Num();
-}
-
 FComputeDataProviderRenderProxy* UOptimusConnectivityDataProvider::GetRenderProxy()
 {
 	return new FOptimusConnectivityDataProviderProxy(SkinnedMesh, AdjacencyBufferPerLod);
diff --git a/Engine/Plugins/Animation/DeformerGraph/Source/OptimusCore/Private/DataInterfaces/OptimusDataInterfaceConnectivity.h b/Engine/Plugins/Animation/DeformerGraph/Source/OptimusCore/Private/DataInterfaces/OptimusDataInterfaceConnectivity.h
--- a/Engine/Plugins/Animation/DeformerGraph/Source/OptimusCore/Private/DataInterfaces/OptimusDataInterfaceConnectivity.h
+++ b/Engine/Plugins/Animation/DeformerGraph/Source/OptimusCore/Private/DataInterfaces/OptimusDataInterfaceConnectivity.h
@@ -50,7 +50,6 @@ public:
 	TObjectPtr<USkinnedMeshComponent> SkinnedMesh = nullptr;
 
 	//~ Begin UComputeDataProvider Interface
-	bool IsValid() const override;
 	FComputeDataProviderRenderProxy* GetRenderProxy() override;
 	//~ End UComputeDataProvider Interface
 
diff --git a/Engine/Plugins/Animation/DeformerGraph/Source/OptimusCore/Private/DataInterfaces/OptimusDataInterfaceDebugDraw.cpp b/Engine/Plugins/Animation/DeformerGraph/Source/OptimusCore/Private/DataInterfaces/OptimusDataInterfaceDebugDraw.cpp
--- a/Engine/Plugins/Animation/DeformerGraph/Source/OptimusCore/Private/DataInterfaces/OptimusDataInterfaceDebugDraw.cpp
+++ b/Engine/Plugins/Animation/DeformerGraph/Source/OptimusCore/Private/DataInterfaces/OptimusDataInterfaceDebugDraw.cpp
@@ -102,11 +102,6 @@ UComputeDataProvider* UOptimusDebugDrawDataInterface::CreateDataProvider(TObject
 }
 
 
-bool UOptimusDebugDrawDataProvider::IsValid() const
-{
-	return PrimitiveComponent != nullptr;
-}
-
 FComputeDataProviderRenderProxy* UOptimusDebugDrawDataProvider::GetRenderProxy()
 {
 	return new FOptimusDebugDrawDataProviderProxy(PrimitiveComponent, DebugDrawParameters);
diff --git a/Engine/Plugins/Animation/DeformerGraph/Source/OptimusCore/Private/DataInterfaces/OptimusDataInterfaceDebugDraw.h b/Engine/Plugins/Animation/DeformerGraph/Source/OptimusCore/Private/DataInterfaces/OptimusDataInterfaceDebugDraw.h
--- a/Engine/Plugins/Animation/DeformerGraph/Source/OptimusCore/Private/DataInterfaces/OptimusDataInterfaceDebugDraw.h
+++ b/Engine/Plugins/Animation/DeformerGraph/Source/OptimusCore/Private/DataInterfaces/OptimusDataInterfaceDebugDraw.h
@@ -79,7 +79,6 @@ public:
 	FOptimusDebugDrawParameters DebugDrawParameters;
 
 	//~ Begin UComputeDataProvider Interface
-	bool IsValid() const override;
 	FComputeDataProviderRenderProxy* GetRenderProxy() override;
 	//~ End UComputeDataProvider Interface
 };
diff --git a/Engine/Plugins/Animation/DeformerGraph/Source/OptimusCore/Private/DataInterfaces/OptimusDataInterfaceDuplicateVertices.cpp b/Engine/Plugins/Animation/DeformerGraph/Source/OptimusCore/Private/DataInterfaces/OptimusDataInterfaceDuplicateVertices.cpp
--- a/Engine/Plugins/Animation/DeformerGraph/Source/OptimusCore/Private/DataInterfaces/OptimusDataInterfaceDuplicateVertices.cpp
+++ b/Engine/Plugins/Animation/DeformerGraph/Source/OptimusCore/Private/DataInterfaces/OptimusDataInterfaceDuplicateVertices.cpp
@@ -104,13 +104,6 @@ UComputeDataProvider* UOptimusDuplicateVerticesDataInterface::CreateDataProvider
 }
 
 
-bool UOptimusDuplicateVerticesDataProvider::IsValid() const
-{
-	return
-		SkinnedMesh != nullptr &&
-		SkinnedMesh->MeshObject != nullptr;
-}
-
 FComputeDataProviderRenderProxy* UOptimusDuplicateVerticesDataProvider::GetRenderProxy()
 {
 	return new FOptimusDuplicateVerticesDataProviderProxy(SkinnedMesh);
diff --git a/Engine/Plugins/Animation/DeformerGraph/Source/OptimusCore/Private/DataInterfaces/OptimusDataInterfaceDuplicateVertices.h b/Engine/Plugins/Animation/DeformerGraph/Source/OptimusCore/Private/DataInterfaces/OptimusDataInterfaceDuplicateVertices.h
--- a/Engine/Plugins/Animation/DeformerGraph/Source/OptimusCore/Private/DataInterfaces/OptimusDataInterfaceDuplicateVertices.h
+++ b/Engine/Plugins/Animation/DeformerGraph/Source/OptimusCore/Private/DataInterfaces/OptimusDataInterfaceDuplicateVertices.h
@@ -49,7 +49,6 @@ public:
 	TObjectPtr<USkinnedMeshComponent> SkinnedMesh = nullptr;
 
 	//~ Begin UComputeDataProvider Interface
-	bool IsValid() const override;
 	FComputeDataProviderRenderProxy* GetRenderProxy() override;
 	//~ End UComputeDataProvider Interface
 };
diff --git a/Engine/Plugins/Animation/DeformerGraph/Source/OptimusCore/Private/DataInterfaces/OptimusDataInterfaceMorphTarget.cpp b/Engine/Plugins/Animation/DeformerGraph/Source/OptimusCore/Private/DataInterfaces/OptimusDataInterfaceMorphTarget.cpp
--- a/Engine/Plugins/Animation/DeformerGraph/Source/OptimusCore/Private/DataInterfaces/OptimusDataInterfaceMorphTarget.cpp
+++ b/Engine/Plugins/Animation/DeformerGraph/Source/OptimusCore/Private/DataInterfaces/OptimusDataInterfaceMorphTarget.cpp
@@ -92,13 +92,6 @@ UComputeDataProvider* UOptimusMorphTargetDataInterface::CreateDataProvider(TObje
 }
 
 
-bool UOptimusMorphTargetDataProvider::IsValid() const
-{
-	return
-		SkinnedMesh != nullptr &&
-		SkinnedMesh->MeshObject != nullptr;
-}
-
 FComputeDataProviderRenderProxy* UOptimusMorphTargetDataProvider::GetRenderProxy()
 {
 	return new FOptimusMorphTargetDataProviderProxy(SkinnedMesh);
diff --git a/Engine/Plugins/Animation/DeformerGraph/Source/OptimusCore/Private/DataInterfaces/OptimusDataInterfaceMorphTarget.h b/Engine/Plugins/Animation/DeformerGraph/Source/OptimusCore/Private/DataInterfaces/OptimusDataInterfaceMorphTarget.h
--- a/Engine/Plugins/Animation/DeformerGraph/Source/OptimusCore/Private/DataInterfaces/OptimusDataInterfaceMorphTarget.h
+++ b/Engine/Plugins/Animation/DeformerGraph/Source/OptimusCore/Private/DataInterfaces/OptimusDataInterfaceMorphTarget.h
@@ -46,7 +46,6 @@ public:
 	TObjectPtr<USkinnedMeshComponent> SkinnedMesh = nullptr;
 
 	//~ Begin UComputeDataProvider Interface
-	bool IsValid() const override;
 	FComputeDataProviderRenderProxy* GetRenderProxy() override;
 	//~ End UComputeDataProvider Interface
 };
diff --git a/Engine/Plugins/Animation/DeformerGraph/Source/OptimusCore/Private/DataInterfaces/OptimusDataInterfaceRawBuffer.cpp b/Engine/Plugins/Animation/DeformerGraph/Source/OptimusCore/Private/DataInterfaces/OptimusDataInterfaceRawBuffer.cpp
--- a/Engine/Plugins/Animation/DeformerGraph/Source/OptimusCore/Private/DataInterfaces/OptimusDataInterfaceRawBuffer.cpp
+++ b/Engine/Plugins/Animation/DeformerGraph/Source/OptimusCore/Private/DataInterfaces/OptimusDataInterfaceRawBuffer.cpp
@@ -195,48 +195,6 @@ UComputeDataProvider* UOptimusPersistentBufferDataInterface::CreateDataProvider(
 }
 
 
-bool UOptimusRawBufferDataProvider::IsValid() const
-{
-	if (!Component.IsValid())
-	{
-		return false;
-	}
-	
-	const UOptimusComponentSource* ComponentSourcePtr = ComponentSource.Get();
-	if (!ComponentSourcePtr)
-	{
-		return false;
-	}
-
-	switch(DataDomain.Type)
-	{
-	case EOptimusDataDomainType::Dimensional:
-		{
-			if (DataDomain.DimensionNames.IsEmpty())
-			{
-				return false;
-			}
-			return ComponentSourcePtr->GetExecutionDomains().Contains(DataDomain.DimensionNames[0]); 
-		}
-	case EOptimusDataDomainType::Expression:
-		{
-			using namespace Optimus::Expression;
-			TMap<FName, int32> DomainNames;
-
-			for (FName DomainName: ComponentSourcePtr->GetExecutionDomains())
-			{
-				DomainNames.Add(DomainName, 0);
-			}
-
-			// Return true if there was no error. Verify returns a TOptional.
-			return !FEngine{DomainNames}.Verify(DataDomain.Expression).IsSet();
-		}
-	}
-
-	return false;
-}
-
-
 bool UOptimusRawBufferDataProvider::GetLodAndInvocationElementCounts(
 	int32& OutLodIndex,
 	TArray<int32>& OutInvocationElementCounts
@@ -331,17 +289,6 @@ FComputeDataProviderRenderProxy* UOptimusTransientBufferDataProvider::GetRenderP
 }
 
 
-bool UOptimusPersistentBufferDataProvider::IsValid() const
-{
-	if (!BufferPool.IsValid())
-	{
-		return false;
-	}
-
-	return UOptimusRawBufferDataProvider::IsValid();
-}
-
-
 FComputeDataProviderRenderProxy* UOptimusPersistentBufferDataProvider::GetRenderProxy()
 {
 	int32 LodIndex = 0;
diff --git a/Engine/Plugins/Animation/DeformerGraph/Source/OptimusCore/Private/DataInterfaces/OptimusDataInterfaceRawBuffer.h b/Engine/Plugins/Animation/DeformerGraph/Source/OptimusCore/Private/DataInterfaces/OptimusDataInterfaceRawBuffer.h
--- a/Engine/Plugins/Animation/DeformerGraph/Source/OptimusCore/Private/DataInterfaces/OptimusDataInterfaceRawBuffer.h
+++ b/Engine/Plugins/Animation/DeformerGraph/Source/OptimusCore/Private/DataInterfaces/OptimusDataInterfaceRawBuffer.h
@@ -143,10 +143,6 @@ class OPTIMUSCORE_API UOptimusRawBufferDataProvider : public UComputeDataProvide
 	GENERATED_BODY()
 
 public:
-	//~ Begin UComputeDataProvider Interface
-	bool IsValid() const override;
-	//~ End UComputeDataProvider Interface
-
 	/** Helper function to calculate the element count for each section invocation of the given skinned/skeletal mesh
 	 *  object and a data domain. Uses the execution domains to compute the information. Returns false if the
 	 *  data domain is not valid for computation.
@@ -196,7 +192,6 @@ class OPTIMUSCORE_API UOptimusPersistentBufferDataProvider : public UOptimusRawB
 
 public:
 	//~ Begin UComputeDataProvider Interface
-	bool IsValid() const override;
 	FComputeDataProviderRenderProxy* GetRenderProxy() override;
 	//~ End UComputeDataProvider Interface
 
diff --git a/Engine/Plugins/Animation/DeformerGraph/Source/OptimusCore/Private/DataInterfaces/OptimusDataInterfaceSkeleton.cpp b/Engine/Plugins/Animation/DeformerGraph/Source/OptimusCore/Private/DataInterfaces/OptimusDataInterfaceSkeleton.cpp
--- a/Engine/Plugins/Animation/DeformerGraph/Source/OptimusCore/Private/DataInterfaces/OptimusDataInterfaceSkeleton.cpp
+++ b/Engine/Plugins/Animation/DeformerGraph/Source/OptimusCore/Private/DataInterfaces/OptimusDataInterfaceSkeleton.cpp
@@ -118,13 +118,6 @@ UComputeDataProvider* UOptimusSkeletonDataInterface::CreateDataProvider(TObjectP
 }
 
 
-bool UOptimusSkeletonDataProvider::IsValid() const
-{
-	return
-		SkinnedMesh != nullptr &&
-		SkinnedMesh->MeshObject != nullptr;
-}
-
 FComputeDataProviderRenderProxy* UOptimusSkeletonDataProvider::GetRenderProxy()
 {
 	return new FOptimusSkeletonDataProviderProxy(SkinnedMesh);
diff --git a/Engine/Plugins/Animation/DeformerGraph/Source/OptimusCore/Private/DataInterfaces/OptimusDataInterfaceSkeleton.h b/Engine/Plugins/Animation/DeformerGraph/Source/OptimusCore/Private/DataInterfaces/OptimusDataInterfaceSkeleton.h
--- a/Engine/Plugins/Animation/DeformerGraph/Source/OptimusCore/Private/DataInterfaces/OptimusDataInterfaceSkeleton.h
+++ b/Engine/Plugins/Animation/DeformerGraph/Source/OptimusCore/Private/DataInterfaces/OptimusDataInterfaceSkeleton.h
@@ -46,7 +46,6 @@ public:
 	TObjectPtr<USkinnedMeshComponent> SkinnedMesh = nullptr;
 
 	//~ Begin UComputeDataProvider Interface
-	bool IsValid() const override;
 	FComputeDataProviderRenderProxy* GetRenderProxy() override;
 	//~ End UComputeDataProvider Interface
 };
diff --git a/Engine/Plugins/Animation/DeformerGraph/Source/OptimusCore/Private/DataInterfaces/OptimusDataInterfaceSkinnedMesh.cpp b/Engine/Plugins/Animation/DeformerGraph/Source/OptimusCore/Private/DataInterfaces/OptimusDataInterfaceSkinnedMesh.cpp
--- a/Engine/Plugins/Animation/DeformerGraph/Source/OptimusCore/Private/DataInterfaces/OptimusDataInterfaceSkinnedMesh.cpp
+++ b/Engine/Plugins/Animation/DeformerGraph/Source/OptimusCore/Private/DataInterfaces/OptimusDataInterfaceSkinnedMesh.cpp
@@ -134,13 +134,6 @@ UComputeDataProvider* UOptimusSkinnedMeshDataInterface::CreateDataProvider(TObje
 }
 
 
-bool UOptimusSkinnedMeshDataProvider::IsValid() const
-{
-	return
-		SkinnedMesh != nullptr &&
-		SkinnedMesh->MeshObject != nullptr;
-}
-
 FComputeDataProviderRenderProxy* UOptimusSkinnedMeshDataProvider::GetRenderProxy()
 {
 	return new FOptimusSkinnedMeshDataProviderProxy(SkinnedMesh);
diff --git a/Engine/Plugins/Animation/DeformerGraph/Source/OptimusCore/Private/DataInterfaces/OptimusDataInterfaceSkinnedMesh.h b/Engine/Plugins/Animation/DeformerGraph/Source/OptimusCore/Private/DataInterfaces/OptimusDataInterfaceSkinnedMesh.h
--- a/Engine/Plugins/Animation/DeformerGraph/Source/OptimusCore/Private/DataInterfaces/OptimusDataInterfaceSkinnedMesh.h
+++ b/Engine/Plugins/Animation/DeformerGraph/Source/OptimusCore/Private/DataInterfaces/OptimusDataInterfaceSkinnedMesh.h
@@ -45,7 +45,6 @@ public:
 	TObjectPtr<USkinnedMeshComponent> SkinnedMesh = nullptr;
 
 	//~ Begin UComputeDataProvider Interface
-	bool IsValid() const override;
 	FComputeDataProviderRenderProxy* GetRenderProxy() override;
 	//~ End UComputeDataProvider Interface
 };
diff --git a/Engine/Plugins/Animation/DeformerGraph/Source/OptimusCore/Private/DataInterfaces/OptimusDataInterfaceSkinnedMeshExec.cpp b/Engine/Plugins/Animation/DeformerGraph/Source/OptimusCore/Private/DataInterfaces/OptimusDataInterfaceSkinnedMeshExec.cpp
--- a/Engine/Plugins/Animation/DeformerGraph/Source/OptimusCore/Private/DataInterfaces/OptimusDataInterfaceSkinnedMeshExec.cpp
+++ b/Engine/Plugins/Animation/DeformerGraph/Source/OptimusCore/Private/DataInterfaces/OptimusDataInterfaceSkinnedMeshExec.cpp
@@ -79,13 +79,6 @@ UComputeDataProvider* UOptimusSkinnedMeshExecDataInterface::CreateDataProvider(T
 }
 
 
-bool UOptimusSkinnedMeshExecDataProvider::IsValid() const
-{
-	return
-		SkinnedMesh != nullptr &&
-		SkinnedMesh->MeshObject != nullptr;
-}
-
 FComputeDataProviderRenderProxy* UOptimusSkinnedMeshExecDataProvider::GetRenderProxy()
 {
 	return new FOptimusSkinnedMeshExecDataProviderProxy(SkinnedMesh, Domain);
diff --git a/Engine/Plugins/Animation/DeformerGraph/Source/OptimusCore/Private/DataInterfaces/OptimusDataInterfaceSkinnedMeshExec.h b/Engine/Plugins/Animation/DeformerGraph/Source/OptimusCore/Private/DataInterfaces/OptimusDataInterfaceSkinnedMeshExec.h
--- a/Engine/Plugins/Animation/DeformerGraph/Source/OptimusCore/Private/DataInterfaces/OptimusDataInterfaceSkinnedMeshExec.h
+++ b/Engine/Plugins/Animation/DeformerGraph/Source/OptimusCore/Private/DataInterfaces/OptimusDataInterfaceSkinnedMeshExec.h
@@ -65,7 +65,6 @@ public:
 	EOptimusSkinnedMeshExecDomain Domain = EOptimusSkinnedMeshExecDomain::Vertex;
 
 	//~ Begin UComputeDataProvider Interface
-	bool IsValid() const override;
 	FComputeDataProviderRenderProxy* GetRenderProxy() override;
 	//~ End UComputeDataProvider Interface
 };
diff --git a/Engine/Plugins/Animation/DeformerGraph/Source/OptimusCore/Private/DataInterfaces/OptimusDataInterfaceSkinnedMeshWrite.cpp b/Engine/Plugins/Animation/DeformerGraph/Source/OptimusCore/Private/DataInterfaces/OptimusDataInterfaceSkinnedMeshWrite.cpp
--- a/Engine/Plugins/Animation/DeformerGraph/Source/OptimusCore/Private/DataInterfaces/OptimusDataInterfaceSkinnedMeshWrite.cpp
+++ b/Engine/Plugins/Animation/DeformerGraph/Source/OptimusCore/Private/DataInterfaces/OptimusDataInterfaceSkinnedMeshWrite.cpp
@@ -113,13 +113,6 @@ UComputeDataProvider* UOptimusSkinnedMeshWriteDataInterface::CreateDataProvider(
 }
 
 
-bool UOptimusSkinnedMeshWriteDataProvider::IsValid() const
-{
-	return
-		SkinnedMesh != nullptr &&
-		SkinnedMesh->MeshObject != nullptr;
-}
-
 FComputeDataProviderRenderProxy* UOptimusSkinnedMeshWriteDataProvider::GetRenderProxy()
 {
 	return new FOptimusSkinnedMeshWriteDataProviderProxy(SkinnedMesh, OutputMask);
diff --git a/Engine/Plugins/Animation/DeformerGraph/Source/OptimusCore/Private/DataInterfaces/OptimusDataInterfaceSkinnedMeshWrite.h b/Engine/Plugins/Animation/DeformerGraph/Source/OptimusCore/Private/DataInterfaces/OptimusDataInterfaceSkinnedMeshWrite.h
--- a/Engine/Plugins/Animation/DeformerGraph/Source/OptimusCore/Private/DataInterfaces/OptimusDataInterfaceSkinnedMeshWrite.h
+++ b/Engine/Plugins/Animation/DeformerGraph/Source/OptimusCore/Private/DataInterfaces/OptimusDataInterfaceSkinnedMeshWrite.h
@@ -51,7 +51,6 @@ public:
 	uint64 OutputMask;
 
 	//~ Begin UComputeDataProvider Interface
-	bool IsValid() const override;
 	FComputeDataProviderRenderProxy* GetRenderProxy() override;
 	//~ End UComputeDataProvider Interface
 };
diff --git a/Engine/Plugins/Animation/MLDeformer/MLDeformerFramework/Source/MLDeformerFramework/Private/MLDeformerGraphDataInterface.cpp b/Engine/Plugins/Animation/MLDeformer/MLDeformerFramework/Source/MLDeformerFramework/Private/MLDeformerGraphDataInterface.cpp
--- a/Engine/Plugins/Animation/MLDeformer/MLDeformerFramework/Source/MLDeformerFramework/Private/MLDeformerGraphDataInterface.cpp
+++ b/Engine/Plugins/Animation/MLDeformer/MLDeformerFramework/Source/MLDeformerFramework/Private/MLDeformerGraphDataInterface.cpp
@@ -81,16 +81,6 @@ FComputeDataProviderRenderProxy* UMLDeformerGraphDataProvider::GetRenderProxy()
 	return new UE::MLDeformer::FMLDeformerGraphDataProviderProxy(DeformerComponent);
 }
 
-bool UMLDeformerGraphDataProvider::IsValid() const
-{
-	if (DeformerComponent == nullptr || DeformerComponent->GetDeformerAsset() == nullptr || DeformerComponent->GetModelInstance() == nullptr)
-	{
-		return false;
-	}
-
-	return DeformerComponent->GetModelInstance()->IsValidForDataProvider();
-}
-
 namespace UE::MLDeformer
 {
 	FMLDeformerGraphDataProviderProxy::FMLDeformerGraphDataProviderProxy(UMLDeformerComponent* DeformerComponent)
diff --git a/Engine/Plugins/Animation/MLDeformer/MLDeformerFramework/Source/MLDeformerFramework/Private/MLDeformerGraphDebugDataInterface.cpp b/Engine/Plugins/Animation/MLDeformer/MLDeformerFramework/Source/MLDeformerFramework/Private/MLDeformerGraphDebugDataInterface.cpp
--- a/Engine/Plugins/Animation/MLDeformer/MLDeformerFramework/Source/MLDeformerFramework/Private/MLDeformerGraphDebugDataInterface.cpp
+++ b/Engine/Plugins/Animation/MLDeformer/MLDeformerFramework/Source/MLDeformerFramework/Private/MLDeformerGraphDebugDataInterface.cpp
@@ -103,24 +103,10 @@ UComputeDataProvider* UMLDeformerGraphDebugDataInterface::CreateDataProvider(TOb
 	return Provider;
 }
 
-bool UMLDeformerGraphDebugDataProvider::IsValid() const
-{
-#if WITH_EDITORONLY_DATA
-	if (DeformerComponent == nullptr || DeformerComponent->GetDeformerAsset() == nullptr || DeformerComponent->GetModelInstance() == nullptr)
-	{
-		return false;
-	}
-
-	return DeformerComponent->GetModelInstance()->IsValidForDataProvider();
-#else
-	return false; // This data interface is only valid in editor.
-#endif
-}
-
 FComputeDataProviderRenderProxy* UMLDeformerGraphDebugDataProvider::GetRenderProxy()
 {
 #if WITH_EDITORONLY_DATA
-	if (DeformerComponent && DeformerAsset)
+	if (DeformerComponent && DeformerAsset && DeformerComponent->GetModelInstance() && DeformerComponent->GetModelInstance()->IsValidForDataProvider())
 	{
 		UE::MLDeformer::FMLDeformerGraphDebugDataProviderProxy* Proxy = new UE::MLDeformer::FMLDeformerGraphDebugDataProviderProxy(DeformerComponent, DeformerAsset, this);
 		const float SampleTime = DeformerComponent->GetModelInstance()->GetSkeletalMeshComponent()->GetPosition();
diff --git a/Engine/Plugins/Animation/MLDeformer/MLDeformerFramework/Source/MLDeformerFramework/Public/MLDeformerGraphDataInterface.h b/Engine/Plugins/Animation/MLDeformer/MLDeformerFramework/Source/MLDeformerFramework/Public/MLDeformerGraphDataInterface.h
--- a/Engine/Plugins/Animation/MLDeformer/MLDeformerFramework/Source/MLDeformerFramework/Public/MLDeformerGraphDataInterface.h
+++ b/Engine/Plugins/Animation/MLDeformer/MLDeformerFramework/Source/MLDeformerFramework/Public/MLDeformerGraphDataInterface.h
@@ -87,7 +87,6 @@ public:
 	TObjectPtr<UMLDeformerComponent> DeformerComponent = nullptr;
 
 	// UComputeDataProvider overrides.
-	virtual bool IsValid() const override;
 	virtual FComputeDataProviderRenderProxy* GetRenderProxy() override;
 	// ~END UComputeDataProvider overrides.
 };
diff --git a/Engine/Plugins/Animation/MLDeformer/MLDeformerFramework/Source/MLDeformerFramework/Public/MLDeformerGraphDebugDataInterface.h b/Engine/Plugins/Animation/MLDeformer/MLDeformerFramework/Source/MLDeformerFramework/Public/MLDeformerGraphDebugDataInterface.h
--- a/Engine/Plugins/Animation/MLDeformer/MLDeformerFramework/Source/MLDeformerFramework/Public/MLDeformerGraphDebugDataInterface.h
+++ b/Engine/Plugins/Animation/MLDeformer/MLDeformerFramework/Source/MLDeformerFramework/Public/MLDeformerGraphDebugDataInterface.h
@@ -73,7 +73,6 @@ public:
 	virtual void Init() {};
 
 	// UComputeDataProvider overrides.
-	virtual bool IsValid() const override;
 	virtual FComputeDataProviderRenderProxy* GetRenderProxy() override;
 	// ~END UComputeDataProvider overrides.
 
diff --git a/Engine/Plugins/Runtime/ComputeFramework/Source/ComputeFramework/Public/ComputeFramework/ComputeDataProvider.h b/Engine/Plugins/Runtime/ComputeFramework/Source/ComputeFramework/Public/ComputeFramework/ComputeDataProvider.h
--- a/Engine/Plugins/Runtime/ComputeFramework/Source/ComputeFramework/Public/ComputeFramework/ComputeDataProvider.h
+++ b/Engine/Plugins/Runtime/ComputeFramework/Source/ComputeFramework/Public/ComputeFramework/ComputeDataProvider.h
@@ -22,12 +22,12 @@ class COMPUTEFRAMEWORK_API UComputeDataProvider : public UObject
 
 public:
 	/** Return false if the provider has not been fully initialized. */
+	UE_DEPRECATED(5.2, "Implement any validation in FComputeDataProviderRenderProxy::IsValid().")
 	virtual bool IsValid() const { return true; }
 	
 	/**
 	 * Get an associated render thread proxy object.
 	 * Currently these are created and destroyed per frame by the FComputeGraphInstance.
-	 * todo[CF]: Don't destroy FComputeDataProviderRenderProxy objects every frame?
 	 */
 	virtual FComputeDataProviderRenderProxy* GetRenderProxy() { return nullptr; }
 };
diff --git a/Engine/Plugins/Runtime/HairStrands/Source/HairStrandsCore/Private/Deformer/DeformerDataInterfaceGroom.cpp b/Engine/Plugins/Runtime/HairStrands/Source/HairStrandsCore/Private/Deformer/DeformerDataInterfaceGroom.cpp
--- a/Engine/Plugins/Runtime/HairStrands/Source/HairStrandsCore/Private/Deformer/DeformerDataInterfaceGroom.cpp
+++ b/Engine/Plugins/Runtime/HairStrands/Source/HairStrandsCore/Private/Deformer/DeformerDataInterfaceGroom.cpp
@@ -133,16 +133,13 @@ UComputeDataProvider* UOptimusGroomDataInterface::CreateDataProvider(TObjectPtr<
 	return Provider;
 }
 
-bool UOptimusGroomDataProvider::IsValid() const
-{
-	return Groom != nullptr;
-}
 
 FComputeDataProviderRenderProxy* UOptimusGroomDataProvider::GetRenderProxy()
 {
 	return new FOptimusGroomDataProviderProxy(Groom);
 }
 
+
 FOptimusGroomDataProviderProxy::FOptimusGroomDataProviderProxy(UGroomComponent* InGroomComponent)
 {
 	GroomComponent = InGroomComponent;
diff --git a/Engine/Plugins/Runtime/HairStrands/Source/HairStrandsCore/Private/Deformer/DeformerDataInterfaceGroom.h b/Engine/Plugins/Runtime/HairStrands/Source/HairStrandsCore/Private/Deformer/DeformerDataInterfaceGroom.h
--- a/Engine/Plugins/Runtime/HairStrands/Source/HairStrandsCore/Private/Deformer/DeformerDataInterfaceGroom.h
+++ b/Engine/Plugins/Runtime/HairStrands/Source/HairStrandsCore/Private/Deformer/DeformerDataInterfaceGroom.h
@@ -45,7 +45,6 @@ public:
 	TObjectPtr<UGroomComponent> Groom = nullptr;
 
 	//~ Begin UComputeDataProvider Interface
-	bool IsValid() const override;
 	FComputeDataProviderRenderProxy* GetRenderProxy() override;
 	//~ End UComputeDataProvider Interface
 };
diff --git a/Engine/Plugins/Runtime/HairStrands/Source/HairStrandsCore/Private/Deformer/DeformerDataInterfaceGroomExec.cpp b/Engine/Plugins/Runtime/HairStrands/Source/HairStrandsCore/Private/Deformer/DeformerDataInterfaceGroomExec.cpp
--- a/Engine/Plugins/Runtime/HairStrands/Source/HairStrandsCore/Private/Deformer/DeformerDataInterfaceGroomExec.cpp
+++ b/Engine/Plugins/Runtime/HairStrands/Source/HairStrandsCore/Private/Deformer/DeformerDataInterfaceGroomExec.cpp
@@ -74,16 +74,13 @@ UComputeDataProvider* UOptimusGroomExecDataInterface::CreateDataProvider(TObject
 	return Provider;
 }
 
-bool UOptimusGroomExecDataProvider::IsValid() const
-{
-	return GroomComponent != nullptr && GroomComponent->GetGroupCount() > 0;
-}
 
 FComputeDataProviderRenderProxy* UOptimusGroomExecDataProvider::GetRenderProxy()
 {
 	return new FOptimusGroomExecDataProviderProxy(GroomComponent, Domain);
 }
 
+
 FOptimusGroomExecDataProviderProxy::FOptimusGroomExecDataProviderProxy(UGroomComponent* InGroomComponent, EOptimusGroomExecDomain InDomain)
 {
 	GroomComponent = InGroomComponent;
diff --git a/Engine/Plugins/Runtime/HairStrands/Source/HairStrandsCore/Private/Deformer/DeformerDataInterfaceGroomExec.h b/Engine/Plugins/Runtime/HairStrands/Source/HairStrandsCore/Private/Deformer/DeformerDataInterfaceGroomExec.h
--- a/Engine/Plugins/Runtime/HairStrands/Source/HairStrandsCore/Private/Deformer/DeformerDataInterfaceGroomExec.h
+++ b/Engine/Plugins/Runtime/HairStrands/Source/HairStrandsCore/Private/Deformer/DeformerDataInterfaceGroomExec.h
@@ -61,7 +61,6 @@ public:
 	EOptimusGroomExecDomain Domain = EOptimusGroomExecDomain::ControlPoint;
 
 	//~ Begin UComputeDataProvider Interface
-	bool IsValid() const override;
 	FComputeDataProviderRenderProxy* GetRenderProxy() override;
 	//~ End UComputeDataProvider Interface
 };
diff --git a/Engine/Plugins/Runtime/HairStrands/Source/HairStrandsCore/Private/Deformer/DeformerDataInterfaceGroomWrite.cpp b/Engine/Plugins/Runtime/HairStrands/Source/HairStrandsCore/Private/Deformer/DeformerDataInterfaceGroomWrite.cpp
--- a/Engine/Plugins/Runtime/HairStrands/Source/HairStrandsCore/Private/Deformer/DeformerDataInterfaceGroomWrite.cpp
+++ b/Engine/Plugins/Runtime/HairStrands/Source/HairStrandsCore/Private/Deformer/DeformerDataInterfaceGroomWrite.cpp
@@ -91,11 +91,6 @@ UComputeDataProvider* UOptimusGroomWriteDataInterface::CreateDataProvider(TObjec
 }
 
 
-bool UOptimusGroomWriteDataProvider::IsValid() const
-{
-	return GroomComponent != nullptr && GroomComponent->GetGroupCount() > 0;
-}
-
 FComputeDataProviderRenderProxy* UOptimusGroomWriteDataProvider::GetRenderProxy()
 {
 	return new FOptimusGroomWriteDataProviderProxy(GroomComponent, OutputMask);
diff --git a/Engine/Plugins/Runtime/HairStrands/Source/HairStrandsCore/Private/Deformer/DeformerDataInterfaceGroomWrite.h b/Engine/Plugins/Runtime/HairStrands/Source/HairStrandsCore/Private/Deformer/DeformerDataInterfaceGroomWrite.h
--- a/Engine/Plugins/Runtime/HairStrands/Source/HairStrandsCore/Private/Deformer/DeformerDataInterfaceGroomWrite.h
+++ b/Engine/Plugins/Runtime/HairStrands/Source/HairStrandsCore/Private/Deformer/DeformerDataInterfaceGroomWrite.h
@@ -51,7 +51,6 @@ public:
 	uint64 OutputMask;
 
 	//~ Begin UComputeDataProvider Interface
-	bool IsValid() const override;
 	FComputeDataProviderRenderProxy* GetRenderProxy() override;
 	//~ End UComputeDataProvider Interface
 };
