Optimus: Fix intermittant crash when opening defomer graph editor.

diff --git a/Engine/Plugins/Animation/DeformerGraph/Source/OptimusCore/Private/ComponentSources/OptimusSkinnedMeshComponentSource.cpp b/Engine/Plugins/Animation/DeformerGraph/Source/OptimusCore/Private/ComponentSources/OptimusSkinnedMeshComponentSource.cpp
--- a/Engine/Plugins/Animation/DeformerGraph/Source/OptimusCore/Private/ComponentSources/OptimusSkinnedMeshComponentSource.cpp
+++ b/Engine/Plugins/Animation/DeformerGraph/Source/OptimusCore/Private/ComponentSources/OptimusSkinnedMeshComponentSource.cpp
@@ -55,6 +55,11 @@ bool UOptimusSkinnedMeshComponentSource::GetComponentElementCountsForExecutionDo
 	}
 
 	const FSkeletalMeshObject* SkeletalMeshObject = SkinnedMeshComponent->MeshObject;
+	if (!SkeletalMeshObject)
+	{
+		return false;
+	}
+
 	FSkeletalMeshRenderData const& SkeletalMeshRenderData = SkeletalMeshObject->GetSkeletalMeshRenderData();
 	FSkeletalMeshLODRenderData const* LodRenderData = &SkeletalMeshRenderData.LODRenderData[InLodIndex];
 
