Optimus: Fix Buffer definition in Connectivity DI.
Was causing Vulkan crash.

diff --git a/Engine/Plugins/Animation/DeformerGraph/Shaders/Private/DataInterfaceConnectivity.ush b/Engine/Plugins/Animation/DeformerGraph/Shaders/Private/DataInterfaceConnectivity.ush
--- a/Engine/Plugins/Animation/DeformerGraph/Shaders/Private/DataInterfaceConnectivity.ush
+++ b/Engine/Plugins/Animation/DeformerGraph/Shaders/Private/DataInterfaceConnectivity.ush
@@ -3,7 +3,7 @@
 uint {DataInterfaceName}_NumVertices;
 uint {DataInterfaceName}_InputStreamStart;
 uint {DataInterfaceName}_MaxConnectedVertexCount;
-Buffer<uint> {DataInterfaceName}_ConnectivityBuffer;
+StructuredBuffer<uint> {DataInterfaceName}_ConnectivityBuffer;
 
 uint ReadNumVertices_{DataInterfaceName}()
 {
