Remove unnecessary check() from Deformer Graph DI.
The check() can now fail since we moved validation later in the compute framework processing.

diff --git a/Engine/Plugins/Animation/DeformerGraph/Source/OptimusCore/Private/DataInterfaces/OptimusDataInterfaceGraph.cpp b/Engine/Plugins/Animation/DeformerGraph/Source/OptimusCore/Private/DataInterfaces/OptimusDataInterfaceGraph.cpp
--- a/Engine/Plugins/Animation/DeformerGraph/Source/OptimusCore/Private/DataInterfaces/OptimusDataInterfaceGraph.cpp
+++ b/Engine/Plugins/Animation/DeformerGraph/Source/OptimusCore/Private/DataInterfaces/OptimusDataInterfaceGraph.cpp
@@ -94,7 +94,6 @@ UComputeDataProvider* UOptimusGraphDataInterface::CreateDataProvider(TObjectPtr<
 	Provider->MeshComponent = Cast<UMeshComponent>(InBinding);
 	Provider->Variables = Variables;
 	Provider->ParameterBufferSize = ParameterBufferSize;
-	check(Provider->MeshComponent);
 	return Provider;
 }
 
