Use correct interface to get text in deformer kernel customization.
GetShaderText() returns the source for editing (what we want) as opposed to GetSource() which returns the source for compilation (which may have extra text inserted).

diff --git a/Engine/Plugins/Animation/DeformerGraph/Source/OptimusEditor/Private/OptimusDetailsCustomization.cpp b/Engine/Plugins/Animation/DeformerGraph/Source/OptimusEditor/Private/OptimusDetailsCustomization.cpp
--- a/Engine/Plugins/Animation/DeformerGraph/Source/OptimusEditor/Private/OptimusDetailsCustomization.cpp
+++ b/Engine/Plugins/Animation/DeformerGraph/Source/OptimusEditor/Private/OptimusDetailsCustomization.cpp
@@ -1387,7 +1387,7 @@ void FOptimusSourceDetailsCustomization::CustomizeDetails(IDetailLayoutBuilder&
 
 FText FOptimusSourceDetailsCustomization::GetText() const
 {
-	return FText::FromString(OptimusSource->GetSource());
+	return FText::FromString(OptimusSource->GetShaderText());
 }
 
 void FOptimusSourceDetailsCustomization::OnTextChanged(const FText& InValue)
