DeformerGraph: Flush shader file cache before compilation.
This ends the confusing workflow of needing to do r.ComputeFramework.RebuildComputeGraphs after changing on disk shader files.
It may make some overhead for shader compilation, but it seems minor.

diff --git a/Engine/Plugins/Animation/DeformerGraph/Source/OptimusCore/Private/OptimusDeformer.cpp b/Engine/Plugins/Animation/DeformerGraph/Source/OptimusCore/Private/OptimusDeformer.cpp
--- a/Engine/Plugins/Animation/DeformerGraph/Source/OptimusCore/Private/OptimusDeformer.cpp
+++ b/Engine/Plugins/Animation/DeformerGraph/Source/OptimusCore/Private/OptimusDeformer.cpp
@@ -32,6 +32,7 @@
 #include "OptimusVariableDescription.h"
 #include "RenderingThread.h"
 #include "SceneInterface.h"
+#include "ShaderCore.h"
 #include "UObject/Package.h"
 
 // FIXME: We should not be accessing nodes directly.
@@ -1370,6 +1371,12 @@ bool UOptimusDeformer::Compile()
 		return false;
 	}
 	
+#if WITH_EDITOR
+	// Flush the shader file cache in case we are editing engine or data interface shaders.
+	// We could make the user do this manually, but that makes iterating on data interfaces really painful.
+	FlushShaderFileCache();
+#endif
+
 	for (const FOptimusComputeGraphInfo& ComputeGraphInfo: ComputeGraphs)
 	{
 		ComputeGraphInfo.ComputeGraph->UpdateResources();
