Optimus: Fix UObject generated shader source paths to fix some path validation warnings.
Change OptimusSource files to set virtual path through virtual file name instead of through (less robust) line directives.

diff --git a/Engine/Plugins/Animation/DeformerGraph/Source/OptimusCore/Private/OptimusHelpers.cpp b/Engine/Plugins/Animation/DeformerGraph/Source/OptimusCore/Private/OptimusHelpers.cpp
--- a/Engine/Plugins/Animation/DeformerGraph/Source/OptimusCore/Private/OptimusHelpers.cpp
+++ b/Engine/Plugins/Animation/DeformerGraph/Source/OptimusCore/Private/OptimusHelpers.cpp
@@ -112,20 +112,21 @@ FName Optimus::GetTypeName(UScriptStruct* InStructType, bool bInShouldGetUniqueN
 
 void Optimus::ConvertObjectPathToShaderFilePath(FString& InOutPath)
 {
-	// Shader compiler recognizes "/UObject/..." path as special. 
+	// Shader compiler recognizes "/Engine/Generated/..." path as special. 
 	// It doesn't validate file suffix etc.
-	InOutPath = TEXT("/UObject") + InOutPath;
+	InOutPath = FString::Printf(TEXT("/Engine/Generated/UObject%s.ush"), *InOutPath);
 	// Shader compilation result parsing will break if it finds ':' where it doesn't expect.
 	InOutPath.ReplaceCharInline(TEXT(':'), TEXT('@'));
 }
 
 bool Optimus::ConvertShaderFilePathToObjectPath(FString& InOutPath)
 {
-	if (!InOutPath.RemoveFromStart(TEXT("/UObject")))
+	if (!InOutPath.RemoveFromStart(TEXT("/Engine/Generated/UObject")))
 	{
 		return false;
 	}
 
 	InOutPath.ReplaceCharInline(TEXT('@'), TEXT(':'));
+	InOutPath.RemoveFromEnd(TEXT(".ush"));
 	return true;
 }
diff --git a/Engine/Plugins/Animation/DeformerGraph/Source/OptimusCore/Private/OptimusSource.cpp b/Engine/Plugins/Animation/DeformerGraph/Source/OptimusCore/Private/OptimusSource.cpp
--- a/Engine/Plugins/Animation/DeformerGraph/Source/OptimusCore/Private/OptimusSource.cpp
+++ b/Engine/Plugins/Animation/DeformerGraph/Source/OptimusCore/Private/OptimusSource.cpp
@@ -12,12 +12,11 @@ void UOptimusSource::SetSource(const FString& InText)
 	Modify();
 }
 
-
-FString UOptimusSource::GetSource() const
+FString UOptimusSource::GetVirtualPath() const 
 {
 	FString ShaderPathName = GetPathName();
 	Optimus::ConvertObjectPathToShaderFilePath(ShaderPathName);
-	return FString::Printf(TEXT("#line 1 \"%s\"\n%s"), *ShaderPathName, *SourceText);
+	return ShaderPathName;
 }
 
 #if WITH_EDITOR	
diff --git a/Engine/Plugins/Animation/DeformerGraph/Source/OptimusCore/Public/OptimusSource.h b/Engine/Plugins/Animation/DeformerGraph/Source/OptimusCore/Public/OptimusSource.h
--- a/Engine/Plugins/Animation/DeformerGraph/Source/OptimusCore/Public/OptimusSource.h
+++ b/Engine/Plugins/Animation/DeformerGraph/Source/OptimusCore/Public/OptimusSource.h
@@ -17,7 +17,8 @@ public:
 	void SetSource(const FString& InText);
 
 	// Begin UComputeSource interface.
-	FString GetSource() const override;
+	FString GetSource() const override { return SourceText; }
+	FString GetVirtualPath() const override;
 	// End UComputeSource interface.
 
 	// Begin IOptimusShaderTextProvider interface.
diff --git a/Engine/Plugins/Runtime/ComputeFramework/Source/ComputeFramework/Public/ComputeFramework/ComputeSource.h b/Engine/Plugins/Runtime/ComputeFramework/Source/ComputeFramework/Public/ComputeFramework/ComputeSource.h
--- a/Engine/Plugins/Runtime/ComputeFramework/Source/ComputeFramework/Public/ComputeFramework/ComputeSource.h
+++ b/Engine/Plugins/Runtime/ComputeFramework/Source/ComputeFramework/Public/ComputeFramework/ComputeSource.h
@@ -18,7 +18,7 @@ public:
 	/** Get source code ready for HLSL compilation. */
 	virtual FString GetSource() const { return FString(); }
 	/** Get virtual file path for the source code. This will be the file name shown in any compilation errors. */
-	virtual FString GetVirtualPath() const { return FString::Printf(TEXT("/Plugin/ComputeFramework/Generated%s.ush"), *GetPathName()); }
+	virtual FString GetVirtualPath() const { return FString::Printf(TEXT("/Engine/Generated/%s.ush"), *GetPathName()); }
 
 public:
 	/** Array of additional source objects. This allows us to specify source dependencies. */
