Optimus: Implement GetShaderVirtualPath() for all data interfaces that can support it.

diff --git a/Engine/Plugins/Animation/DeformerGraph/Source/OptimusCore/Private/DataInterfaces/OptimusDataInterfaceCloth.cpp b/Engine/Plugins/Animation/DeformerGraph/Source/OptimusCore/Private/DataInterfaces/OptimusDataInterfaceCloth.cpp
--- a/Engine/Plugins/Animation/DeformerGraph/Source/OptimusCore/Private/DataInterfaces/OptimusDataInterfaceCloth.cpp
+++ b/Engine/Plugins/Animation/DeformerGraph/Source/OptimusCore/Private/DataInterfaces/OptimusDataInterfaceCloth.cpp
@@ -90,9 +90,16 @@ void UOptimusClothDataInterface::GetPermutations(FComputeKernelPermutationVector
 	OutPermutationVector.AddPermutation(TEXT("ENABLE_DEFORMER_CLOTH"), 2);
 }
 
+TCHAR const* UOptimusClothDataInterface::TemplateFilePath = TEXT("/Plugin/Optimus/Private/DataInterfaceCloth.ush");
+
+TCHAR const* UOptimusClothDataInterface::GetShaderVirtualPath() const
+{
+	return TemplateFilePath;
+}
+
 void UOptimusClothDataInterface::GetShaderHash(FString& InOutKey) const
 {
-	GetShaderFileHash(TEXT("/Plugin/Optimus/Private/DataInterfaceCloth.ush"), EShaderPlatform::SP_PCD3D_SM5).AppendString(InOutKey);
+	GetShaderFileHash(TemplateFilePath, EShaderPlatform::SP_PCD3D_SM5).AppendString(InOutKey);
 }
 
 void UOptimusClothDataInterface::GetHLSL(FString& OutHLSL, FString const& InDataInterfaceName) const
@@ -103,7 +110,7 @@ void UOptimusClothDataInterface::GetHLSL(FString& OutHLSL, FString const& InData
 	};
 
 	FString TemplateFile;
-	LoadShaderSourceFile(TEXT("/Plugin/Optimus/Private/DataInterfaceCloth.ush"), EShaderPlatform::SP_PCD3D_SM5, &TemplateFile, nullptr);
+	LoadShaderSourceFile(TemplateFilePath, EShaderPlatform::SP_PCD3D_SM5, &TemplateFile, nullptr);
 	OutHLSL += FString::Format(*TemplateFile, TemplateArgs);
 }
 
diff --git a/Engine/Plugins/Animation/DeformerGraph/Source/OptimusCore/Private/DataInterfaces/OptimusDataInterfaceCloth.h b/Engine/Plugins/Animation/DeformerGraph/Source/OptimusCore/Private/DataInterfaces/OptimusDataInterfaceCloth.h
--- a/Engine/Plugins/Animation/DeformerGraph/Source/OptimusCore/Private/DataInterfaces/OptimusDataInterfaceCloth.h
+++ b/Engine/Plugins/Animation/DeformerGraph/Source/OptimusCore/Private/DataInterfaces/OptimusDataInterfaceCloth.h
@@ -28,10 +28,14 @@ public:
 	void GetSupportedInputs(TArray<FShaderFunctionDefinition>& OutFunctions) const override;
 	void GetShaderParameters(TCHAR const* UID, FShaderParametersMetadataBuilder& InOutBuilder, FShaderParametersMetadataAllocations& InOutAllocations) const override;
 	void GetPermutations(FComputeKernelPermutationVector& OutPermutationVector) const override;
+	TCHAR const* GetShaderVirtualPath() const override;
 	void GetShaderHash(FString& InOutKey) const override;
 	void GetHLSL(FString& OutHLSL, FString const& InDataInterfaceName) const override;
 	UComputeDataProvider* CreateDataProvider(TObjectPtr<UObject> InBinding, uint64 InInputMask, uint64 InOutputMask) const override;
 	//~ End UComputeDataInterface Interface
+
+private:
+	static TCHAR const* TemplateFilePath;
 };
 
 /** Compute Framework Data Provider for reading skeletal mesh. */
diff --git a/Engine/Plugins/Animation/DeformerGraph/Source/OptimusCore/Private/DataInterfaces/OptimusDataInterfaceConnectivity.cpp b/Engine/Plugins/Animation/DeformerGraph/Source/OptimusCore/Private/DataInterfaces/OptimusDataInterfaceConnectivity.cpp
--- a/Engine/Plugins/Animation/DeformerGraph/Source/OptimusCore/Private/DataInterfaces/OptimusDataInterfaceConnectivity.cpp
+++ b/Engine/Plugins/Animation/DeformerGraph/Source/OptimusCore/Private/DataInterfaces/OptimusDataInterfaceConnectivity.cpp
@@ -178,9 +178,16 @@ void UOptimusConnectivityDataInterface::GetShaderParameters(TCHAR const* UID, FS
 	InOutBuilder.AddNestedStruct<FConnectivityDataInterfaceParameters>(UID);
 }
 
+TCHAR const* UOptimusConnectivityDataInterface::TemplateFilePath = TEXT("/Plugin/Optimus/Private/DataInterfaceConnectivity.ush");
+
+TCHAR const* UOptimusConnectivityDataInterface::GetShaderVirtualPath() const
+{
+	return TemplateFilePath;
+}
+
 void UOptimusConnectivityDataInterface::GetShaderHash(FString& InOutKey) const
 {
-	GetShaderFileHash(TEXT("/Plugin/Optimus/Private/DataInterfaceConnectivity.ush"), EShaderPlatform::SP_PCD3D_SM5).AppendString(InOutKey);
+	GetShaderFileHash(TemplateFilePath, EShaderPlatform::SP_PCD3D_SM5).AppendString(InOutKey);
 }
 
 void UOptimusConnectivityDataInterface::GetHLSL(FString& OutHLSL, FString const& InDataInterfaceName) const
@@ -191,7 +198,7 @@ void UOptimusConnectivityDataInterface::GetHLSL(FString& OutHLSL, FString const&
 	};
 
 	FString TemplateFile;
-	LoadShaderSourceFile(TEXT("/Plugin/Optimus/Private/DataInterfaceConnectivity.ush"), EShaderPlatform::SP_PCD3D_SM5, &TemplateFile, nullptr);
+	LoadShaderSourceFile(TemplateFilePath, EShaderPlatform::SP_PCD3D_SM5, &TemplateFile, nullptr);
 	OutHLSL += FString::Format(*TemplateFile, TemplateArgs);
 }
 
diff --git a/Engine/Plugins/Animation/DeformerGraph/Source/OptimusCore/Private/DataInterfaces/OptimusDataInterfaceConnectivity.h b/Engine/Plugins/Animation/DeformerGraph/Source/OptimusCore/Private/DataInterfaces/OptimusDataInterfaceConnectivity.h
--- a/Engine/Plugins/Animation/DeformerGraph/Source/OptimusCore/Private/DataInterfaces/OptimusDataInterfaceConnectivity.h
+++ b/Engine/Plugins/Animation/DeformerGraph/Source/OptimusCore/Private/DataInterfaces/OptimusDataInterfaceConnectivity.h
@@ -30,6 +30,7 @@ public:
 	TCHAR const* GetClassName() const override { return TEXT("Connectivity"); }
 	void GetSupportedInputs(TArray<FShaderFunctionDefinition>& OutFunctions) const override;
 	void GetShaderParameters(TCHAR const* UID, FShaderParametersMetadataBuilder& InOutBuilder, FShaderParametersMetadataAllocations& InOutAllocations) const override;
+	TCHAR const* GetShaderVirtualPath() const override;
 	void GetShaderHash(FString& InOutKey) const override;
 	void GetHLSL(FString& OutHLSL, FString const& InDataInterfaceName) const override;
 	UComputeDataProvider* CreateDataProvider(TObjectPtr<UObject> InBinding, uint64 InInputMask, uint64 InOutputMask) const override;
@@ -37,6 +38,9 @@ public:
 
 	// Hardcoded max connected count.
 	static const int32 MaxConnectedVertexCount = 12;
+
+private:
+	static TCHAR const* TemplateFilePath;
 };
 
 /** Compute Framework Data Provider for reading skeletal mesh. */
diff --git a/Engine/Plugins/Animation/DeformerGraph/Source/OptimusCore/Private/DataInterfaces/OptimusDataInterfaceDebugDraw.cpp b/Engine/Plugins/Animation/DeformerGraph/Source/OptimusCore/Private/DataInterfaces/OptimusDataInterfaceDebugDraw.cpp
--- a/Engine/Plugins/Animation/DeformerGraph/Source/OptimusCore/Private/DataInterfaces/OptimusDataInterfaceDebugDraw.cpp
+++ b/Engine/Plugins/Animation/DeformerGraph/Source/OptimusCore/Private/DataInterfaces/OptimusDataInterfaceDebugDraw.cpp
@@ -76,9 +76,16 @@ void UOptimusDebugDrawDataInterface::GetShaderParameters(TCHAR const* UID, FShad
 	InOutBuilder.AddNestedStruct<FDebugDrawDataInterfaceParameters>(UID);
 }
 
+TCHAR const* UOptimusDebugDrawDataInterface::TemplateFilePath = TEXT("/Plugin/Optimus/Private/DataInterfaceDebugDraw.ush");
+
+TCHAR const* UOptimusDebugDrawDataInterface::GetShaderVirtualPath() const
+{
+	return TemplateFilePath;
+}
+
 void UOptimusDebugDrawDataInterface::GetShaderHash(FString& InOutKey) const
 {
-	GetShaderFileHash(TEXT("/Plugin/Optimus/Private/DataInterfaceDebugDraw.ush"), EShaderPlatform::SP_PCD3D_SM5).AppendString(InOutKey);
+	GetShaderFileHash(TemplateFilePath, EShaderPlatform::SP_PCD3D_SM5).AppendString(InOutKey);
 }
 
 void UOptimusDebugDrawDataInterface::GetHLSL(FString& OutHLSL, FString const& InDataInterfaceName) const
@@ -89,7 +96,7 @@ void UOptimusDebugDrawDataInterface::GetHLSL(FString& OutHLSL, FString const& In
 	};
 
 	FString TemplateFile;
-	LoadShaderSourceFile(TEXT("/Plugin/Optimus/Private/DataInterfaceDebugDraw.ush"), EShaderPlatform::SP_PCD3D_SM5, &TemplateFile, nullptr);
+	LoadShaderSourceFile(TemplateFilePath, EShaderPlatform::SP_PCD3D_SM5, &TemplateFile, nullptr);
 	OutHLSL += FString::Format(*TemplateFile, TemplateArgs);
 }
 
diff --git a/Engine/Plugins/Animation/DeformerGraph/Source/OptimusCore/Private/DataInterfaces/OptimusDataInterfaceDebugDraw.h b/Engine/Plugins/Animation/DeformerGraph/Source/OptimusCore/Private/DataInterfaces/OptimusDataInterfaceDebugDraw.h
--- a/Engine/Plugins/Animation/DeformerGraph/Source/OptimusCore/Private/DataInterfaces/OptimusDataInterfaceDebugDraw.h
+++ b/Engine/Plugins/Animation/DeformerGraph/Source/OptimusCore/Private/DataInterfaces/OptimusDataInterfaceDebugDraw.h
@@ -56,6 +56,7 @@ public:
 	TCHAR const* GetClassName() const override { return TEXT("DebugDraw"); }
 	void GetSupportedInputs(TArray<FShaderFunctionDefinition>& OutFunctions) const override;
 	void GetShaderParameters(TCHAR const* UID, FShaderParametersMetadataBuilder& InOutBuilder, FShaderParametersMetadataAllocations& InOutAllocations) const override;
+	TCHAR const* GetShaderVirtualPath() const override;
 	void GetShaderHash(FString& InOutKey) const override;
 	void GetHLSL(FString& OutHLSL, FString const& InDataInterfaceName) const override;
 	UComputeDataProvider* CreateDataProvider(TObjectPtr<UObject> InBinding, uint64 InInputMask, uint64 InOutputMask) const override;
@@ -63,6 +64,9 @@ public:
 
 	UPROPERTY(EditAnywhere, Category = DebugDraw, meta = (ShowOnlyInnerProperties))
 	FOptimusDebugDrawParameters DebugDrawParameters;
+
+private:
+	static TCHAR const* TemplateFilePath;
 };
 
 /** Compute Framework Data Provider for writing skinned mesh. */
diff --git a/Engine/Plugins/Animation/DeformerGraph/Source/OptimusCore/Private/DataInterfaces/OptimusDataInterfaceDuplicateVertices.cpp b/Engine/Plugins/Animation/DeformerGraph/Source/OptimusCore/Private/DataInterfaces/OptimusDataInterfaceDuplicateVertices.cpp
--- a/Engine/Plugins/Animation/DeformerGraph/Source/OptimusCore/Private/DataInterfaces/OptimusDataInterfaceDuplicateVertices.cpp
+++ b/Engine/Plugins/Animation/DeformerGraph/Source/OptimusCore/Private/DataInterfaces/OptimusDataInterfaceDuplicateVertices.cpp
@@ -79,9 +79,16 @@ void UOptimusDuplicateVerticesDataInterface::GetPermutations(FComputeKernelPermu
 	OutPermutationVector.AddPermutation(TEXT("ENABLE_DUPLICATED_VERTICES"), 2);
 }
 
+TCHAR const* UOptimusDuplicateVerticesDataInterface::TemplateFilePath = TEXT("/Plugin/Optimus/Private/DataInterfaceDuplicateVertices.ush");
+
+TCHAR const* UOptimusDuplicateVerticesDataInterface::GetShaderVirtualPath() const
+{
+	return TemplateFilePath;
+}
+
 void UOptimusDuplicateVerticesDataInterface::GetShaderHash(FString& InOutKey) const
 {
-	GetShaderFileHash(TEXT("/Plugin/Optimus/Private/DataInterfaceSkinnedMesh.ush"), EShaderPlatform::SP_PCD3D_SM5).AppendString(InOutKey);
+	GetShaderFileHash(TemplateFilePath, EShaderPlatform::SP_PCD3D_SM5).AppendString(InOutKey);
 }
 
 void UOptimusDuplicateVerticesDataInterface::GetHLSL(FString& OutHLSL, FString const& InDataInterfaceName) const
@@ -92,7 +99,7 @@ void UOptimusDuplicateVerticesDataInterface::GetHLSL(FString& OutHLSL, FString c
 	};
 
 	FString TemplateFile;
-	LoadShaderSourceFile(TEXT("/Plugin/Optimus/Private/DataInterfaceDuplicateVertices.ush"), EShaderPlatform::SP_PCD3D_SM5, &TemplateFile, nullptr);
+	LoadShaderSourceFile(TemplateFilePath, EShaderPlatform::SP_PCD3D_SM5, &TemplateFile, nullptr);
 	OutHLSL += FString::Format(*TemplateFile, TemplateArgs);
 }
 
diff --git a/Engine/Plugins/Animation/DeformerGraph/Source/OptimusCore/Private/DataInterfaces/OptimusDataInterfaceDuplicateVertices.h b/Engine/Plugins/Animation/DeformerGraph/Source/OptimusCore/Private/DataInterfaces/OptimusDataInterfaceDuplicateVertices.h
--- a/Engine/Plugins/Animation/DeformerGraph/Source/OptimusCore/Private/DataInterfaces/OptimusDataInterfaceDuplicateVertices.h
+++ b/Engine/Plugins/Animation/DeformerGraph/Source/OptimusCore/Private/DataInterfaces/OptimusDataInterfaceDuplicateVertices.h
@@ -32,10 +32,14 @@ public:
 	void GetSupportedInputs(TArray<FShaderFunctionDefinition>& OutFunctions) const override;
 	void GetShaderParameters(TCHAR const* UID, FShaderParametersMetadataBuilder& InOutBuilder, FShaderParametersMetadataAllocations& InOutAllocations) const override;
 	void GetPermutations(FComputeKernelPermutationVector& OutPermutationVector) const override;
+	TCHAR const* GetShaderVirtualPath() const override;
 	void GetShaderHash(FString& InOutKey) const override;
 	void GetHLSL(FString& OutHLSL, FString const& InDataInterfaceName) const override;
 	UComputeDataProvider* CreateDataProvider(TObjectPtr<UObject> InBinding, uint64 InInputMask, uint64 InOutputMask) const override;
 	//~ End UComputeDataInterface Interface
+
+private:
+	static TCHAR const* TemplateFilePath;
 };
 
 /** Compute Framework Data Provider for reading skeletal mesh. */
diff --git a/Engine/Plugins/Animation/DeformerGraph/Source/OptimusCore/Private/DataInterfaces/OptimusDataInterfaceMorphTarget.cpp b/Engine/Plugins/Animation/DeformerGraph/Source/OptimusCore/Private/DataInterfaces/OptimusDataInterfaceMorphTarget.cpp
--- a/Engine/Plugins/Animation/DeformerGraph/Source/OptimusCore/Private/DataInterfaces/OptimusDataInterfaceMorphTarget.cpp
+++ b/Engine/Plugins/Animation/DeformerGraph/Source/OptimusCore/Private/DataInterfaces/OptimusDataInterfaceMorphTarget.cpp
@@ -67,9 +67,16 @@ void UOptimusMorphTargetDataInterface::GetPermutations(FComputeKernelPermutation
 	OutPermutationVector.AddPermutation(TEXT("ENABLE_DEFORMER_MORPHTARGET"), 2);
 }
 
+TCHAR const* UOptimusMorphTargetDataInterface::TemplateFilePath = TEXT("/Plugin/Optimus/Private/DataInterfaceMorphTarget.ush");
+
+TCHAR const* UOptimusMorphTargetDataInterface::GetShaderVirtualPath() const
+{
+	return TemplateFilePath;
+}
+
 void UOptimusMorphTargetDataInterface::GetShaderHash(FString& InOutKey) const
 {
-	GetShaderFileHash(TEXT("/Plugin/Optimus/Private/DataInterfaceMorphTarget.ush"), EShaderPlatform::SP_PCD3D_SM5).AppendString(InOutKey);
+	GetShaderFileHash(TemplateFilePath, EShaderPlatform::SP_PCD3D_SM5).AppendString(InOutKey);
 }
 
 void UOptimusMorphTargetDataInterface::GetHLSL(FString& OutHLSL, FString const& InDataInterfaceName) const
@@ -80,7 +87,7 @@ void UOptimusMorphTargetDataInterface::GetHLSL(FString& OutHLSL, FString const&
 	};
 
 	FString TemplateFile;
-	LoadShaderSourceFile(TEXT("/Plugin/Optimus/Private/DataInterfaceMorphTarget.ush"), EShaderPlatform::SP_PCD3D_SM5, &TemplateFile, nullptr);
+	LoadShaderSourceFile(TemplateFilePath, EShaderPlatform::SP_PCD3D_SM5, &TemplateFile, nullptr);
 	OutHLSL += FString::Format(*TemplateFile, TemplateArgs);
 }
 
diff --git a/Engine/Plugins/Animation/DeformerGraph/Source/OptimusCore/Private/DataInterfaces/OptimusDataInterfaceMorphTarget.h b/Engine/Plugins/Animation/DeformerGraph/Source/OptimusCore/Private/DataInterfaces/OptimusDataInterfaceMorphTarget.h
--- a/Engine/Plugins/Animation/DeformerGraph/Source/OptimusCore/Private/DataInterfaces/OptimusDataInterfaceMorphTarget.h
+++ b/Engine/Plugins/Animation/DeformerGraph/Source/OptimusCore/Private/DataInterfaces/OptimusDataInterfaceMorphTarget.h
@@ -29,10 +29,14 @@ public:
 	void GetSupportedInputs(TArray<FShaderFunctionDefinition>& OutFunctions) const override;
 	void GetShaderParameters(TCHAR const* UID, FShaderParametersMetadataBuilder& InOutBuilder, FShaderParametersMetadataAllocations& InOutAllocations) const override;
 	void GetPermutations(FComputeKernelPermutationVector& OutPermutationVector) const override;
+	TCHAR const* GetShaderVirtualPath() const override;
 	void GetShaderHash(FString& InOutKey) const override;
 	void GetHLSL(FString& OutHLSL, FString const& InDataInterfaceName) const override;
 	UComputeDataProvider* CreateDataProvider(TObjectPtr<UObject> InBinding, uint64 InInputMask, uint64 InOutputMask) const override;
 	//~ End UComputeDataInterface Interface
+
+private:
+	static TCHAR const* TemplateFilePath;
 };
 
 /** Compute Framework Data Provider for reading skeletal mesh. */
diff --git a/Engine/Plugins/Animation/DeformerGraph/Source/OptimusCore/Private/DataInterfaces/OptimusDataInterfaceRawBuffer.cpp b/Engine/Plugins/Animation/DeformerGraph/Source/OptimusCore/Private/DataInterfaces/OptimusDataInterfaceRawBuffer.cpp
--- a/Engine/Plugins/Animation/DeformerGraph/Source/OptimusCore/Private/DataInterfaces/OptimusDataInterfaceRawBuffer.cpp
+++ b/Engine/Plugins/Animation/DeformerGraph/Source/OptimusCore/Private/DataInterfaces/OptimusDataInterfaceRawBuffer.cpp
@@ -135,9 +135,16 @@ void UOptimusPersistentBufferDataInterface::GetShaderParameters(TCHAR const* UID
 	InOutBuilder.AddNestedStruct<FPersistentBufferDataInterfaceParameters>(UID);
 }
 
+TCHAR const* UOptimusRawBufferDataInterface::TemplateFilePath = TEXT("/Plugin/Optimus/Private/DataInterfaceRawBuffer.ush");
+
+TCHAR const* UOptimusRawBufferDataInterface::GetShaderVirtualPath() const
+{
+	return TemplateFilePath;
+}
+
 void UOptimusRawBufferDataInterface::GetShaderHash(FString& InOutKey) const
 {
-	GetShaderFileHash(TEXT("/Plugin/Optimus/Private/DataInterfaceRawBuffer.ush"), EShaderPlatform::SP_PCD3D_SM5).AppendString(InOutKey);
+	GetShaderFileHash(TemplateFilePath, EShaderPlatform::SP_PCD3D_SM5).AppendString(InOutKey);
 }
 
 void UOptimusRawBufferDataInterface::GetHLSL(FString& OutHLSL, FString const& InDataInterfaceName) const
@@ -157,7 +164,7 @@ void UOptimusRawBufferDataInterface::GetHLSL(FString& OutHLSL, FString const& In
 	};
 
 	FString TemplateFile;
-	LoadShaderSourceFile(TEXT("/Plugin/Optimus/Private/DataInterfaceRawBuffer.ush"), EShaderPlatform::SP_PCD3D_SM5, &TemplateFile, nullptr);
+	LoadShaderSourceFile(TemplateFilePath, EShaderPlatform::SP_PCD3D_SM5, &TemplateFile, nullptr);
 	OutHLSL += FString::Format(*TemplateFile, TemplateArgs);
 }
 
diff --git a/Engine/Plugins/Animation/DeformerGraph/Source/OptimusCore/Private/DataInterfaces/OptimusDataInterfaceRawBuffer.h b/Engine/Plugins/Animation/DeformerGraph/Source/OptimusCore/Private/DataInterfaces/OptimusDataInterfaceRawBuffer.h
--- a/Engine/Plugins/Animation/DeformerGraph/Source/OptimusCore/Private/DataInterfaces/OptimusDataInterfaceRawBuffer.h
+++ b/Engine/Plugins/Animation/DeformerGraph/Source/OptimusCore/Private/DataInterfaces/OptimusDataInterfaceRawBuffer.h
@@ -48,6 +48,7 @@ public:
 	//~ Begin UComputeDataInterface Interface
 	void GetSupportedInputs(TArray<FShaderFunctionDefinition>& OutFunctions) const override;
 	void GetSupportedOutputs(TArray<FShaderFunctionDefinition>& OutFunctions) const override;
+	TCHAR const* GetShaderVirtualPath() const override;
 	void GetShaderHash(FString& InOutKey) const override;
 	void GetHLSL(FString& OutHLSL, FString const& InDataInterfaceName) const override;
 	//~ End UComputeDataInterface Interface
@@ -83,7 +84,10 @@ protected:
 		}
 		return Provider;
 	}
+
 private:
+	static TCHAR const* TemplateFilePath;
+
 	const UOptimusComponentSource* GetComponentSource() const;
 	bool SupportsAtomics() const;
 	FString GetRawType() const;
diff --git a/Engine/Plugins/Animation/DeformerGraph/Source/OptimusCore/Private/DataInterfaces/OptimusDataInterfaceScene.cpp b/Engine/Plugins/Animation/DeformerGraph/Source/OptimusCore/Private/DataInterfaces/OptimusDataInterfaceScene.cpp
--- a/Engine/Plugins/Animation/DeformerGraph/Source/OptimusCore/Private/DataInterfaces/OptimusDataInterfaceScene.cpp
+++ b/Engine/Plugins/Animation/DeformerGraph/Source/OptimusCore/Private/DataInterfaces/OptimusDataInterfaceScene.cpp
@@ -60,9 +60,16 @@ void UOptimusSceneDataInterface::GetShaderParameters(TCHAR const* UID, FShaderPa
 	InOutBuilder.AddNestedStruct<FSceneDataInterfaceParameters>(UID);
 }
 
+TCHAR const* UOptimusSceneDataInterface::TemplateFilePath = TEXT("/Plugin/Optimus/Private/DataInterfaceScene.ush");
+
+TCHAR const* UOptimusSceneDataInterface::GetShaderVirtualPath() const
+{
+	return TemplateFilePath;
+}
+
 void UOptimusSceneDataInterface::GetShaderHash(FString& InOutKey) const
 {
-	GetShaderFileHash(TEXT("/Plugin/Optimus/Private/DataInterfaceScene.ush"), EShaderPlatform::SP_PCD3D_SM5).AppendString(InOutKey);
+	GetShaderFileHash(TemplateFilePath, EShaderPlatform::SP_PCD3D_SM5).AppendString(InOutKey);
 }
 
 void UOptimusSceneDataInterface::GetHLSL(FString& OutHLSL, FString const& InDataInterfaceName) const
@@ -73,7 +80,7 @@ void UOptimusSceneDataInterface::GetHLSL(FString& OutHLSL, FString const& InData
 	};
 
 	FString TemplateFile;
-	LoadShaderSourceFile(TEXT("/Plugin/Optimus/Private/DataInterfaceScene.ush"), EShaderPlatform::SP_PCD3D_SM5, &TemplateFile, nullptr);
+	LoadShaderSourceFile(TemplateFilePath, EShaderPlatform::SP_PCD3D_SM5, &TemplateFile, nullptr);
 	OutHLSL += FString::Format(*TemplateFile, TemplateArgs);
 }
 
diff --git a/Engine/Plugins/Animation/DeformerGraph/Source/OptimusCore/Private/DataInterfaces/OptimusDataInterfaceScene.h b/Engine/Plugins/Animation/DeformerGraph/Source/OptimusCore/Private/DataInterfaces/OptimusDataInterfaceScene.h
--- a/Engine/Plugins/Animation/DeformerGraph/Source/OptimusCore/Private/DataInterfaces/OptimusDataInterfaceScene.h
+++ b/Engine/Plugins/Animation/DeformerGraph/Source/OptimusCore/Private/DataInterfaces/OptimusDataInterfaceScene.h
@@ -26,10 +26,14 @@ public:
 	TCHAR const* GetClassName() const override { return TEXT("Scene"); }
 	void GetSupportedInputs(TArray<FShaderFunctionDefinition>& OutFunctions) const override;
 	void GetShaderParameters(TCHAR const* UID, FShaderParametersMetadataBuilder& InOutBuilder, FShaderParametersMetadataAllocations& InOutAllocations) const override;
+	TCHAR const* GetShaderVirtualPath() const override;
 	void GetShaderHash(FString& InOutKey) const override;
 	void GetHLSL(FString& OutHLSL, FString const& InDataInterfaceName) const override;
 	UComputeDataProvider* CreateDataProvider(TObjectPtr<UObject> InBinding, uint64 InInputMask, uint64 InOutputMask) const override;
 	//~ End UComputeDataInterface Interface
+
+private:
+	static TCHAR const* TemplateFilePath;
 };
 
 /** Compute Framework Data Provider for reading general scene data. */
diff --git a/Engine/Plugins/Animation/DeformerGraph/Source/OptimusCore/Private/DataInterfaces/OptimusDataInterfaceSkeleton.cpp b/Engine/Plugins/Animation/DeformerGraph/Source/OptimusCore/Private/DataInterfaces/OptimusDataInterfaceSkeleton.cpp
--- a/Engine/Plugins/Animation/DeformerGraph/Source/OptimusCore/Private/DataInterfaces/OptimusDataInterfaceSkeleton.cpp
+++ b/Engine/Plugins/Animation/DeformerGraph/Source/OptimusCore/Private/DataInterfaces/OptimusDataInterfaceSkeleton.cpp
@@ -90,12 +90,18 @@ void UOptimusSkeletonDataInterface::GetPermutations(FComputeKernelPermutationVec
 	OutPermutationVector.AddPermutation(TEXT("GPUSKIN_UNLIMITED_BONE_INFLUENCE"), 2);
 	OutPermutationVector.AddPermutation(TEXT("GPUSKIN_BONE_INDEX_UINT16"), 2);
 	OutPermutationVector.AddPermutation(TEXT("GPUSKIN_BONE_WEIGHTS_UINT16"), 2);
-	//OutPermutationVector.AddPermutation(TEXT("MERGE_DUPLICATED_VERTICES"), 2);
+}
+
+TCHAR const* UOptimusSkeletonDataInterface::TemplateFilePath = TEXT("/Plugin/Optimus/Private/DataInterfaceSkeleton.ush");
+
+TCHAR const* UOptimusSkeletonDataInterface::GetShaderVirtualPath() const
+{
+	return TemplateFilePath;
 }
 
 void UOptimusSkeletonDataInterface::GetShaderHash(FString& InOutKey) const
 {
-	GetShaderFileHash(TEXT("/Plugin/Optimus/Private/DataInterfaceSkeleton.ush"), EShaderPlatform::SP_PCD3D_SM5).AppendString(InOutKey);
+	GetShaderFileHash(TemplateFilePath, EShaderPlatform::SP_PCD3D_SM5).AppendString(InOutKey);
 }
 
 void UOptimusSkeletonDataInterface::GetHLSL(FString& OutHLSL, FString const& InDataInterfaceName) const
@@ -106,7 +112,7 @@ void UOptimusSkeletonDataInterface::GetHLSL(FString& OutHLSL, FString const& InD
 	};
 
 	FString TemplateFile;
-	LoadShaderSourceFile(TEXT("/Plugin/Optimus/Private/DataInterfaceSkeleton.ush"), EShaderPlatform::SP_PCD3D_SM5, &TemplateFile, nullptr);
+	LoadShaderSourceFile(TemplateFilePath, EShaderPlatform::SP_PCD3D_SM5, &TemplateFile, nullptr);
 	OutHLSL += FString::Format(*TemplateFile, TemplateArgs);
 }
 
diff --git a/Engine/Plugins/Animation/DeformerGraph/Source/OptimusCore/Private/DataInterfaces/OptimusDataInterfaceSkeleton.h b/Engine/Plugins/Animation/DeformerGraph/Source/OptimusCore/Private/DataInterfaces/OptimusDataInterfaceSkeleton.h
--- a/Engine/Plugins/Animation/DeformerGraph/Source/OptimusCore/Private/DataInterfaces/OptimusDataInterfaceSkeleton.h
+++ b/Engine/Plugins/Animation/DeformerGraph/Source/OptimusCore/Private/DataInterfaces/OptimusDataInterfaceSkeleton.h
@@ -29,10 +29,14 @@ public:
 	void GetSupportedInputs(TArray<FShaderFunctionDefinition>& OutFunctions) const override;
 	void GetShaderParameters(TCHAR const* UID, FShaderParametersMetadataBuilder& InOutBuilder, FShaderParametersMetadataAllocations& InOutAllocations) const override;
 	void GetPermutations(FComputeKernelPermutationVector& OutPermutationVector) const override;
+	TCHAR const* GetShaderVirtualPath() const override;
 	void GetShaderHash(FString& InOutKey) const override;
 	void GetHLSL(FString& OutHLSL, FString const& InDataInterfaceName) const override;
 	UComputeDataProvider* CreateDataProvider(TObjectPtr<UObject> InBinding, uint64 InInputMask, uint64 InOutputMask) const override;
 	//~ End UComputeDataInterface Interface
+
+private:
+	static TCHAR const* TemplateFilePath;
 };
 
 /** Compute Framework Data Provider for reading skeletal mesh. */
diff --git a/Engine/Plugins/Animation/DeformerGraph/Source/OptimusCore/Private/DataInterfaces/OptimusDataInterfaceSkinnedMesh.cpp b/Engine/Plugins/Animation/DeformerGraph/Source/OptimusCore/Private/DataInterfaces/OptimusDataInterfaceSkinnedMesh.cpp
--- a/Engine/Plugins/Animation/DeformerGraph/Source/OptimusCore/Private/DataInterfaces/OptimusDataInterfaceSkinnedMesh.cpp
+++ b/Engine/Plugins/Animation/DeformerGraph/Source/OptimusCore/Private/DataInterfaces/OptimusDataInterfaceSkinnedMesh.cpp
@@ -109,9 +109,16 @@ void UOptimusSkinnedMeshDataInterface::GetShaderParameters(TCHAR const* UID, FSh
 	InOutBuilder.AddNestedStruct<FSkinnedMeshDataInterfaceParameters>(UID);
 }
 
+TCHAR const* UOptimusSkinnedMeshDataInterface::TemplateFilePath = TEXT("/Plugin/Optimus/Private/DataInterfaceSkinnedMesh.ush");
+
+TCHAR const* UOptimusSkinnedMeshDataInterface::GetShaderVirtualPath() const
+{
+	return TemplateFilePath;
+}
+
 void UOptimusSkinnedMeshDataInterface::GetShaderHash(FString& InOutKey) const
 {
-	GetShaderFileHash(TEXT("/Plugin/Optimus/Private/DataInterfaceSkinnedMesh.ush"), EShaderPlatform::SP_PCD3D_SM5).AppendString(InOutKey);
+	GetShaderFileHash(TemplateFilePath, EShaderPlatform::SP_PCD3D_SM5).AppendString(InOutKey);
 }
 
 void UOptimusSkinnedMeshDataInterface::GetHLSL(FString& OutHLSL, FString const& InDataInterfaceName) const
@@ -122,7 +129,7 @@ void UOptimusSkinnedMeshDataInterface::GetHLSL(FString& OutHLSL, FString const&
 	};
 
 	FString TemplateFile;
-	LoadShaderSourceFile(TEXT("/Plugin/Optimus/Private/DataInterfaceSkinnedMesh.ush"), EShaderPlatform::SP_PCD3D_SM5, &TemplateFile, nullptr);
+	LoadShaderSourceFile(TemplateFilePath, EShaderPlatform::SP_PCD3D_SM5, &TemplateFile, nullptr);
 	OutHLSL += FString::Format(*TemplateFile, TemplateArgs);
 }
 
diff --git a/Engine/Plugins/Animation/DeformerGraph/Source/OptimusCore/Private/DataInterfaces/OptimusDataInterfaceSkinnedMesh.h b/Engine/Plugins/Animation/DeformerGraph/Source/OptimusCore/Private/DataInterfaces/OptimusDataInterfaceSkinnedMesh.h
--- a/Engine/Plugins/Animation/DeformerGraph/Source/OptimusCore/Private/DataInterfaces/OptimusDataInterfaceSkinnedMesh.h
+++ b/Engine/Plugins/Animation/DeformerGraph/Source/OptimusCore/Private/DataInterfaces/OptimusDataInterfaceSkinnedMesh.h
@@ -28,10 +28,14 @@ public:
 	TCHAR const* GetClassName() const override { return TEXT("SkinnedMesh"); }
 	void GetSupportedInputs(TArray<FShaderFunctionDefinition>& OutFunctions) const override;
 	void GetShaderParameters(TCHAR const* UID, FShaderParametersMetadataBuilder& InOutBuilder, FShaderParametersMetadataAllocations& InOutAllocations) const override;
+	TCHAR const* GetShaderVirtualPath() const override;
 	void GetShaderHash(FString& InOutKey) const override;
 	void GetHLSL(FString& OutHLSL, FString const& InDataInterfaceName) const override;
 	UComputeDataProvider* CreateDataProvider(TObjectPtr<UObject> InBinding, uint64 InInputMask, uint64 InOutputMask) const override;
 	//~ End UComputeDataInterface Interface
+
+private:
+	static TCHAR const* TemplateFilePath;
 };
 
 /** Compute Framework Data Provider for reading skeletal mesh. */
diff --git a/Engine/Plugins/Animation/DeformerGraph/Source/OptimusCore/Private/DataInterfaces/OptimusDataInterfaceSkinnedMeshExec.cpp b/Engine/Plugins/Animation/DeformerGraph/Source/OptimusCore/Private/DataInterfaces/OptimusDataInterfaceSkinnedMeshExec.cpp
--- a/Engine/Plugins/Animation/DeformerGraph/Source/OptimusCore/Private/DataInterfaces/OptimusDataInterfaceSkinnedMeshExec.cpp
+++ b/Engine/Plugins/Animation/DeformerGraph/Source/OptimusCore/Private/DataInterfaces/OptimusDataInterfaceSkinnedMeshExec.cpp
@@ -53,9 +53,16 @@ void UOptimusSkinnedMeshExecDataInterface::GetShaderParameters(TCHAR const* UID,
 	InOutBuilder.AddNestedStruct<FSkinedMeshExecDataInterfaceParameters>(UID);
 }
 
+TCHAR const* UOptimusSkinnedMeshExecDataInterface::TemplateFilePath = TEXT("/Plugin/Optimus/Private/DataInterfaceSkinnedMeshExec.ush");
+
+TCHAR const* UOptimusSkinnedMeshExecDataInterface::GetShaderVirtualPath() const
+{
+	return TemplateFilePath;
+}
+
 void UOptimusSkinnedMeshExecDataInterface::GetShaderHash(FString& InOutKey) const
 {
-	GetShaderFileHash(TEXT("/Plugin/Optimus/Private/DataInterfaceSkinnedMeshExec.ush"), EShaderPlatform::SP_PCD3D_SM5).AppendString(InOutKey);
+	GetShaderFileHash(TemplateFilePath, EShaderPlatform::SP_PCD3D_SM5).AppendString(InOutKey);
 }
 
 void UOptimusSkinnedMeshExecDataInterface::GetHLSL(FString& OutHLSL, FString const& InDataInterfaceName) const
@@ -66,7 +73,7 @@ void UOptimusSkinnedMeshExecDataInterface::GetHLSL(FString& OutHLSL, FString con
 	};
 
 	FString TemplateFile;
-	LoadShaderSourceFile(TEXT("/Plugin/Optimus/Private/DataInterfaceSkinnedMeshExec.ush"), EShaderPlatform::SP_PCD3D_SM5, &TemplateFile, nullptr);
+	LoadShaderSourceFile(TemplateFilePath, EShaderPlatform::SP_PCD3D_SM5, &TemplateFile, nullptr);
 	OutHLSL += FString::Format(*TemplateFile, TemplateArgs);
 }
 
diff --git a/Engine/Plugins/Animation/DeformerGraph/Source/OptimusCore/Private/DataInterfaces/OptimusDataInterfaceSkinnedMeshExec.h b/Engine/Plugins/Animation/DeformerGraph/Source/OptimusCore/Private/DataInterfaces/OptimusDataInterfaceSkinnedMeshExec.h
--- a/Engine/Plugins/Animation/DeformerGraph/Source/OptimusCore/Private/DataInterfaces/OptimusDataInterfaceSkinnedMeshExec.h
+++ b/Engine/Plugins/Animation/DeformerGraph/Source/OptimusCore/Private/DataInterfaces/OptimusDataInterfaceSkinnedMeshExec.h
@@ -42,6 +42,7 @@ public:
 	bool IsExecutionInterface() const override { return true; }
 	void GetSupportedInputs(TArray<FShaderFunctionDefinition>& OutFunctions) const override;
 	void GetShaderParameters(TCHAR const* UID, FShaderParametersMetadataBuilder& InOutBuilder, FShaderParametersMetadataAllocations& InOutAllocations) const override;
+	TCHAR const* GetShaderVirtualPath() const override;
 	void GetShaderHash(FString& InOutKey) const override;
 	void GetHLSL(FString& OutHLSL, FString const& InDataInterfaceName) const override;
 	UComputeDataProvider* CreateDataProvider(TObjectPtr<UObject> InBinding, uint64 InInputMask, uint64 InOutputMask) const override;
@@ -49,6 +50,9 @@ public:
 
 	UPROPERTY(EditAnywhere, Category = Execution)
 	EOptimusSkinnedMeshExecDomain Domain = EOptimusSkinnedMeshExecDomain::Vertex;
+
+private:
+	static TCHAR const* TemplateFilePath;
 };
 
 /** Compute Framework Data Provider for executing kernels over a skinned mesh domain. */
diff --git a/Engine/Plugins/Animation/DeformerGraph/Source/OptimusCore/Private/DataInterfaces/OptimusDataInterfaceSkinnedMeshWrite.cpp b/Engine/Plugins/Animation/DeformerGraph/Source/OptimusCore/Private/DataInterfaces/OptimusDataInterfaceSkinnedMeshWrite.cpp
--- a/Engine/Plugins/Animation/DeformerGraph/Source/OptimusCore/Private/DataInterfaces/OptimusDataInterfaceSkinnedMeshWrite.cpp
+++ b/Engine/Plugins/Animation/DeformerGraph/Source/OptimusCore/Private/DataInterfaces/OptimusDataInterfaceSkinnedMeshWrite.cpp
@@ -87,9 +87,16 @@ void UOptimusSkinnedMeshWriteDataInterface::GetShaderParameters(TCHAR const* UID
 	InOutBuilder.AddNestedStruct<FSkinedMeshWriteDataInterfaceParameters>(UID);
 }
 
+TCHAR const* UOptimusSkinnedMeshWriteDataInterface::TemplateFilePath = TEXT("/Plugin/Optimus/Private/DataInterfaceSkinnedMeshWrite.ush");
+
+TCHAR const* UOptimusSkinnedMeshWriteDataInterface::GetShaderVirtualPath() const
+{
+	return TemplateFilePath;
+}
+
 void UOptimusSkinnedMeshWriteDataInterface::GetShaderHash(FString& InOutKey) const
 {
-	GetShaderFileHash(TEXT("/Plugin/Optimus/Private/DataInterfaceSkinnedMeshWrite.ush"), EShaderPlatform::SP_PCD3D_SM5).AppendString(InOutKey);
+	GetShaderFileHash(TemplateFilePath, EShaderPlatform::SP_PCD3D_SM5).AppendString(InOutKey);
 }
 
 void UOptimusSkinnedMeshWriteDataInterface::GetHLSL(FString& OutHLSL, FString const& InDataInterfaceName) const
@@ -100,7 +107,7 @@ void UOptimusSkinnedMeshWriteDataInterface::GetHLSL(FString& OutHLSL, FString co
 	};
 
 	FString TemplateFile;
-	LoadShaderSourceFile(TEXT("/Plugin/Optimus/Private/DataInterfaceSkinnedMeshWrite.ush"), EShaderPlatform::SP_PCD3D_SM5, &TemplateFile, nullptr);
+	LoadShaderSourceFile(TemplateFilePath, EShaderPlatform::SP_PCD3D_SM5, &TemplateFile, nullptr);
 	OutHLSL += FString::Format(*TemplateFile, TemplateArgs);
 }
 
diff --git a/Engine/Plugins/Animation/DeformerGraph/Source/OptimusCore/Private/DataInterfaces/OptimusDataInterfaceSkinnedMeshWrite.h b/Engine/Plugins/Animation/DeformerGraph/Source/OptimusCore/Private/DataInterfaces/OptimusDataInterfaceSkinnedMeshWrite.h
--- a/Engine/Plugins/Animation/DeformerGraph/Source/OptimusCore/Private/DataInterfaces/OptimusDataInterfaceSkinnedMeshWrite.h
+++ b/Engine/Plugins/Animation/DeformerGraph/Source/OptimusCore/Private/DataInterfaces/OptimusDataInterfaceSkinnedMeshWrite.h
@@ -32,10 +32,14 @@ public:
 	void GetSupportedInputs(TArray<FShaderFunctionDefinition>& OutFunctions) const override;
 	void GetSupportedOutputs(TArray<FShaderFunctionDefinition>& OutFunctions) const override;
 	void GetShaderParameters(TCHAR const* UID, FShaderParametersMetadataBuilder& InOutBuilder, FShaderParametersMetadataAllocations& InOutAllocations) const override;
+	TCHAR const* GetShaderVirtualPath() const override;
 	void GetShaderHash(FString& InOutKey) const override;
 	void GetHLSL(FString& OutHLSL, FString const& InDataInterfaceName) const override;
 	UComputeDataProvider* CreateDataProvider(TObjectPtr<UObject> InBinding, uint64 InInputMask, uint64 InOutputMask) const override;
 	//~ End UComputeDataInterface Interface
+
+private:
+	static TCHAR const* TemplateFilePath;
 };
 
 /** Compute Framework Data Provider for writing skinned mesh. */
