Optimus: Scale cloth mapping influence distances when skeletal mesh actor is scaled.

diff --git a/Engine/Plugins/Animation/DeformerGraph/Shaders/Private/DataInterfaceCloth.ush b/Engine/Plugins/Animation/DeformerGraph/Shaders/Private/DataInterfaceCloth.ush
--- a/Engine/Plugins/Animation/DeformerGraph/Shaders/Private/DataInterfaceCloth.ush
+++ b/Engine/Plugins/Animation/DeformerGraph/Shaders/Private/DataInterfaceCloth.ush
@@ -26,6 +26,7 @@ uint {DataInterfaceName}_NumVertices;
 uint {DataInterfaceName}_InputStreamStart;
 uint {DataInterfaceName}_NumInfluencesPerVertex;
 float {DataInterfaceName}_ClothBlendWeight;
+float3 {DataInterfaceName}_MeshScale;
 float4x4 {DataInterfaceName}_ClothToLocal;
 Buffer<float4> {DataInterfaceName}_ClothBuffer;
 Buffer<float2> {DataInterfaceName}_ClothPositionsAndNormalsBuffer;
@@ -107,22 +108,24 @@ FClothResult GetClothResult_{DataInterfaceName}(uint VertexIndex)
 				Weight = 1.0f;
 				SumWeights = 1.0f;
 			}
+			
+			float3 Scale = {DataInterfaceName}_MeshScale;
 
-			NormalPosition += Weight * (Influence.NormalBaryCoordsAndDist.x * (A + NA * Influence.NormalBaryCoordsAndDist.w)
-				+ Influence.NormalBaryCoordsAndDist.y * (B + NB * Influence.NormalBaryCoordsAndDist.w)
-				+ Influence.NormalBaryCoordsAndDist.z * (C + NC * Influence.NormalBaryCoordsAndDist.w));
+			NormalPosition += Weight * (Influence.NormalBaryCoordsAndDist.x * (A + NA * Influence.NormalBaryCoordsAndDist.w * Scale.x)
+				+ Influence.NormalBaryCoordsAndDist.y * (B + NB * Influence.NormalBaryCoordsAndDist.w * Scale.y)
+				+ Influence.NormalBaryCoordsAndDist.z * (C + NC * Influence.NormalBaryCoordsAndDist.w * Scale.z));
 
-			TangentPosition += Weight * (Influence.TangentBaryCoordsAndDist.x * (A + NA * Influence.TangentBaryCoordsAndDist.w)
-				+ Influence.TangentBaryCoordsAndDist.y * (B + NB * Influence.TangentBaryCoordsAndDist.w)
-				+ Influence.TangentBaryCoordsAndDist.z * (C + NC * Influence.TangentBaryCoordsAndDist.w));
+			TangentPosition += Weight * (Influence.TangentBaryCoordsAndDist.x * (A + NA * Influence.TangentBaryCoordsAndDist.w * Scale.x)
+				+ Influence.TangentBaryCoordsAndDist.y * (B + NB * Influence.TangentBaryCoordsAndDist.w * Scale.y)
+				+ Influence.TangentBaryCoordsAndDist.z * (C + NC * Influence.TangentBaryCoordsAndDist.w * Scale.z));
 
 			float3 TriangleBary = float3(Influence.PositionBaryCoordsAndDist.x,
 				Influence.PositionBaryCoordsAndDist.y,
 				1.0f - Influence.PositionBaryCoordsAndDist.x - Influence.PositionBaryCoordsAndDist.y);
 
-			float3 SimPosition = TriangleBary.x * (A + NA * Influence.PositionBaryCoordsAndDist.w)
-				+ TriangleBary.y * (B + NB * Influence.PositionBaryCoordsAndDist.w)
-				+ TriangleBary.z * (C + NC * Influence.PositionBaryCoordsAndDist.w);
+			float3 SimPosition = TriangleBary.x * (A + NA * Influence.PositionBaryCoordsAndDist.w * Scale.x)
+				+ TriangleBary.y * (B + NB * Influence.PositionBaryCoordsAndDist.w * Scale.y)
+				+ TriangleBary.z * (C + NC * Influence.PositionBaryCoordsAndDist.w * Scale.z);
 
 			SimulatedPosition += Weight * SimPosition;
 		}
diff --git a/Engine/Plugins/Animation/DeformerGraph/Source/OptimusCore/Private/DataInterfaces/OptimusDataInterfaceCloth.cpp b/Engine/Plugins/Animation/DeformerGraph/Source/OptimusCore/Private/DataInterfaces/OptimusDataInterfaceCloth.cpp
--- a/Engine/Plugins/Animation/DeformerGraph/Source/OptimusCore/Private/DataInterfaces/OptimusDataInterfaceCloth.cpp
+++ b/Engine/Plugins/Animation/DeformerGraph/Source/OptimusCore/Private/DataInterfaces/OptimusDataInterfaceCloth.cpp
@@ -75,6 +75,7 @@ BEGIN_SHADER_PARAMETER_STRUCT(FClothDataInterfaceParameters, )
 	SHADER_PARAMETER(uint32, InputStreamStart)
 	SHADER_PARAMETER(uint32, NumInfluencesPerVertex)
 	SHADER_PARAMETER(float, ClothBlendWeight)
+	SHADER_PARAMETER(FVector3f, MeshScale)
 	SHADER_PARAMETER(FMatrix44f, ClothToLocal)
 	SHADER_PARAMETER_SRV(Buffer<float4>, ClothBuffer)
 	SHADER_PARAMETER_SRV(Buffer<float2>, ClothPositionsAndNormalsBuffer)
@@ -138,6 +139,7 @@ FOptimusClothDataProviderProxy::FOptimusClothDataProviderProxy(USkinnedMeshCompo
 		if (USkeletalMeshComponent* SkeletalMeshComponent = Cast<USkeletalMeshComponent>(SkinnedMeshComponent))
 		{
 			ClothBlendWeight = SkeletalMeshComponent->ClothBlendWeight;
+			MeshScale = (FVector3f)SkeletalMeshComponent->GetComponentScale();
 		}
 	}
 }
@@ -216,6 +218,7 @@ void FOptimusClothDataProviderProxy::GatherDispatchData(FDispatchData const& InD
 		Parameters.NumVertices = RenderSection.NumVertices;
 		Parameters.InputStreamStart = ClothBuffers.ClothInfluenceBufferOffset;
 		Parameters.ClothBlendWeight = bValidCloth ? ClothBlendWeight : 0.f;
+		Parameters.MeshScale = bValidCloth ? MeshScale : FVector3f::OneVector;
 		Parameters.NumInfluencesPerVertex = bValidCloth ? NumClothInfluencesPerVertex : 0;
 		Parameters.ClothToLocal = ClothBuffers.ClothToLocal;
 		Parameters.ClothBuffer = bValidCloth ? ClothBuffers.ClothInfluenceBuffer : NullSRVBinding;
diff --git a/Engine/Plugins/Animation/DeformerGraph/Source/OptimusCore/Private/DataInterfaces/OptimusDataInterfaceCloth.h b/Engine/Plugins/Animation/DeformerGraph/Source/OptimusCore/Private/DataInterfaces/OptimusDataInterfaceCloth.h
--- a/Engine/Plugins/Animation/DeformerGraph/Source/OptimusCore/Private/DataInterfaces/OptimusDataInterfaceCloth.h
+++ b/Engine/Plugins/Animation/DeformerGraph/Source/OptimusCore/Private/DataInterfaces/OptimusDataInterfaceCloth.h
@@ -69,5 +69,6 @@ private:
 
 	FSkeletalMeshObject* SkeletalMeshObject = nullptr;
 	float ClothBlendWeight = 1.0f;
+	FVector3f MeshScale = FVector3f::OneVector;
 	uint32 FrameNumber = 0;
 };
