Optimus: Removed extra floor in the preview viewport

diff --git a/Engine/Plugins/Animation/DeformerGraph/Source/OptimusEditor/Private/OptimusEditor.cpp b/Engine/Plugins/Animation/DeformerGraph/Source/OptimusEditor/Private/OptimusEditor.cpp
--- a/Engine/Plugins/Animation/DeformerGraph/Source/OptimusEditor/Private/OptimusEditor.cpp
+++ b/Engine/Plugins/Animation/DeformerGraph/Source/OptimusEditor/Private/OptimusEditor.cpp
@@ -786,24 +786,6 @@ void FOptimusEditor::BindCommands()
 
 void FOptimusEditor::HandlePreviewSceneCreated(const TSharedRef<IPersonaPreviewScene>& InPreviewScene)
 {
-	static const TCHAR* GroundAssetPath = TEXT("/Engine/MapTemplates/SM_Template_Map_Floor.SM_Template_Map_Floor");
-	UStaticMesh* FloorMesh = Cast<UStaticMesh>(StaticLoadObject(UStaticMesh::StaticClass(), nullptr, GroundAssetPath, nullptr, LOAD_None, nullptr));
-	UMaterial* DefaultMaterial = UMaterial::GetDefaultMaterial(MD_Surface);
-	check(FloorMesh);
-	check(DefaultMaterial);
-
-	// Create ground mesh actor if all the pre-requisites are available.
-	if (ensure(FloorMesh) && ensure(DefaultMaterial))
-	{
-		AStaticMeshActor* GroundActor = InPreviewScene->GetWorld()->SpawnActor<AStaticMeshActor>(AStaticMeshActor::StaticClass(), FTransform::Identity);
-		GroundActor->SetFlags(RF_Transient);
-		GroundActor->GetStaticMeshComponent()->SetStaticMesh(FloorMesh);
-		GroundActor->GetStaticMeshComponent()->SetMaterial(0, DefaultMaterial);
-		GroundActor->SetMobility(EComponentMobility::Static);
-		GroundActor->GetStaticMeshComponent()->SetCollisionEnabled(ECollisionEnabled::QueryAndPhysics);
-		GroundActor->GetStaticMeshComponent()->SetCollisionProfileName(UCollisionProfile::BlockAll_ProfileName);
-		GroundActor->GetStaticMeshComponent()->bSelectable = false;
-	}
 
 	AAnimationEditorPreviewActor* Actor = InPreviewScene->GetWorld()->SpawnActor<AAnimationEditorPreviewActor>(AAnimationEditorPreviewActor::StaticClass(), FTransform::Identity);
 	Actor->SetFlags(RF_Transient);
