Optimus: Fixed bug related to copy pasting nodes

diff --git a/Engine/Plugins/Animation/DeformerGraph/Source/OptimusCore/Private/Nodes/OptimusNode_DataInterface.cpp b/Engine/Plugins/Animation/DeformerGraph/Source/OptimusCore/Private/Nodes/OptimusNode_DataInterface.cpp
--- a/Engine/Plugins/Animation/DeformerGraph/Source/OptimusCore/Private/Nodes/OptimusNode_DataInterface.cpp
+++ b/Engine/Plugins/Animation/DeformerGraph/Source/OptimusCore/Private/Nodes/OptimusNode_DataInterface.cpp
@@ -219,9 +219,11 @@ void UOptimusNode_DataInterface::PostDuplicate(EDuplicateMode::Type DuplicateMod
 {
 	// Currently duplication doesn't set the correct outer so fix here.
 	// We can remove this when duplication handles the outer correctly.
-	if (ensure(DataInterfaceData))
+	if (ensure(DataInterfaceData) && DataInterfaceData->GetOuter() != this)
 	{
-		DataInterfaceData->Rename(nullptr, GetOuter());
+		FObjectDuplicationParameters DupParams = InitStaticDuplicateObjectParams(DataInterfaceData, this);
+		
+		DataInterfaceData = Cast<UOptimusComputeDataInterface>(StaticDuplicateObjectEx(DupParams));	
 	}
 }
 
diff --git a/Engine/Plugins/Animation/DeformerGraph/Source/OptimusCore/Private/Nodes/OptimusNode_ResourceAccessorBase.cpp b/Engine/Plugins/Animation/DeformerGraph/Source/OptimusCore/Private/Nodes/OptimusNode_ResourceAccessorBase.cpp
--- a/Engine/Plugins/Animation/DeformerGraph/Source/OptimusCore/Private/Nodes/OptimusNode_ResourceAccessorBase.cpp
+++ b/Engine/Plugins/Animation/DeformerGraph/Source/OptimusCore/Private/Nodes/OptimusNode_ResourceAccessorBase.cpp
@@ -44,6 +44,7 @@ void UOptimusNode_ResourceAccessorBase::PreDuplicateRequirementActions(
 	bool bFoundDuplicate = false;
 	do
 	{
+		bFoundDuplicate = false;
 		for (const UOptimusResourceDescription* ExistingResourceDesc: Deformer->GetResources())
 		{
 			if (ExistingResourceDesc->ResourceName == DuplicationInfo.ResourceName)
@@ -76,6 +77,7 @@ void UOptimusNode_ResourceAccessorBase::PostDuplicate(EDuplicateMode::Type Dupli
 	
 	if (ResourceDesc.IsValid())
 	{
+		// This path is taken when the entire asset is duplicated
 		const UOptimusDeformer* OldDescOwner = ResourceDesc->GetOwningDeformer();
 		
 		// No action needed if we are copying/pasting within the same deformer asset 
@@ -109,7 +111,7 @@ void UOptimusNode_ResourceAccessorBase::ExportCustomProperties(FOutputDevice& Ou
 {
 	if (const UOptimusResourceDescription* Res = ResourceDesc.Get())
 	{
-		Out.Logf(TEXT("%sCustomProperties ResourceDefinition Name=\"%s\" Type=%s DataDomain=%s\n"),
+		Out.Logf(TEXT("%sCustomProperties ResourceDefinition Name=\"%s\" Type=%s DataDomain=\"%s\"\n"),
 			FCString::Spc(Indent), *Res->ResourceName.ToString(), *Res->DataType->TypeName.ToString(),
 			*Res->DataDomain.ToString());
 	}
