Optimus: Fix non unity errors, see linked horde issue, trivial

#horde https://horde.devtools.epicgames.com/log/63d473595354589b5c232c61?issue=255358
#swarm https://p4-swarm.epicgames.net/reviews/23908937

diff --git a/Engine/Plugins/Animation/DeformerGraph/Source/OptimusCore/Public/OptimusDiagnostic.h b/Engine/Plugins/Animation/DeformerGraph/Source/OptimusCore/Public/OptimusDiagnostic.h
--- a/Engine/Plugins/Animation/DeformerGraph/Source/OptimusCore/Public/OptimusDiagnostic.h
+++ b/Engine/Plugins/Animation/DeformerGraph/Source/OptimusCore/Public/OptimusDiagnostic.h
@@ -3,6 +3,7 @@
 #pragma once
 
 #include "CoreMinimal.h"
+#include "UObject/WeakObjectPtrTemplates.h"
 #include "OptimusDiagnostic.generated.h"
 
 UENUM()
