Optimus: Add half edge connectivity data interface.
Creation of buffers is all done within the data interface, and needs to be moved into the skel mesh compilation. But submitting in this form for testing.

diff --git a/Engine/Plugins/Animation/DeformerGraph/Shaders/Private/DataInterfaceHalfEdge.ush b/Engine/Plugins/Animation/DeformerGraph/Shaders/Private/DataInterfaceHalfEdge.ush
new file mode 100644
--- /dev/null
+++ b/Engine/Plugins/Animation/DeformerGraph/Shaders/Private/DataInterfaceHalfEdge.ush
@@ -0,0 +1,29 @@
+// Copyright Epic Games, Inc. All Rights Reserved.
+
+uint {DataInterfaceName}_NumVertices;
+uint {DataInterfaceName}_NumTriangles;
+uint {DataInterfaceName}_IndexBufferStart;
+uint {DataInterfaceName}_InputStreamStart;
+StructuredBuffer<int> {DataInterfaceName}_VertexToEdgeBuffer;
+StructuredBuffer<int> {DataInterfaceName}_EdgeToTwinEdgeBuffer;
+
+uint ReadNumVertices_{DataInterfaceName}()
+{
+	return {DataInterfaceName}_NumVertices;
+}
+
+uint ReadNumTriangles_{DataInterfaceName}()
+{
+	return {DataInterfaceName}_NumTriangles;
+}
+
+int ReadEdge_{DataInterfaceName}(uint VertexIndex)
+{
+	return {DataInterfaceName}_VertexToEdgeBuffer[VertexIndex + {DataInterfaceName}_InputStreamStart] - {DataInterfaceName}_IndexBufferStart;
+}
+
+int ReadTwinEdge_{DataInterfaceName}(uint EdgeIndex)
+{
+	int TwinEdge = {DataInterfaceName}_EdgeToTwinEdgeBuffer[EdgeIndex + {DataInterfaceName}_IndexBufferStart];
+	return TwinEdge == -1 ? -1 : TwinEdge - {DataInterfaceName}_IndexBufferStart;
+}
diff --git a/Engine/Plugins/Animation/DeformerGraph/Source/OptimusCore/Private/DataInterfaces/OptimusDataInterfaceHalfEdge.cpp b/Engine/Plugins/Animation/DeformerGraph/Source/OptimusCore/Private/DataInterfaces/OptimusDataInterfaceHalfEdge.cpp
new file mode 100644
--- /dev/null
+++ b/Engine/Plugins/Animation/DeformerGraph/Source/OptimusCore/Private/DataInterfaces/OptimusDataInterfaceHalfEdge.cpp
@@ -0,0 +1,411 @@
+// Copyright Epic Games, Inc. All Rights Reserved.
+
+#include "OptimusDataInterfaceHalfEdge.h"
+
+#include "Components/SkinnedMeshComponent.h"
+#include "ComponentSources/OptimusSkinnedMeshComponentSource.h"
+#include "ComputeFramework/ShaderParamTypeDefinition.h"
+#include "OptimusDataDomain.h"
+#include "RenderGraphBuilder.h"
+#include "Rendering/SkeletalMeshLODRenderData.h"
+#include "Rendering/SkeletalMeshRenderData.h"
+#include "ShaderCompilerCore.h"
+#include "ShaderParameterMetadataBuilder.h"
+#include "SkeletalRenderPublic.h"
+
+#include UE_INLINE_GENERATED_CPP_BY_NAME(OptimusDataInterfaceHalfEdge)
+
+struct FEdgeKey
+{
+	FVector3f V0;
+	FVector3f V1;
+
+	bool operator==(FEdgeKey const& Rhs) const { return V0 == Rhs.V0 && V1 == Rhs.V1; }
+};
+
+static uint32 GetTypeHash(FEdgeKey const& Key)
+{
+	return HashCombine(GetTypeHash(Key.V0), GetTypeHash(Key.V1));
+}
+
+namespace
+{
+	/** Builds lookup from edge to twin edge. */
+	class FTwinEdgeBuilder
+	{
+	public:
+		FTwinEdgeBuilder(int32 InNumTriangles)
+		{
+			TwinEdges.SetNum(InNumTriangles * 3);
+			for (int32& TwinEdge : TwinEdges)
+			{
+				TwinEdge = -1;
+			}
+		}
+
+		int32 SmallestVector(FVector3f V0, FVector3f V1)
+		{
+			if (V0.X != V1.X) return V1.X > V0.X;
+			if (V0.Y != V1.Y) return V1.Y > V0.Y;
+			return V1.Z > V0.Z;
+		}
+
+		void Add(int32 InTriangleIndex, TStaticArray<FVector3f, 3>& InVertexPositions)
+		{
+			for (uint32 TriangleEdgeIndex = 0; TriangleEdgeIndex < 3; ++TriangleEdgeIndex)
+			{
+				FVector3f V[2] = { 
+					InVertexPositions[TriangleEdgeIndex],
+					InVertexPositions[(TriangleEdgeIndex + 1) % 3] };
+				int32 VMin = SmallestVector(V[0], V[1]);
+				const FEdgeKey Key({V[VMin], V[1 - VMin]});
+
+				if (FEdgeDescription* EdgeDesc = EdgeMap.Find(Key))
+				{
+					if (EdgeDesc->E1 == -1)
+					{
+						// The second time that we've seen this edge.
+						EdgeDesc->E1 = InTriangleIndex * 3 + TriangleEdgeIndex;
+						TwinEdges[EdgeDesc->E0] = EdgeDesc->E1;
+						TwinEdges[EdgeDesc->E1] = EdgeDesc->E0;
+					}
+					else
+					{
+						// Non-manifold geometry can end up here. Two options: 
+						// (i) Don't store twins for this edge the first pair.
+						// (ii) Reset the edge map entry and hope that good twins are met consecutively when walking the index buffer.
+ 						EdgeDesc->E0 = InTriangleIndex * 3 + TriangleEdgeIndex;
+ 						EdgeDesc->E1 = -1;
+					}
+				}
+				else
+				{
+					// First time that we've seen this edge.
+					EdgeMap.Add(Key).E0 = InTriangleIndex * 3 + TriangleEdgeIndex;
+				}
+			}
+		}
+
+		TArray<int32>& GetTwinEdges()
+		{
+			return TwinEdges;
+		}
+
+	private:
+		struct FEdgeDescription
+		{
+			int32 E0 = -1;
+			int32 E1 = -1;
+		};
+
+		TMap<FEdgeKey, FEdgeDescription> EdgeMap;
+
+		TArray<int32> TwinEdges;
+	};
+
+	/** Builds lookup from (one) vertex to (many) edges. */
+	class FVertexToEdgeBuilder
+	{
+	public:
+		FVertexToEdgeBuilder(int32 InNumVertices, int32 InNumTriangles)
+		{
+			VertexToEdgeItems.Reserve(InNumTriangles * 3);
+			VertexToEdgeItems.AddDefaulted(InNumVertices);
+		}
+
+		void Add(int32 InVertexIndex, int32 InEdgeIndex)
+		{
+			int32 ItemIndex = InVertexIndex;
+			if (VertexToEdgeItems[ItemIndex].EdgeIndex == -1)
+			{
+				VertexToEdgeItems[ItemIndex].EdgeIndex = InEdgeIndex;
+				return;
+			}
+			while (VertexToEdgeItems[ItemIndex].NextItem != -1)
+			{
+				ItemIndex = VertexToEdgeItems[ItemIndex].NextItem;
+			}
+			int32 NewItemIndex = VertexToEdgeItems.AddDefaulted();
+			VertexToEdgeItems[ItemIndex].NextItem = NewItemIndex;
+			VertexToEdgeItems[NewItemIndex].EdgeIndex = InEdgeIndex;
+		}
+
+		using FEdgeInlineArray = TArray<int32, TInlineAllocator<16>>;
+		void Get(int32 InVertexIndex, FEdgeInlineArray& OutEdgeIndices) const
+		{
+			int32 ItemIndex = InVertexIndex;
+			while (ItemIndex != -1)
+			{
+				OutEdgeIndices.Add(VertexToEdgeItems[ItemIndex].EdgeIndex);
+				ItemIndex = VertexToEdgeItems[ItemIndex].NextItem;
+			}
+		}
+
+	private:
+		struct FVertexToEdgeItem
+		{
+			int32 EdgeIndex = -1;
+			int32 NextItem = -1;
+		};
+
+		/** First NumVertices entries are allocated up front as start of per vertex linked lists. */
+		TArray<FVertexToEdgeItem> VertexToEdgeItems;
+	};
+
+	static int32 GetHeadEdgeForVertex(int32 InVertexIndex, FVertexToEdgeBuilder const& InVertexToEdge, TArray<int32> const& InEdgeToTwinEdge)
+	{
+		FVertexToEdgeBuilder::FEdgeInlineArray EdgeIndicesScratchBuffer;
+		InVertexToEdge.Get(InVertexIndex, EdgeIndicesScratchBuffer);
+
+		// Iterate backwards around the edges until we find a border, or a non recognized edge (which may lead to a border).
+		const int32 StartEdgeIndex = EdgeIndicesScratchBuffer[0];
+		int32 EdgeIndex = StartEdgeIndex;
+		while(1)
+		{
+			const int32 LastEdgeIndex = EdgeIndex;
+			EdgeIndex = ((EdgeIndex / 3) * 3) + ((EdgeIndex + 2) % 3);
+			EdgeIndex = InEdgeToTwinEdge[EdgeIndex];
+			
+			if (EdgeIndex == StartEdgeIndex)
+			{
+				break;
+			}
+			if (EdgeIndex == -1 || EdgeIndicesScratchBuffer.Find(EdgeIndex) == INDEX_NONE)
+			{
+				EdgeIndex = LastEdgeIndex;
+				break;
+			}
+		}
+
+		return EdgeIndex;
+	}
+
+	static void BuildHalfEdgeBuffers(const FSkeletalMeshLODRenderData& InLodRenderData, TArray<int32>& OutVertexToEdge, TArray<int32>& OutEdgeToTwinEdge)
+	{
+		const FRawStaticIndexBuffer16or32Interface* IndexBuffer = InLodRenderData.MultiSizeIndexContainer.GetIndexBuffer();
+		const int32 IndexCount = IndexBuffer->Num();
+		const int32 TriangleCount = IndexCount / 3;
+
+		const FPositionVertexBuffer& VertexBuffer = InLodRenderData.StaticVertexBuffers.PositionVertexBuffer;
+		const int32 VertexCount = InLodRenderData.StaticVertexBuffers.PositionVertexBuffer.GetNumVertices();
+
+		// Build edge to twin edge map.
+		{
+			FTwinEdgeBuilder TwinEdgeBuilder(TriangleCount);
+			for (int32 TriangleIndex = 0; TriangleIndex < TriangleCount; ++TriangleIndex)
+			{
+				TStaticArray<FVector3f, 3> VertexPositions;
+				VertexPositions[0] = VertexBuffer.VertexPosition(IndexBuffer->Get(TriangleIndex * 3 + 0));
+				VertexPositions[1] = VertexBuffer.VertexPosition(IndexBuffer->Get(TriangleIndex * 3 + 1));
+				VertexPositions[2] = VertexBuffer.VertexPosition(IndexBuffer->Get(TriangleIndex * 3 + 2));
+				TwinEdgeBuilder.Add(TriangleIndex, VertexPositions);
+			}
+			OutEdgeToTwinEdge = MoveTemp(TwinEdgeBuilder.GetTwinEdges());
+		}
+
+		// Build vertex to edge index map.
+		FVertexToEdgeBuilder VertexToEdgeBuilder(VertexCount, TriangleCount);
+		for (int32 Index = 0; Index < TriangleCount * 3; ++Index)
+		{
+			const int32 VertexIndex = IndexBuffer->Get(Index);
+			const int32 EdgeIndex = Index;
+			VertexToEdgeBuilder.Add(VertexIndex, EdgeIndex);
+		}
+
+		// VertexToEdgeBuilder gives us multiple edges per vertex.
+		// These should form a linked list of edges. 
+		// Usually that list will be a loop, but at border geometry there will be a head and end to the list.
+		// For loops it doesn't matter where we start iteration. But at borders we want to use the head of the linked list to start iteration.
+		// Non-manifold geometry will be strange, and potentially have multiple edge independent loops. We just do the best we can in that case.
+		OutVertexToEdge.SetNumUninitialized(VertexCount);
+		for (int32 VertexIndex = 0; VertexIndex < VertexCount; ++VertexIndex)
+		{
+			OutVertexToEdge[VertexIndex] = GetHeadEdgeForVertex(VertexIndex, VertexToEdgeBuilder, OutEdgeToTwinEdge);
+		}
+	}
+}
+
+
+FString UOptimusHalfEdgeDataInterface::GetDisplayName() const
+{
+	return TEXT("HalfEdge");
+}
+
+TArray<FOptimusCDIPinDefinition> UOptimusHalfEdgeDataInterface::GetPinDefinitions() const
+{
+	FName Vertex(UOptimusSkinnedMeshComponentSource::Domains::Vertex);
+	FName Triangle(UOptimusSkinnedMeshComponentSource::Domains::Triangle);
+
+	TArray<FOptimusCDIPinDefinition> Defs;
+	Defs.Add({ "EdgesPerVertex", "ReadEdge", Vertex, "ReadNumVertices" });
+	Defs.Add({ "TwinEdges", "ReadTwinEdge", Triangle, 3, "ReadNumTriangles" });
+	return Defs;
+}
+
+TSubclassOf<UActorComponent> UOptimusHalfEdgeDataInterface::GetRequiredComponentClass() const
+{
+	return USkinnedMeshComponent::StaticClass();
+}
+
+void UOptimusHalfEdgeDataInterface::GetSupportedInputs(TArray<FShaderFunctionDefinition>& OutFunctions) const
+{
+	OutFunctions.AddDefaulted_GetRef()
+		.SetName(TEXT("ReadNumVertices"))
+		.AddReturnType(EShaderFundamentalType::Uint);
+
+	OutFunctions.AddDefaulted_GetRef()
+		.SetName(TEXT("ReadNumTriangles"))
+		.AddReturnType(EShaderFundamentalType::Uint);
+
+	OutFunctions.AddDefaulted_GetRef()
+		.SetName(TEXT("ReadEdge"))
+		.AddReturnType(EShaderFundamentalType::Int)
+		.AddParam(EShaderFundamentalType::Uint);
+
+	OutFunctions.AddDefaulted_GetRef()
+		.SetName(TEXT("ReadTwinEdge"))
+		.AddReturnType(EShaderFundamentalType::Int)
+		.AddParam(EShaderFundamentalType::Uint);
+}
+
+BEGIN_SHADER_PARAMETER_STRUCT(FHalfEdgeDataInterfaceParameters, )
+	SHADER_PARAMETER(uint32, NumVertices)
+	SHADER_PARAMETER(uint32, NumTriangles)
+	SHADER_PARAMETER(uint32, IndexBufferStart)
+	SHADER_PARAMETER(uint32, InputStreamStart)
+	SHADER_PARAMETER_RDG_BUFFER_SRV(Buffer<int>, VertexToEdgeBuffer)
+	SHADER_PARAMETER_RDG_BUFFER_SRV(Buffer<int>, EdgeToTwinEdgeBuffer)
+END_SHADER_PARAMETER_STRUCT()
+
+void UOptimusHalfEdgeDataInterface::GetShaderParameters(TCHAR const* UID, FShaderParametersMetadataBuilder& InOutBuilder, FShaderParametersMetadataAllocations& InOutAllocations) const
+{
+	InOutBuilder.AddNestedStruct<FHalfEdgeDataInterfaceParameters>(UID);
+}
+
+TCHAR const* UOptimusHalfEdgeDataInterface::TemplateFilePath = TEXT("/Plugin/Optimus/Private/DataInterfaceHalfEdge.ush");
+
+TCHAR const* UOptimusHalfEdgeDataInterface::GetShaderVirtualPath() const
+{
+	return TemplateFilePath;
+}
+
+void UOptimusHalfEdgeDataInterface::GetShaderHash(FString& InOutKey) const
+{
+	GetShaderFileHash(TemplateFilePath, EShaderPlatform::SP_PCD3D_SM5).AppendString(InOutKey);
+}
+
+void UOptimusHalfEdgeDataInterface::GetHLSL(FString& OutHLSL, FString const& InDataInterfaceName) const
+{
+	TMap<FString, FStringFormatArg> TemplateArgs =
+	{
+		{ TEXT("DataInterfaceName"), InDataInterfaceName },
+	};
+
+	FString TemplateFile;
+	LoadShaderSourceFile(TemplateFilePath, EShaderPlatform::SP_PCD3D_SM5, &TemplateFile, nullptr);
+	OutHLSL += FString::Format(*TemplateFile, TemplateArgs);
+}
+
+UComputeDataProvider* UOptimusHalfEdgeDataInterface::CreateDataProvider(TObjectPtr<UObject> InBinding, uint64 InInputMask, uint64 InOutputMask) const
+{
+	UOptimusHalfEdgeDataProvider* Provider = NewObject<UOptimusHalfEdgeDataProvider>();
+
+	Provider->SkinnedMesh = Cast<USkinnedMeshComponent>(InBinding);
+
+	if (Provider->SkinnedMesh != nullptr)
+	{
+		// Build data and store with the provider.
+		// todo[CF]: We need to move this to the skeletal mesh and make part of cooked mesh data instead.
+		FSkeletalMeshRenderData const* SkeletalMeshRenderData = Provider->SkinnedMesh->GetSkeletalMeshRenderData();
+		if (SkeletalMeshRenderData != nullptr)
+		{
+			Provider->VertexToEdgePerLod.SetNum(SkeletalMeshRenderData->NumInlinedLODs);
+			Provider->EdgeToTwinEdgePerLod.SetNum(SkeletalMeshRenderData->NumInlinedLODs);
+
+			for (int32 LodIndex = 0; LodIndex < SkeletalMeshRenderData->NumInlinedLODs; ++LodIndex)
+			{
+				FSkeletalMeshLODRenderData const* LodRenderData = &SkeletalMeshRenderData->LODRenderData[LodIndex];
+				BuildHalfEdgeBuffers(*LodRenderData, Provider->VertexToEdgePerLod[LodIndex], Provider->EdgeToTwinEdgePerLod[LodIndex]);
+			}
+		}
+	}
+
+	return Provider;
+}
+
+
+FComputeDataProviderRenderProxy* UOptimusHalfEdgeDataProvider::GetRenderProxy()
+{
+	return new FOptimusHalfEdgeDataProviderProxy(SkinnedMesh, VertexToEdgePerLod, EdgeToTwinEdgePerLod);
+}
+
+
+FOptimusHalfEdgeDataProviderProxy::FOptimusHalfEdgeDataProviderProxy(
+	USkinnedMeshComponent* SkinnedMeshComponent, 
+	TArray< TArray<int32> >& InVertexToEdgePerLod,
+	TArray< TArray<int32> >& InEdgeToTwinEdgePerLod)
+	: SkeletalMeshObject(SkinnedMeshComponent != nullptr ? SkinnedMeshComponent->MeshObject : nullptr)
+	, VertexToEdgePerLod(InVertexToEdgePerLod)
+	, EdgeToTwinEdgePerLod(InEdgeToTwinEdgePerLod)
+{
+}
+
+bool FOptimusHalfEdgeDataProviderProxy::IsValid(FValidationData const& InValidationData) const
+{
+	if (InValidationData.ParameterStructSize != sizeof(FParameters))
+	{
+		return false;
+	}
+	if (SkeletalMeshObject == nullptr || VertexToEdgePerLod.Num() == 0 || EdgeToTwinEdgePerLod.Num() == 0)
+	{
+		return false;
+	}
+	if (SkeletalMeshObject->GetSkeletalMeshRenderData().LODRenderData[SkeletalMeshObject->GetLOD()].RenderSections.Num() != InValidationData.NumInvocations)
+	{
+		return false;
+	}
+
+	return true;
+}
+
+void FOptimusHalfEdgeDataProviderProxy::AllocateResources(FRDGBuilder& GraphBuilder)
+{
+	const int32 LodIndex = SkeletalMeshObject->GetLOD();
+
+	// todo[CF]: Updating buffer every frame is obviously bad, but just getting things working initially.
+	{
+		TArray<int32> const& Data = VertexToEdgePerLod[LodIndex];
+		VertexToEdgeBuffer = GraphBuilder.CreateBuffer(FRDGBufferDesc::CreateStructuredDesc(sizeof(uint32), Data.Num()), TEXT("Optimus.VertexToEdge"));
+		VertexToEdgeBufferSRV = GraphBuilder.CreateSRV(VertexToEdgeBuffer);
+		GraphBuilder.QueueBufferUpload(VertexToEdgeBuffer, Data.GetData(), Data.Num() * sizeof(uint32), ERDGInitialDataFlags::None);
+	}
+	{
+		TArray<int32> const& Data = EdgeToTwinEdgePerLod[LodIndex];
+		EdgeToTwinEdgeBuffer = GraphBuilder.CreateBuffer(FRDGBufferDesc::CreateStructuredDesc(sizeof(uint32), Data.Num()), TEXT("Optimus.EdgeToTwinEdge"));
+		EdgeToTwinEdgeBufferSRV = GraphBuilder.CreateSRV(EdgeToTwinEdgeBuffer);
+		GraphBuilder.QueueBufferUpload(EdgeToTwinEdgeBuffer, Data.GetData(), Data.Num() * sizeof(uint32), ERDGInitialDataFlags::None);
+	}
+}
+
+void FOptimusHalfEdgeDataProviderProxy::GatherDispatchData(FDispatchData const& InDispatchData)
+{
+	const int32 LodIndex = SkeletalMeshObject->GetLOD();
+	FSkeletalMeshRenderData const& SkeletalMeshRenderData = SkeletalMeshObject->GetSkeletalMeshRenderData();
+	FSkeletalMeshLODRenderData const* LodRenderData = &SkeletalMeshRenderData.LODRenderData[LodIndex];
+
+	FRHIShaderResourceView* NullSRVBinding = GWhiteVertexBufferWithSRV->ShaderResourceViewRHI.GetReference();
+
+	const TStridedView<FParameters> ParameterArray = MakeStridedParameterView<FParameters>(InDispatchData);
+	for (int32 InvocationIndex = 0; InvocationIndex < ParameterArray.Num(); ++InvocationIndex)
+	{
+		FSkelMeshRenderSection const& RenderSection = LodRenderData->RenderSections[InvocationIndex];
+
+		FParameters& Parameters = ParameterArray[InvocationIndex];
+		Parameters.NumVertices = RenderSection.NumVertices;
+		Parameters.NumTriangles = RenderSection.NumTriangles;
+		Parameters.IndexBufferStart = RenderSection.BaseIndex;
+		Parameters.InputStreamStart = RenderSection.BaseVertexIndex;
+		Parameters.VertexToEdgeBuffer = VertexToEdgeBufferSRV;
+		Parameters.EdgeToTwinEdgeBuffer = EdgeToTwinEdgeBufferSRV;
+	}
+}
diff --git a/Engine/Plugins/Animation/DeformerGraph/Source/OptimusCore/Private/DataInterfaces/OptimusDataInterfaceHalfEdge.h b/Engine/Plugins/Animation/DeformerGraph/Source/OptimusCore/Private/DataInterfaces/OptimusDataInterfaceHalfEdge.h
new file mode 100644
--- /dev/null
+++ b/Engine/Plugins/Animation/DeformerGraph/Source/OptimusCore/Private/DataInterfaces/OptimusDataInterfaceHalfEdge.h
@@ -0,0 +1,87 @@
+// Copyright Epic Games, Inc. All Rights Reserved.
+
+#pragma once
+
+#include "OptimusComputeDataInterface.h"
+#include "ComputeFramework/ComputeDataProvider.h"
+#include "OptimusDataInterfaceHalfEdge.generated.h"
+
+class FHalfEdgeDataInterfaceParameters;
+class FRDGBuffer;
+class FRDGBufferSRV;
+class FSkeletalMeshObject;
+class USkinnedMeshComponent;
+
+
+/** Compute Framework Data Interface for reading mesh half edge data. */
+UCLASS(Category = ComputeFramework)
+class OPTIMUSCORE_API UOptimusHalfEdgeDataInterface : public UOptimusComputeDataInterface
+{
+	GENERATED_BODY()
+
+public:
+	//~ Begin UOptimusComputeDataInterface Interface
+	FString GetDisplayName() const override;
+	TArray<FOptimusCDIPinDefinition> GetPinDefinitions() const override;
+	TSubclassOf<UActorComponent> GetRequiredComponentClass() const override;
+	//~ End UOptimusComputeDataInterface Interface
+	
+	//~ Begin UComputeDataInterface Interface
+	TCHAR const* GetClassName() const override { return TEXT("HalfEdge"); }
+	void GetSupportedInputs(TArray<FShaderFunctionDefinition>& OutFunctions) const override;
+	void GetShaderParameters(TCHAR const* UID, FShaderParametersMetadataBuilder& InOutBuilder, FShaderParametersMetadataAllocations& InOutAllocations) const override;
+	TCHAR const* GetShaderVirtualPath() const override;
+	void GetShaderHash(FString& InOutKey) const override;
+	void GetHLSL(FString& OutHLSL, FString const& InDataInterfaceName) const override;
+	UComputeDataProvider* CreateDataProvider(TObjectPtr<UObject> InBinding, uint64 InInputMask, uint64 InOutputMask) const override;
+	//~ End UComputeDataInterface Interface
+
+private:
+	static TCHAR const* TemplateFilePath;
+};
+
+/** Compute Framework Data Provider for reading skeletal mesh. */
+UCLASS(BlueprintType, editinlinenew, Category = ComputeFramework)
+class UOptimusHalfEdgeDataProvider : public UComputeDataProvider
+{
+	GENERATED_BODY()
+
+public:
+	UPROPERTY(EditAnywhere, BlueprintReadWrite, Category = Binding)
+	TObjectPtr<USkinnedMeshComponent> SkinnedMesh = nullptr;
+
+	//~ Begin UComputeDataProvider Interface
+	FComputeDataProviderRenderProxy* GetRenderProxy() override;
+	//~ End UComputeDataProvider Interface
+
+	TArray< TArray<int32> > VertexToEdgePerLod;
+	TArray< TArray<int32> > EdgeToTwinEdgePerLod;
+};
+
+class FOptimusHalfEdgeDataProviderProxy : public FComputeDataProviderRenderProxy
+{
+public:
+	FOptimusHalfEdgeDataProviderProxy(
+		USkinnedMeshComponent* SkinnedMeshComponent, 
+		TArray< TArray<int32> >& InVertexToEdgePerLod,
+		TArray< TArray<int32> >& InEdgeToTwinEdgeBuffer);
+
+	//~ Begin FComputeDataProviderRenderProxy Interface
+	bool IsValid(FValidationData const& InValidationData) const override;
+	void AllocateResources(FRDGBuilder& GraphBuilder) override;
+	void GatherDispatchData(FDispatchData const& InDispatchData);
+	//~ End FComputeDataProviderRenderProxy Interface
+
+private:
+	using FParameters = FHalfEdgeDataInterfaceParameters;
+
+	FSkeletalMeshObject* SkeletalMeshObject = nullptr;
+
+	TArray< TArray<int32> > const& VertexToEdgePerLod;
+	FRDGBuffer* VertexToEdgeBuffer = nullptr;
+	FRDGBufferSRV* VertexToEdgeBufferSRV = nullptr;
+
+	TArray< TArray<int32> > const& EdgeToTwinEdgePerLod;
+	FRDGBuffer* EdgeToTwinEdgeBuffer = nullptr;
+	FRDGBufferSRV* EdgeToTwinEdgeBufferSRV = nullptr;
+};
