Optimus: added a required header include to fix the nonunity build error, Notice that in the preflight there are no optimus errors anymore

#horde https://horde.devtools.epicgames.com/job/63dcbf128e05384fafc13434?step=4edd

diff --git a/Engine/Plugins/Animation/DeformerGraph/Source/OptimusCore/Public/OptimusDiagnostic.h b/Engine/Plugins/Animation/DeformerGraph/Source/OptimusCore/Public/OptimusDiagnostic.h
--- a/Engine/Plugins/Animation/DeformerGraph/Source/OptimusCore/Public/OptimusDiagnostic.h
+++ b/Engine/Plugins/Animation/DeformerGraph/Source/OptimusCore/Public/OptimusDiagnostic.h
@@ -2,8 +2,8 @@
 
 #pragma once
 
+#include "UObject/WeakObjectPtr.h"
 #include "CoreMinimal.h"
-#include "UObject/WeakObjectPtrTemplates.h"
 #include "OptimusDiagnostic.generated.h"
 
 UENUM()
