Optimus: Change buffer behavior so that we no longer split buffers according to sub invocations.
Instead we use a unified buffer for each invocation.
This is much more natural, and will work better with more complex domain definitions.
One reason to not do this in the past was to avoid unnecesary GPU write barriers. But the ERDGUnorderedAccessViewFlags::SkipBarrier flag now supports this correctly.

diff --git a/Engine/Plugins/Animation/DeformerGraph/Source/OptimusCore/Private/DataInterfaces/OptimusDataInterfaceRawBuffer.cpp b/Engine/Plugins/Animation/DeformerGraph/Source/OptimusCore/Private/DataInterfaces/OptimusDataInterfaceRawBuffer.cpp
--- a/Engine/Plugins/Animation/DeformerGraph/Source/OptimusCore/Private/DataInterfaces/OptimusDataInterfaceRawBuffer.cpp
+++ b/Engine/Plugins/Animation/DeformerGraph/Source/OptimusCore/Private/DataInterfaces/OptimusDataInterfaceRawBuffer.cpp
@@ -311,9 +311,14 @@ FOptimusTransientBufferDataProviderProxy::FOptimusTransientBufferDataProviderPro
 	int32 InElementStride,
 	int32 InRawStride) 
 	: InvocationElementCounts(InInvocationElementCounts)
+	, TotalElementCount(0)
 	, ElementStride(InElementStride)
 	, RawStride(InRawStride)
 {
+	for (int32 NumElements : InvocationElementCounts)
+	{
+		TotalElementCount += NumElements;
+	}
 }
 
 bool FOptimusTransientBufferDataProviderProxy::IsValid(FValidationData const& InValidationData) const
@@ -337,24 +342,23 @@ void FOptimusTransientBufferDataProviderProxy::AllocateResources(FRDGBuilder& Gr
 	const int32 Stride = RawStride ? RawStride : ElementStride;
 	const int32 ElementStrideMultiplier = RawStride ? ElementStride / RawStride : 1;
 
-	for (const int32 NumElements: InvocationElementCounts)
-	{
-		Buffer.Add(GraphBuilder.CreateBuffer(FRDGBufferDesc::CreateStructuredDesc(Stride, NumElements * ElementStrideMultiplier), TEXT("TransientBuffer"), ERDGBufferFlags::None));
-		BufferSRV.Add(GraphBuilder.CreateSRV(Buffer.Last()));
-		BufferUAV.Add(GraphBuilder.CreateUAV(Buffer.Last()));
-	}
+	Buffer = GraphBuilder.CreateBuffer(FRDGBufferDesc::CreateStructuredDesc(Stride, TotalElementCount * ElementStrideMultiplier), TEXT("TransientBuffer"), ERDGBufferFlags::None);
+	BufferSRV = GraphBuilder.CreateSRV(Buffer);
+	BufferUAV = GraphBuilder.CreateUAV(Buffer, ERDGUnorderedAccessViewFlags::SkipBarrier);
 }
 
 void FOptimusTransientBufferDataProviderProxy::GatherDispatchData(FDispatchData const& InDispatchData)
 {
 	TStridedView<FParameters> ParameterArray = MakeStridedParameterView<FParameters>(InDispatchData);
-	for (int32 InvocationIndex = 0; InvocationIndex < ParameterArray.Num(); ++InvocationIndex)
+	for (int32 InvocationIndex = 0, StartOffset = 0; InvocationIndex < ParameterArray.Num(); ++InvocationIndex)
 	{
 		FParameters& Parameters = ParameterArray[InvocationIndex];
-		Parameters.StartOffset = 0;
+		Parameters.StartOffset = StartOffset;
 		Parameters.BufferSize = InvocationElementCounts[InvocationIndex];
-		Parameters.BufferSRV = BufferSRV[InvocationIndex];
-		Parameters.BufferUAV = BufferUAV[InvocationIndex];
+		Parameters.BufferSRV = BufferSRV;
+		Parameters.BufferUAV = BufferUAV;
+		
+		StartOffset += InvocationElementCounts[InvocationIndex];
 	}
 }
 
@@ -367,12 +371,17 @@ FOptimusPersistentBufferDataProviderProxy::FOptimusPersistentBufferDataProviderP
 	FName InResourceName,
 	int32 InLODIndex)
 	: InvocationElementCounts(InInvocationElementCounts)
+	, TotalElementCount(0)
 	, ElementStride(InElementStride)
 	, RawStride(InRawStride)
 	, BufferPool(InBufferPool)
 	, ResourceName(InResourceName)
 	, LODIndex(InLODIndex)
 {
+	for (int32 NumElements : InvocationElementCounts)
+	{
+		TotalElementCount += NumElements;
+	}
 }
 
 bool FOptimusPersistentBufferDataProviderProxy::IsValid(FValidationData const& InValidationData) const
@@ -391,23 +400,27 @@ bool FOptimusPersistentBufferDataProviderProxy::IsValid(FValidationData const& I
 
 void FOptimusPersistentBufferDataProviderProxy::AllocateResources(FRDGBuilder& GraphBuilder)
 {
-	BufferPool->GetResourceBuffers(GraphBuilder, ResourceName, LODIndex, ElementStride, RawStride, InvocationElementCounts, Buffers);
-	ensure(Buffers.Num() == InvocationElementCounts.Num());
-	BufferUAVs.Reserve(Buffers.Num());
-	for (FRDGBufferRef BufferRef : Buffers)
-	{
-		BufferUAVs.Add(GraphBuilder.CreateUAV(BufferRef));
-	}
+	TArray<int32> Count;
+	Count.Add(TotalElementCount);
+	TArray<FRDGBufferRef> Buffers;
+	BufferPool->GetResourceBuffers(GraphBuilder, ResourceName, LODIndex, ElementStride, RawStride, Count, Buffers);
+
+	ensure(Buffers.Num() == 1);
+	Buffer = Buffers[0];
+
+	BufferUAV = GraphBuilder.CreateUAV(Buffer, ERDGUnorderedAccessViewFlags::SkipBarrier);
 }
 
 void FOptimusPersistentBufferDataProviderProxy::GatherDispatchData(FDispatchData const& InDispatchData)
 {
 	TStridedView<FParameters> ParameterArray = MakeStridedParameterView<FParameters>(InDispatchData);
-	for (int32 InvocationIndex = 0; InvocationIndex < ParameterArray.Num(); ++InvocationIndex)
+	for (int32 InvocationIndex = 0, StartOffset = 0; InvocationIndex < ParameterArray.Num(); ++InvocationIndex)
 	{
 		FParameters& Parameters = ParameterArray[InvocationIndex];
-		Parameters.StartOffset = 0;
+		Parameters.StartOffset = StartOffset;
 		Parameters.BufferSize = InvocationElementCounts[InvocationIndex];
-		Parameters.BufferUAV = BufferUAVs[InvocationIndex];
+		Parameters.BufferUAV = BufferUAV;
+
+		StartOffset += InvocationElementCounts[InvocationIndex];
 	}
 }
diff --git a/Engine/Plugins/Animation/DeformerGraph/Source/OptimusCore/Private/DataInterfaces/OptimusDataInterfaceRawBuffer.h b/Engine/Plugins/Animation/DeformerGraph/Source/OptimusCore/Private/DataInterfaces/OptimusDataInterfaceRawBuffer.h
--- a/Engine/Plugins/Animation/DeformerGraph/Source/OptimusCore/Private/DataInterfaces/OptimusDataInterfaceRawBuffer.h
+++ b/Engine/Plugins/Animation/DeformerGraph/Source/OptimusCore/Private/DataInterfaces/OptimusDataInterfaceRawBuffer.h
@@ -6,14 +6,12 @@
 #include "ComputeFramework/ShaderParamTypeDefinition.h"
 #include "OptimusComputeDataInterface.h"
 #include "OptimusDataDomain.h"
+#include "RenderGraphFwd.h"
 
 #include "OptimusDataInterfaceRawBuffer.generated.h"
 
 class FOptimusPersistentBufferPool;
 class FPersistentBufferDataInterfaceParameters;
-class FRDGBuffer;
-class FRDGBufferSRV;
-class FRDGBufferUAV;
 class FTransientBufferDataInterfaceParameters;
 class UOptimusComponentSource;
 class UOptimusComponentSourceBinding;
@@ -228,12 +226,13 @@ private:
 	using FParameters = FTransientBufferDataInterfaceParameters;
 
 	const TArray<int32> InvocationElementCounts;
+	int32 TotalElementCount;
 	const int32 ElementStride;
 	const int32 RawStride;
 
-	TArray<FRDGBuffer*> Buffer;
-	TArray<FRDGBufferSRV*> BufferSRV;
-	TArray<FRDGBufferUAV*> BufferUAV;
+	FRDGBufferRef Buffer;
+	FRDGBufferSRVRef BufferSRV;
+	FRDGBufferUAVRef BufferUAV;
 };
 
 
@@ -260,12 +259,13 @@ private:
 	using FParameters = FPersistentBufferDataInterfaceParameters;
 
 	const TArray<int32> InvocationElementCounts;
+	int32 TotalElementCount;
 	const int32 ElementStride;
 	const int32 RawStride;
 	const TSharedPtr<FOptimusPersistentBufferPool> BufferPool;
 	const FName ResourceName;
 	const int32 LODIndex;
 
-	TArray<FRDGBuffer*> Buffers;
-	TArray<FRDGBufferUAV*> BufferUAVs;
+	FRDGBufferRef Buffer;
+	FRDGBufferUAVRef BufferUAV;
 };
