Optimus: Add single dispatch support to some data interfaces.
In future we want to get to a point where as many data interfaces as possible support single dispatch.
This will allow better GPU utilization. But also it allows other data interfaces such as half edge connectivity to correctly access the full set of vertices.

diff --git a/Engine/Plugins/Animation/DeformerGraph/Source/OptimusCore/Private/DataInterfaces/OptimusDataInterfaceAnimAttribute.h b/Engine/Plugins/Animation/DeformerGraph/Source/OptimusCore/Private/DataInterfaces/OptimusDataInterfaceAnimAttribute.h
--- a/Engine/Plugins/Animation/DeformerGraph/Source/OptimusCore/Private/DataInterfaces/OptimusDataInterfaceAnimAttribute.h
+++ b/Engine/Plugins/Animation/DeformerGraph/Source/OptimusCore/Private/DataInterfaces/OptimusDataInterfaceAnimAttribute.h
@@ -120,7 +120,8 @@ public:
 	
 	//~ Begin UComputeDataInterface Interface
 	TCHAR const* GetClassName() const override { return TEXT("AnimAttribute"); }
-	virtual void GetSupportedInputs(TArray<FShaderFunctionDefinition>& OutFunctions) const override;
+	bool CanSupportUnifiedDispatch() const override { return true; }
+	void GetSupportedInputs(TArray<FShaderFunctionDefinition>& OutFunctions) const override;
 	void GetShaderParameters(TCHAR const* UID, FShaderParametersMetadataBuilder& InOutBuilder, FShaderParametersMetadataAllocations& InOutAllocations) const override;
 	void GetHLSL(FString& OutHLSL, FString const& InDataInterfaceName) const override;
 	void GetStructDeclarations(TSet<FString>& OutStructsSeen, TArray<FString>& OutStructs) const override;
diff --git a/Engine/Plugins/Animation/DeformerGraph/Source/OptimusCore/Private/DataInterfaces/OptimusDataInterfaceDebugDraw.h b/Engine/Plugins/Animation/DeformerGraph/Source/OptimusCore/Private/DataInterfaces/OptimusDataInterfaceDebugDraw.h
--- a/Engine/Plugins/Animation/DeformerGraph/Source/OptimusCore/Private/DataInterfaces/OptimusDataInterfaceDebugDraw.h
+++ b/Engine/Plugins/Animation/DeformerGraph/Source/OptimusCore/Private/DataInterfaces/OptimusDataInterfaceDebugDraw.h
@@ -54,6 +54,7 @@ public:
 	
 	//~ Begin UComputeDataInterface Interface
 	TCHAR const* GetClassName() const override { return TEXT("DebugDraw"); }
+	bool CanSupportUnifiedDispatch() const override { return true; }
 	void GetSupportedInputs(TArray<FShaderFunctionDefinition>& OutFunctions) const override;
 	void GetShaderParameters(TCHAR const* UID, FShaderParametersMetadataBuilder& InOutBuilder, FShaderParametersMetadataAllocations& InOutAllocations) const override;
 	TCHAR const* GetShaderVirtualPath() const override;
diff --git a/Engine/Plugins/Animation/DeformerGraph/Source/OptimusCore/Private/DataInterfaces/OptimusDataInterfaceGraph.h b/Engine/Plugins/Animation/DeformerGraph/Source/OptimusCore/Private/DataInterfaces/OptimusDataInterfaceGraph.h
--- a/Engine/Plugins/Animation/DeformerGraph/Source/OptimusCore/Private/DataInterfaces/OptimusDataInterfaceGraph.h
+++ b/Engine/Plugins/Animation/DeformerGraph/Source/OptimusCore/Private/DataInterfaces/OptimusDataInterfaceGraph.h
@@ -40,6 +40,7 @@ public:
 
 	//~ Begin UComputeDataInterface Interface
 	TCHAR const* GetClassName() const override { return TEXT("Graph"); }
+	bool CanSupportUnifiedDispatch() const override { return true; }
 	void GetSupportedInputs(TArray<FShaderFunctionDefinition>& OutFunctions) const override;
 	void GetShaderParameters(TCHAR const* UID, FShaderParametersMetadataBuilder& InOutBuilder, FShaderParametersMetadataAllocations& InOutAllocations) const override;
 	void GetHLSL(FString& OutHLSL, FString const& InDataInterfaceName) const override;
diff --git a/Engine/Plugins/Animation/DeformerGraph/Source/OptimusCore/Private/DataInterfaces/OptimusDataInterfaceHalfEdge.cpp b/Engine/Plugins/Animation/DeformerGraph/Source/OptimusCore/Private/DataInterfaces/OptimusDataInterfaceHalfEdge.cpp
--- a/Engine/Plugins/Animation/DeformerGraph/Source/OptimusCore/Private/DataInterfaces/OptimusDataInterfaceHalfEdge.cpp
+++ b/Engine/Plugins/Animation/DeformerGraph/Source/OptimusCore/Private/DataInterfaces/OptimusDataInterfaceHalfEdge.cpp
@@ -401,10 +401,10 @@ void FOptimusHalfEdgeDataProviderProxy::GatherDispatchData(FDispatchData const&
 		FSkelMeshRenderSection const& RenderSection = LodRenderData->RenderSections[InvocationIndex];
 
 		FParameters& Parameters = ParameterArray[InvocationIndex];
-		Parameters.NumVertices = RenderSection.NumVertices;
-		Parameters.NumTriangles = RenderSection.NumTriangles;
-		Parameters.IndexBufferStart = RenderSection.BaseIndex;
-		Parameters.InputStreamStart = RenderSection.BaseVertexIndex;
+		Parameters.NumVertices = InDispatchData.bUnifiedDispatch ? LodRenderData->GetNumVertices() : RenderSection.NumVertices;
+		Parameters.NumTriangles = InDispatchData.bUnifiedDispatch ? LodRenderData->GetTotalFaces() : RenderSection.NumTriangles;
+		Parameters.IndexBufferStart = InDispatchData.bUnifiedDispatch ? 0 :RenderSection.BaseIndex;
+		Parameters.InputStreamStart = InDispatchData.bUnifiedDispatch ? 0 : RenderSection.BaseVertexIndex;
 		Parameters.VertexToEdgeBuffer = VertexToEdgeBufferSRV;
 		Parameters.EdgeToTwinEdgeBuffer = EdgeToTwinEdgeBufferSRV;
 	}
diff --git a/Engine/Plugins/Animation/DeformerGraph/Source/OptimusCore/Private/DataInterfaces/OptimusDataInterfaceHalfEdge.h b/Engine/Plugins/Animation/DeformerGraph/Source/OptimusCore/Private/DataInterfaces/OptimusDataInterfaceHalfEdge.h
--- a/Engine/Plugins/Animation/DeformerGraph/Source/OptimusCore/Private/DataInterfaces/OptimusDataInterfaceHalfEdge.h
+++ b/Engine/Plugins/Animation/DeformerGraph/Source/OptimusCore/Private/DataInterfaces/OptimusDataInterfaceHalfEdge.h
@@ -28,6 +28,7 @@ public:
 	
 	//~ Begin UComputeDataInterface Interface
 	TCHAR const* GetClassName() const override { return TEXT("HalfEdge"); }
+	bool CanSupportUnifiedDispatch() const override { return true; }
 	void GetSupportedInputs(TArray<FShaderFunctionDefinition>& OutFunctions) const override;
 	void GetShaderParameters(TCHAR const* UID, FShaderParametersMetadataBuilder& InOutBuilder, FShaderParametersMetadataAllocations& InOutAllocations) const override;
 	TCHAR const* GetShaderVirtualPath() const override;
diff --git a/Engine/Plugins/Animation/DeformerGraph/Source/OptimusCore/Private/DataInterfaces/OptimusDataInterfaceRawBuffer.cpp b/Engine/Plugins/Animation/DeformerGraph/Source/OptimusCore/Private/DataInterfaces/OptimusDataInterfaceRawBuffer.cpp
--- a/Engine/Plugins/Animation/DeformerGraph/Source/OptimusCore/Private/DataInterfaces/OptimusDataInterfaceRawBuffer.cpp
+++ b/Engine/Plugins/Animation/DeformerGraph/Source/OptimusCore/Private/DataInterfaces/OptimusDataInterfaceRawBuffer.cpp
@@ -353,8 +353,8 @@ void FOptimusTransientBufferDataProviderProxy::GatherDispatchData(FDispatchData
 	for (int32 InvocationIndex = 0, StartOffset = 0; InvocationIndex < ParameterArray.Num(); ++InvocationIndex)
 	{
 		FParameters& Parameters = ParameterArray[InvocationIndex];
-		Parameters.StartOffset = StartOffset;
-		Parameters.BufferSize = InvocationElementCounts[InvocationIndex];
+		Parameters.StartOffset = InDispatchData.bUnifiedDispatch ? 0 : StartOffset;
+		Parameters.BufferSize = InDispatchData.bUnifiedDispatch ? TotalElementCount : InvocationElementCounts[InvocationIndex];
 		Parameters.BufferSRV = BufferSRV;
 		Parameters.BufferUAV = BufferUAV;
 		
@@ -417,8 +417,8 @@ void FOptimusPersistentBufferDataProviderProxy::GatherDispatchData(FDispatchData
 	for (int32 InvocationIndex = 0, StartOffset = 0; InvocationIndex < ParameterArray.Num(); ++InvocationIndex)
 	{
 		FParameters& Parameters = ParameterArray[InvocationIndex];
-		Parameters.StartOffset = StartOffset;
-		Parameters.BufferSize = InvocationElementCounts[InvocationIndex];
+		Parameters.StartOffset = InDispatchData.bUnifiedDispatch ? 0 : StartOffset;
+		Parameters.BufferSize = InDispatchData.bUnifiedDispatch ? TotalElementCount : InvocationElementCounts[InvocationIndex];
 		Parameters.BufferUAV = BufferUAV;
 
 		StartOffset += InvocationElementCounts[InvocationIndex];
diff --git a/Engine/Plugins/Animation/DeformerGraph/Source/OptimusCore/Private/DataInterfaces/OptimusDataInterfaceRawBuffer.h b/Engine/Plugins/Animation/DeformerGraph/Source/OptimusCore/Private/DataInterfaces/OptimusDataInterfaceRawBuffer.h
--- a/Engine/Plugins/Animation/DeformerGraph/Source/OptimusCore/Private/DataInterfaces/OptimusDataInterfaceRawBuffer.h
+++ b/Engine/Plugins/Animation/DeformerGraph/Source/OptimusCore/Private/DataInterfaces/OptimusDataInterfaceRawBuffer.h
@@ -105,6 +105,7 @@ public:
 
 	//~ Begin UComputeDataInterface Interface
 	TCHAR const* GetClassName() const override { return TEXT("TransientBuffer"); }
+	bool CanSupportUnifiedDispatch() const override { return true; }
 	void GetShaderParameters(TCHAR const* UID, FShaderParametersMetadataBuilder& InOutBuilder, FShaderParametersMetadataAllocations& InOutAllocations) const override;
 	UComputeDataProvider* CreateDataProvider(TObjectPtr<UObject> InBinding, uint64 InInputMask, uint64 InOutputMask) const override;
 	//~ End UComputeDataInterface Interface
@@ -125,6 +126,7 @@ public:
 
 	//~ Begin UComputeDataInterface Interface
 	TCHAR const* GetClassName() const override { return TEXT("PersistentBuffer"); }
+	bool CanSupportUnifiedDispatch() const override { return true; }
 	void GetShaderParameters(TCHAR const* UID, FShaderParametersMetadataBuilder& InOutBuilder, FShaderParametersMetadataAllocations& InOutAllocations) const override;
 	UComputeDataProvider* CreateDataProvider(TObjectPtr<UObject> InBinding, uint64 InInputMask, uint64 InOutputMask) const override;
 	//~ End UComputeDataInterface Interface
diff --git a/Engine/Plugins/Animation/DeformerGraph/Source/OptimusCore/Private/DataInterfaces/OptimusDataInterfaceScene.h b/Engine/Plugins/Animation/DeformerGraph/Source/OptimusCore/Private/DataInterfaces/OptimusDataInterfaceScene.h
--- a/Engine/Plugins/Animation/DeformerGraph/Source/OptimusCore/Private/DataInterfaces/OptimusDataInterfaceScene.h
+++ b/Engine/Plugins/Animation/DeformerGraph/Source/OptimusCore/Private/DataInterfaces/OptimusDataInterfaceScene.h
@@ -24,6 +24,7 @@ public:
 
 	//~ Begin UComputeDataInterface Interface
 	TCHAR const* GetClassName() const override { return TEXT("Scene"); }
+	bool CanSupportUnifiedDispatch() const override { return true; }
 	void GetSupportedInputs(TArray<FShaderFunctionDefinition>& OutFunctions) const override;
 	void GetShaderParameters(TCHAR const* UID, FShaderParametersMetadataBuilder& InOutBuilder, FShaderParametersMetadataAllocations& InOutAllocations) const override;
 	TCHAR const* GetShaderVirtualPath() const override;
diff --git a/Engine/Plugins/Animation/DeformerGraph/Source/OptimusCore/Private/DataInterfaces/OptimusDataInterfaceSkinnedMesh.cpp b/Engine/Plugins/Animation/DeformerGraph/Source/OptimusCore/Private/DataInterfaces/OptimusDataInterfaceSkinnedMesh.cpp
--- a/Engine/Plugins/Animation/DeformerGraph/Source/OptimusCore/Private/DataInterfaces/OptimusDataInterfaceSkinnedMesh.cpp
+++ b/Engine/Plugins/Animation/DeformerGraph/Source/OptimusCore/Private/DataInterfaces/OptimusDataInterfaceSkinnedMesh.cpp
@@ -190,11 +190,11 @@ void FOptimusSkinnedMeshDataProviderProxy::GatherDispatchData(FDispatchData cons
 		FRHIShaderResourceView* MeshColorBufferSRV = LodRenderData->StaticVertexBuffers.ColorVertexBuffer.GetColorComponentsSRV();
 
 		FParameters& Parameters = ParameterArray[InvocationIndex];
-		Parameters.NumVertices = RenderSection.NumVertices;
-		Parameters.NumTriangles = RenderSection.NumTriangles;
+		Parameters.NumVertices = InDispatchData.bUnifiedDispatch ? LodRenderData->GetNumVertices() : RenderSection.NumVertices;
+		Parameters.NumTriangles = InDispatchData.bUnifiedDispatch ? LodRenderData->GetTotalFaces() : RenderSection.NumTriangles;
 		Parameters.NumUVChannels = LodRenderData->StaticVertexBuffers.StaticMeshVertexBuffer.GetNumTexCoords();
-		Parameters.IndexBufferStart = RenderSection.BaseIndex;
-		Parameters.InputStreamStart = RenderSection.BaseVertexIndex;
+		Parameters.IndexBufferStart = InDispatchData.bUnifiedDispatch ? 0 : RenderSection.BaseIndex;
+		Parameters.InputStreamStart = InDispatchData.bUnifiedDispatch ? 0 : RenderSection.BaseVertexIndex;
 		Parameters.IndexBuffer = IndexBufferSRV != nullptr ? IndexBufferSRV : NullSRVBinding;
 		Parameters.PositionInputBuffer = MeshVertexBufferSRV != nullptr ? MeshVertexBufferSRV : NullSRVBinding;
 		Parameters.TangentInputBuffer = MeshTangentBufferSRV != nullptr ? MeshTangentBufferSRV : NullSRVBinding;
diff --git a/Engine/Plugins/Animation/DeformerGraph/Source/OptimusCore/Private/DataInterfaces/OptimusDataInterfaceSkinnedMesh.h b/Engine/Plugins/Animation/DeformerGraph/Source/OptimusCore/Private/DataInterfaces/OptimusDataInterfaceSkinnedMesh.h
--- a/Engine/Plugins/Animation/DeformerGraph/Source/OptimusCore/Private/DataInterfaces/OptimusDataInterfaceSkinnedMesh.h
+++ b/Engine/Plugins/Animation/DeformerGraph/Source/OptimusCore/Private/DataInterfaces/OptimusDataInterfaceSkinnedMesh.h
@@ -26,6 +26,7 @@ public:
 	
 	//~ Begin UComputeDataInterface Interface
 	TCHAR const* GetClassName() const override { return TEXT("SkinnedMesh"); }
+	bool CanSupportUnifiedDispatch() const override { return true; }
 	void GetSupportedInputs(TArray<FShaderFunctionDefinition>& OutFunctions) const override;
 	void GetShaderParameters(TCHAR const* UID, FShaderParametersMetadataBuilder& InOutBuilder, FShaderParametersMetadataAllocations& InOutAllocations) const override;
 	TCHAR const* GetShaderVirtualPath() const override;
diff --git a/Engine/Plugins/Animation/DeformerGraph/Source/OptimusCore/Private/DataInterfaces/OptimusDataInterfaceSkinnedMeshExec.cpp b/Engine/Plugins/Animation/DeformerGraph/Source/OptimusCore/Private/DataInterfaces/OptimusDataInterfaceSkinnedMeshExec.cpp
--- a/Engine/Plugins/Animation/DeformerGraph/Source/OptimusCore/Private/DataInterfaces/OptimusDataInterfaceSkinnedMeshExec.cpp
+++ b/Engine/Plugins/Animation/DeformerGraph/Source/OptimusCore/Private/DataInterfaces/OptimusDataInterfaceSkinnedMeshExec.cpp
@@ -149,8 +149,17 @@ void FOptimusSkinnedMeshExecDataProviderProxy::GatherDispatchData(FDispatchData
 	const TStridedView<FParameters> ParameterArray = MakeStridedParameterView<FParameters>(InDispatchData);
 	for (int32 InvocationIndex = 0; InvocationIndex < ParameterArray.Num(); ++InvocationIndex)
 	{
-		FSkelMeshRenderSection const& RenderSection = LodRenderData->RenderSections[InvocationIndex];
-		const int32 NumThreads = Domain == EOptimusSkinnedMeshExecDomain::Vertex ? RenderSection.NumVertices : RenderSection.NumTriangles;
+		int32 NumThreads = 0;
+
+		if (InDispatchData.bUnifiedDispatch)
+		{
+			NumThreads = Domain == EOptimusSkinnedMeshExecDomain::Vertex ? LodRenderData->GetNumVertices() : LodRenderData->GetTotalFaces();
+		}
+		else
+		{
+			FSkelMeshRenderSection const& RenderSection = LodRenderData->RenderSections[InvocationIndex];
+			NumThreads = Domain == EOptimusSkinnedMeshExecDomain::Vertex ? RenderSection.NumVertices : RenderSection.NumTriangles;
+		}
 
 		FParameters& Parameters = ParameterArray[InvocationIndex];
 		Parameters.NumThreads = FIntVector(NumThreads, 1, 1);
diff --git a/Engine/Plugins/Animation/DeformerGraph/Source/OptimusCore/Private/DataInterfaces/OptimusDataInterfaceSkinnedMeshExec.h b/Engine/Plugins/Animation/DeformerGraph/Source/OptimusCore/Private/DataInterfaces/OptimusDataInterfaceSkinnedMeshExec.h
--- a/Engine/Plugins/Animation/DeformerGraph/Source/OptimusCore/Private/DataInterfaces/OptimusDataInterfaceSkinnedMeshExec.h
+++ b/Engine/Plugins/Animation/DeformerGraph/Source/OptimusCore/Private/DataInterfaces/OptimusDataInterfaceSkinnedMeshExec.h
@@ -40,6 +40,7 @@ public:
 	//~ Begin UComputeDataInterface Interface
 	TCHAR const* GetClassName() const override { return TEXT("SkinnedMeshExec"); }
 	bool IsExecutionInterface() const override { return true; }
+	bool CanSupportUnifiedDispatch() const override { return true; }
 	void GetSupportedInputs(TArray<FShaderFunctionDefinition>& OutFunctions) const override;
 	void GetShaderParameters(TCHAR const* UID, FShaderParametersMetadataBuilder& InOutBuilder, FShaderParametersMetadataAllocations& InOutAllocations) const override;
 	TCHAR const* GetShaderVirtualPath() const override;
diff --git a/Engine/Plugins/Animation/DeformerGraph/Source/OptimusCore/Private/DataInterfaces/OptimusDataInterfaceSkinnedMeshWrite.cpp b/Engine/Plugins/Animation/DeformerGraph/Source/OptimusCore/Private/DataInterfaces/OptimusDataInterfaceSkinnedMeshWrite.cpp
--- a/Engine/Plugins/Animation/DeformerGraph/Source/OptimusCore/Private/DataInterfaces/OptimusDataInterfaceSkinnedMeshWrite.cpp
+++ b/Engine/Plugins/Animation/DeformerGraph/Source/OptimusCore/Private/DataInterfaces/OptimusDataInterfaceSkinnedMeshWrite.cpp
@@ -207,8 +207,8 @@ void FOptimusSkinnedMeshWriteDataProviderProxy::GatherDispatchData(FDispatchData
 		FSkelMeshRenderSection const& RenderSection = LodRenderData->RenderSections[InvocationIndex];
 
 		FParameters& Parameters = ParameterArray[InvocationIndex];
-		Parameters.NumVertices = RenderSection.GetNumVertices();
-		Parameters.OutputStreamStart = RenderSection.GetVertexBufferIndex();
+		Parameters.NumVertices = InDispatchData.bUnifiedDispatch ? LodRenderData->GetNumVertices() : RenderSection.GetNumVertices();
+		Parameters.OutputStreamStart = InDispatchData.bUnifiedDispatch ? 0 : RenderSection.GetVertexBufferIndex();
 		Parameters.PositionBufferUAV = PositionBufferUAV;
 		Parameters.TangentBufferUAV = TangentBufferUAV;
 		Parameters.ColorBufferUAV = ColorBufferUAV;
diff --git a/Engine/Plugins/Animation/DeformerGraph/Source/OptimusCore/Private/DataInterfaces/OptimusDataInterfaceSkinnedMeshWrite.h b/Engine/Plugins/Animation/DeformerGraph/Source/OptimusCore/Private/DataInterfaces/OptimusDataInterfaceSkinnedMeshWrite.h
--- a/Engine/Plugins/Animation/DeformerGraph/Source/OptimusCore/Private/DataInterfaces/OptimusDataInterfaceSkinnedMeshWrite.h
+++ b/Engine/Plugins/Animation/DeformerGraph/Source/OptimusCore/Private/DataInterfaces/OptimusDataInterfaceSkinnedMeshWrite.h
@@ -29,6 +29,7 @@ public:
 	
 	//~ Begin UComputeDataInterface Interface
 	TCHAR const* GetClassName() const override { return TEXT("SkinnedMeshWrite"); }
+	bool CanSupportUnifiedDispatch() const override { return true; }
 	void GetSupportedInputs(TArray<FShaderFunctionDefinition>& OutFunctions) const override;
 	void GetSupportedOutputs(TArray<FShaderFunctionDefinition>& OutFunctions) const override;
 	void GetShaderParameters(TCHAR const* UID, FShaderParametersMetadataBuilder& InOutBuilder, FShaderParametersMetadataAllocations& InOutAllocations) const override;
diff --git a/Engine/Plugins/Animation/MLDeformer/MLDeformerFramework/Source/MLDeformerFramework/Private/MLDeformerGraphDebugDataInterface.cpp b/Engine/Plugins/Animation/MLDeformer/MLDeformerFramework/Source/MLDeformerFramework/Private/MLDeformerGraphDebugDataInterface.cpp
--- a/Engine/Plugins/Animation/MLDeformer/MLDeformerFramework/Source/MLDeformerFramework/Private/MLDeformerGraphDebugDataInterface.cpp
+++ b/Engine/Plugins/Animation/MLDeformer/MLDeformerFramework/Source/MLDeformerFramework/Private/MLDeformerGraphDebugDataInterface.cpp
@@ -212,8 +212,8 @@ namespace UE::MLDeformer
 		{
 			const FSkelMeshRenderSection& RenderSection = LodRenderData->RenderSections[InvocationIndex];
 			FMLDeformerGraphDebugDataInterfaceParameters& Parameters = ParameterArray[InvocationIndex];
-			Parameters.NumVertices = 0;
-			Parameters.InputStreamStart = RenderSection.BaseVertexIndex;
+			Parameters.NumVertices = InDispatchData.bUnifiedDispatch ? LodRenderData->GetNumVertices() : RenderSection.GetNumVertices();
+			Parameters.InputStreamStart = InDispatchData.bUnifiedDispatch ? 0 : RenderSection.BaseVertexIndex;
 			Parameters.HeatMapMode = HeatMapMode;
 			Parameters.HeatMapMax = HeatMapMax;
 			Parameters.GroundTruthLerp = GroundTruthLerp;
diff --git a/Engine/Plugins/Animation/MLDeformer/MLDeformerFramework/Source/MLDeformerFramework/Public/MLDeformerGraphDebugDataInterface.h b/Engine/Plugins/Animation/MLDeformer/MLDeformerFramework/Source/MLDeformerFramework/Public/MLDeformerGraphDebugDataInterface.h
--- a/Engine/Plugins/Animation/MLDeformer/MLDeformerFramework/Source/MLDeformerFramework/Public/MLDeformerGraphDebugDataInterface.h
+++ b/Engine/Plugins/Animation/MLDeformer/MLDeformerFramework/Source/MLDeformerFramework/Public/MLDeformerGraphDebugDataInterface.h
@@ -34,6 +34,7 @@ public:
 
 	// UComputeDataInterface overrides.
 	TCHAR const* GetClassName() const override { return TEXT("MLDeformerGraphDebugData"); }
+	virtual bool CanSupportUnifiedDispatch() const override { return true; }
 	virtual void GetSupportedInputs(TArray<FShaderFunctionDefinition>& OutFunctions) const override;
 	virtual void GetShaderParameters(TCHAR const* UID, FShaderParametersMetadataBuilder& OutBuilder, FShaderParametersMetadataAllocations& InOutAllocations) const override;
 	virtual TCHAR const* GetShaderVirtualPath() const override;
diff --git a/Engine/Plugins/Animation/MLDeformer/VertexDeltaModel/Source/VertexDeltaModel/Private/VertexDeltaGraphDataInterface.cpp b/Engine/Plugins/Animation/MLDeformer/VertexDeltaModel/Source/VertexDeltaModel/Private/VertexDeltaGraphDataInterface.cpp
--- a/Engine/Plugins/Animation/MLDeformer/VertexDeltaModel/Source/VertexDeltaModel/Private/VertexDeltaGraphDataInterface.cpp
+++ b/Engine/Plugins/Animation/MLDeformer/VertexDeltaModel/Source/VertexDeltaModel/Private/VertexDeltaGraphDataInterface.cpp
@@ -148,8 +148,8 @@ namespace UE::VertexDeltaModel
 		{
 			const FSkelMeshRenderSection& RenderSection = LodRenderData->RenderSections[InvocationIndex];
 			FVertexDeltaGraphDataInterfaceParameters& Parameters = ParameterArray[InvocationIndex];
-			Parameters.NumVertices = 0;
-			Parameters.InputStreamStart = RenderSection.BaseVertexIndex;
+			Parameters.NumVertices = InDispatchData.bUnifiedDispatch ? LodRenderData->GetNumVertices() : RenderSection.GetNumVertices();
+			Parameters.InputStreamStart = InDispatchData.bUnifiedDispatch ? 0 : RenderSection.BaseVertexIndex;
 			Parameters.Weight = Weight;
 			Parameters.PositionDeltaBuffer = BufferSRV;
 			Parameters.VertexMapBuffer = VertexMapBufferSRV;
diff --git a/Engine/Plugins/Animation/MLDeformer/VertexDeltaModel/Source/VertexDeltaModel/Private/VertexDeltaGraphDebugDataInterface.cpp b/Engine/Plugins/Animation/MLDeformer/VertexDeltaModel/Source/VertexDeltaModel/Private/VertexDeltaGraphDebugDataInterface.cpp
--- a/Engine/Plugins/Animation/MLDeformer/VertexDeltaModel/Source/VertexDeltaModel/Private/VertexDeltaGraphDebugDataInterface.cpp
+++ b/Engine/Plugins/Animation/MLDeformer/VertexDeltaModel/Source/VertexDeltaModel/Private/VertexDeltaGraphDebugDataInterface.cpp
@@ -212,8 +212,8 @@ namespace UE::VertexDeltaModel
 		{
 			const FSkelMeshRenderSection& RenderSection = LodRenderData->RenderSections[InvocationIndex];
 			FVertexDeltaGraphDebugDataInterfaceParameters& Parameters = ParameterArray[InvocationIndex];
-			Parameters.NumVertices = 0;
-			Parameters.InputStreamStart = RenderSection.BaseVertexIndex;
+			Parameters.NumVertices = InDispatchData.bUnifiedDispatch ? LodRenderData->GetNumVertices() : RenderSection.GetNumVertices();
+			Parameters.InputStreamStart = InDispatchData.bUnifiedDispatch ? 0 : RenderSection.BaseVertexIndex;
 			Parameters.HeatMapMode = HeatMapMode;
 			Parameters.HeatMapMax = HeatMapMax;
 			Parameters.GroundTruthLerp = GroundTruthLerp;
diff --git a/Engine/Plugins/Animation/MLDeformer/VertexDeltaModel/Source/VertexDeltaModel/Public/VertexDeltaGraphDataInterface.h b/Engine/Plugins/Animation/MLDeformer/VertexDeltaModel/Source/VertexDeltaModel/Public/VertexDeltaGraphDataInterface.h
--- a/Engine/Plugins/Animation/MLDeformer/VertexDeltaModel/Source/VertexDeltaModel/Public/VertexDeltaGraphDataInterface.h
+++ b/Engine/Plugins/Animation/MLDeformer/VertexDeltaModel/Source/VertexDeltaModel/Public/VertexDeltaGraphDataInterface.h
@@ -35,6 +35,7 @@ public:
 
 	// UComputeDataInterface overrides.
 	TCHAR const* GetClassName() const override { return TEXT("VertexDeltaModelData"); }
+	virtual bool CanSupportUnifiedDispatch() const override { return true; }
 	virtual void GetSupportedInputs(TArray<FShaderFunctionDefinition>& OutFunctions) const override;
 	virtual void GetShaderParameters(TCHAR const* UID, FShaderParametersMetadataBuilder& InOutBuilder, FShaderParametersMetadataAllocations& InOutAllocations) const override;
 	virtual TCHAR const* GetShaderVirtualPath() const override;
diff --git a/Engine/Plugins/Animation/MLDeformer/VertexDeltaModel/Source/VertexDeltaModel/Public/VertexDeltaGraphDebugDataInterface.h b/Engine/Plugins/Animation/MLDeformer/VertexDeltaModel/Source/VertexDeltaModel/Public/VertexDeltaGraphDebugDataInterface.h
--- a/Engine/Plugins/Animation/MLDeformer/VertexDeltaModel/Source/VertexDeltaModel/Public/VertexDeltaGraphDebugDataInterface.h
+++ b/Engine/Plugins/Animation/MLDeformer/VertexDeltaModel/Source/VertexDeltaModel/Public/VertexDeltaGraphDebugDataInterface.h
@@ -35,6 +35,7 @@ public:
 
 	// UComputeDataInterface overrides.
 	TCHAR const* GetClassName() const override { return TEXT("VertexDeltaModelDebugData"); }
+	virtual bool CanSupportUnifiedDispatch() const override { return true; }
 	virtual void GetSupportedInputs(TArray<FShaderFunctionDefinition>& OutFunctions) const override;
 	virtual void GetShaderParameters(TCHAR const* UID, FShaderParametersMetadataBuilder& OutBuilder, FShaderParametersMetadataAllocations& InOutAllocations) const override;
 	virtual TCHAR const* GetShaderVirtualPath() const override;
