Deformer Graph: Ensure we #include all prerequisites for managing skeleton data on the GPU.

#rnx

diff --git a/Engine/Plugins/Animation/DeformerGraph/Shaders/Private/DataInterfaceSkeleton.ush b/Engine/Plugins/Animation/DeformerGraph/Shaders/Private/DataInterfaceSkeleton.ush
--- a/Engine/Plugins/Animation/DeformerGraph/Shaders/Private/DataInterfaceSkeleton.ush
+++ b/Engine/Plugins/Animation/DeformerGraph/Shaders/Private/DataInterfaceSkeleton.ush
@@ -1,5 +1,7 @@
 // Copyright Epic Games, Inc. All Rights Reserved.
 
+#include "/Engine/Private/GpuSkinCommon.ush"
+
 uint {DataInterfaceName}_NumVertices;
 uint {DataInterfaceName}_InputStreamStart;
 uint {DataInterfaceName}_NumBoneInfluences;
