Optimus: Editor fixes related to adder pin when used with multiple component bindings

-single adder pin per kernel
-upon adder pin connection, it will try to find a pin group automatically based on pin's component source
-if there are multiple choices, a menu pops up allowing users to pick a preferred group

diff --git a/Engine/Plugins/Animation/DeformerGraph/Source/OptimusCore/Private/Actions/OptimusComponentBindingActions.cpp b/Engine/Plugins/Animation/DeformerGraph/Source/OptimusCore/Private/Actions/OptimusComponentBindingActions.cpp
--- a/Engine/Plugins/Animation/DeformerGraph/Source/OptimusCore/Private/Actions/OptimusComponentBindingActions.cpp
+++ b/Engine/Plugins/Animation/DeformerGraph/Source/OptimusCore/Private/Actions/OptimusComponentBindingActions.cpp
@@ -74,7 +74,7 @@ bool FOptimusComponentBindingAction_AddBinding::Undo(IOptimusPathResolver* InRoo
 	}
 
 	UOptimusDeformer* Deformer = Cast<UOptimusDeformer>(InRoot);
-	return Deformer->RemoveComponentBinding(Binding);
+	return Deformer->RemoveComponentBindingDirect(Binding);
 }
 
 
diff --git a/Engine/Plugins/Animation/DeformerGraph/Source/OptimusCore/Private/Actions/OptimusNodeActions.cpp b/Engine/Plugins/Animation/DeformerGraph/Source/OptimusCore/Private/Actions/OptimusNodeActions.cpp
--- a/Engine/Plugins/Animation/DeformerGraph/Source/OptimusCore/Private/Actions/OptimusNodeActions.cpp
+++ b/Engine/Plugins/Animation/DeformerGraph/Source/OptimusCore/Private/Actions/OptimusNodeActions.cpp
@@ -261,62 +261,6 @@ bool FOptimusNodeAction_SetPinDataDomain::SetPinDataDomain(
 }
 
 
-FOptimusNodeAction_ConnectAdderPin::FOptimusNodeAction_ConnectAdderPin(
-	IOptimusNodeAdderPinProvider* InAdderPinProvider,
-	UOptimusNodePin* InSourcePin,
-	FName InNewPinName)
-{
-	UOptimusNode* Node = Cast<UOptimusNode>(InAdderPinProvider);
-	
-	if (ensure(Node))
-	{
-		NodePath = Node->GetNodePath();
-	}
-
-	SourcePinPath = InSourcePin->GetPinPath();
-
-	NewPinName = InNewPinName;
-}
-
-bool FOptimusNodeAction_ConnectAdderPin::Do(IOptimusPathResolver* InRoot)
-{
-	UOptimusNode* Node = InRoot->ResolveNodePath(NodePath);
-	IOptimusNodeAdderPinProvider* AdderPinProvider = Cast<IOptimusNodeAdderPinProvider>(Node);
-	
-	if (!ensure(AdderPinProvider))
-	{
-		return false;
-	}
-
-	UOptimusNodePin* SourcePin = InRoot->ResolvePinPath(SourcePinPath);
-	
-	const UOptimusNodePin* NewPin = AdderPinProvider->TryAddPinFromPin(SourcePin, NewPinName);
-
-	if (!NewPin)
-	{
-		return false;
-	}
-	
-	NewPinPath = NewPin->GetPinPath();
-
-	return true;
-}
-
-bool FOptimusNodeAction_ConnectAdderPin::Undo(IOptimusPathResolver* InRoot)
-{
-	UOptimusNode* Node = InRoot->ResolveNodePath(NodePath);
-	IOptimusNodeAdderPinProvider* AdderPinProvider = Cast<IOptimusNodeAdderPinProvider>(Node);
-	
-	if (!ensure(AdderPinProvider))
-	{
-		return false;
-	}
-
-	UOptimusNodePin* NewPin = InRoot->ResolvePinPath(NewPinPath);
-
-	return AdderPinProvider->RemoveAddedPin(NewPin);
-}
-
 FOptimusNodeAction_AddRemovePin::FOptimusNodeAction_AddRemovePin(
 	UOptimusNode* InNode,
 	FName InName,
diff --git a/Engine/Plugins/Animation/DeformerGraph/Source/OptimusCore/Private/Actions/OptimusNodeActions.h b/Engine/Plugins/Animation/DeformerGraph/Source/OptimusCore/Private/Actions/OptimusNodeActions.h
--- a/Engine/Plugins/Animation/DeformerGraph/Source/OptimusCore/Private/Actions/OptimusNodeActions.h
+++ b/Engine/Plugins/Animation/DeformerGraph/Source/OptimusCore/Private/Actions/OptimusNodeActions.h
@@ -203,36 +203,6 @@ private:
 	FOptimusDataDomain OldDataDomain;
 };
 
-USTRUCT()
-struct FOptimusNodeAction_ConnectAdderPin
-	: public FOptimusAction
-{
-	GENERATED_BODY()
-	
-	FOptimusNodeAction_ConnectAdderPin() = default;
-
-	FOptimusNodeAction_ConnectAdderPin(
-		IOptimusNodeAdderPinProvider* InAdderPinProvider,
-		UOptimusNodePin* InSourcePin,
-		FName InNewPinName);
-
-
-protected:
-	bool Do(IOptimusPathResolver* InRoot) override; 
-	bool Undo(IOptimusPathResolver* InRoot) override;
-
-private:
-	FString NodePath;
-
-	FString SourcePinPath;
-
-	FName NewPinName;
-
-	// New Pin Path will be updated to the actual path after Do()
-	FString NewPinPath;
-};
-
-
 USTRUCT()
 struct FOptimusNodeAction_AddRemovePin :
 	public FOptimusAction
diff --git a/Engine/Plugins/Animation/DeformerGraph/Source/OptimusCore/Private/Actions/OptimusNodeGraphActions.cpp b/Engine/Plugins/Animation/DeformerGraph/Source/OptimusCore/Private/Actions/OptimusNodeGraphActions.cpp
--- a/Engine/Plugins/Animation/DeformerGraph/Source/OptimusCore/Private/Actions/OptimusNodeGraphActions.cpp
+++ b/Engine/Plugins/Animation/DeformerGraph/Source/OptimusCore/Private/Actions/OptimusNodeGraphActions.cpp
@@ -618,6 +618,127 @@ FOptimusNodeGraphAction_RemoveLink::FOptimusNodeGraphAction_RemoveLink(
 	SetTitlef(TEXT("Remove Link"));
 }
 
+FOptimusNodeGraphAction_ConnectAdderPin::FOptimusNodeGraphAction_ConnectAdderPin(
+	IOptimusNodeAdderPinProvider* InAdderPinProvider,
+	UOptimusNodePin* InPreferredTargetParentPin,
+	UOptimusNodePin* InSourcePin,
+	bool bInShouldLink
+	) :
+	bShouldLink(bInShouldLink)
+{
+	if (bInShouldLink)
+	{
+		SetTitlef(TEXT("Add Pin & Link"));
+	}
+	else
+	{
+		SetTitlef(TEXT("Add Pin From Pin"));
+	}
+	
+	UOptimusNode* Node = Cast<UOptimusNode>(InAdderPinProvider);
+	
+	if (ensure(Node))
+	{
+		NodePath = Node->GetNodePath();
+	}
+
+	if (InPreferredTargetParentPin)
+	{
+		PreferredParentPinPath = InPreferredTargetParentPin->GetPinPath();
+	}
+	
+	SourcePinPath = InSourcePin->GetPinPath();
+}
+
+bool FOptimusNodeGraphAction_ConnectAdderPin::Do(IOptimusPathResolver* InRoot)
+{
+	UOptimusNode* Node = InRoot->ResolveNodePath(NodePath);
+	IOptimusNodeAdderPinProvider* AdderPinProvider = Cast<IOptimusNodeAdderPinProvider>(Node);
+	
+	if (!ensure(AdderPinProvider))
+	{
+		return false;
+	}
+
+	UOptimusNodePin* PreferredTargetParentPin = InRoot->ResolvePinPath(PreferredParentPinPath);
+
+	UOptimusNodePin* SourcePin = InRoot->ResolvePinPath(SourcePinPath);
+
+	TArray<UOptimusNodePin*> AddedPins = AdderPinProvider->TryAddPinFromPin(PreferredTargetParentPin, SourcePin);
+
+	if (AddedPins.Num() == 0)
+	{
+		return false;
+	}
+
+	AddedPinPaths.Reset();
+	Algo::Transform(AddedPins, AddedPinPaths,
+		[](UOptimusNodePin* InPin)
+		{
+			return InPin->GetPinPath();
+		});
+
+	UOptimusNodePin* AddedPin = AddedPins.Last();
+	
+	if (!ensure(!AddedPin->IsGroupingPin()))
+	{
+		AdderPinProvider->RemoveAddedPins(AddedPins);
+		return false;
+	}
+
+	if (bShouldLink)
+	{
+		if (SourcePin->GetDirection() == EOptimusNodePinDirection::Output)
+		{
+			NodeOutputPinPath = SourcePin->GetPinPath();
+			NodeInputPinPath = AddedPin->GetPinPath();
+		}
+		else
+		{
+			
+			NodeOutputPinPath = AddedPin->GetPinPath();
+			NodeInputPinPath = SourcePin->GetPinPath();
+		}
+		
+		if (!ensure(AddLink(InRoot)))
+		{
+			AdderPinProvider->RemoveAddedPins(AddedPins);
+			return false;
+		}
+	}
+	
+	return true;
+}
+
+bool FOptimusNodeGraphAction_ConnectAdderPin::Undo(IOptimusPathResolver* InRoot)
+{
+	UOptimusNode* Node = InRoot->ResolveNodePath(NodePath);
+	IOptimusNodeAdderPinProvider* AdderPinProvider = Cast<IOptimusNodeAdderPinProvider>(Node);
+	
+	if (!ensure(AdderPinProvider))
+	{
+		return false;
+	}
+
+	if (bShouldLink)
+	{
+		if (!ensure(RemoveLink(InRoot)))
+		{
+			return false;
+		}
+	}
+
+	TArray<UOptimusNodePin*> AddedPins;
+
+	Algo::Transform(AddedPinPaths, AddedPins,
+		[InRoot](const FString& InPinPath)
+		{
+			return InRoot->ResolvePinPath(InPinPath);
+		});
+
+	return AdderPinProvider->RemoveAddedPins(AddedPins);
+}
+
 
 FOptimusNodeGraphAction_PackageKernelFunction::FOptimusNodeGraphAction_PackageKernelFunction(
 	UOptimusNode_CustomComputeKernel* InKernelNode,
diff --git a/Engine/Plugins/Animation/DeformerGraph/Source/OptimusCore/Private/Actions/OptimusNodeGraphActions.h b/Engine/Plugins/Animation/DeformerGraph/Source/OptimusCore/Private/Actions/OptimusNodeGraphActions.h
--- a/Engine/Plugins/Animation/DeformerGraph/Source/OptimusCore/Private/Actions/OptimusNodeGraphActions.h
+++ b/Engine/Plugins/Animation/DeformerGraph/Source/OptimusCore/Private/Actions/OptimusNodeGraphActions.h
@@ -11,6 +11,7 @@
 
 #include "OptimusNodeGraphActions.generated.h"
 
+class IOptimusNodeAdderPinProvider;
 class UOptimusNode_ComputeKernelFunction;
 class UOptimusNode_CustomComputeKernel;
 enum class EOptimusNodeGraphType;
@@ -355,6 +356,36 @@ protected:
 	bool Undo(IOptimusPathResolver* InRoot) override { return AddLink(InRoot); }
 };
 
+USTRUCT()
+struct FOptimusNodeGraphAction_ConnectAdderPin
+	: public FOptimusNodeGraphAction_AddRemoveLink
+{
+	GENERATED_BODY()
+	
+	FOptimusNodeGraphAction_ConnectAdderPin() = default;
+
+	FOptimusNodeGraphAction_ConnectAdderPin(
+		IOptimusNodeAdderPinProvider* InAdderPinProvider,
+		UOptimusNodePin* InPreferredTargetParentPin,
+		UOptimusNodePin* InSourcePin,
+		bool bInShouldLink
+		);
+
+protected:
+	bool Do(IOptimusPathResolver* InRoot) override; 
+	bool Undo(IOptimusPathResolver* InRoot) override;
+
+private:
+	FString NodePath;
+
+	FString PreferredParentPinPath;
+	
+	FString SourcePinPath;
+
+	TArray<FString> AddedPinPaths;
+
+	bool bShouldLink;
+};
 
 USTRUCT()
 struct FOptimusNodeGraphAction_PackageKernelFunction :
diff --git a/Engine/Plugins/Animation/DeformerGraph/Source/OptimusCore/Private/Nodes/OptimusNode_CustomComputeKernel.cpp b/Engine/Plugins/Animation/DeformerGraph/Source/OptimusCore/Private/Nodes/OptimusNode_CustomComputeKernel.cpp
--- a/Engine/Plugins/Animation/DeformerGraph/Source/OptimusCore/Private/Nodes/OptimusNode_CustomComputeKernel.cpp
+++ b/Engine/Plugins/Animation/DeformerGraph/Source/OptimusCore/Private/Nodes/OptimusNode_CustomComputeKernel.cpp
@@ -159,61 +159,253 @@ bool UOptimusNode_CustomComputeKernel::CanAddPinFromPin(
 	return true;
 }
 
+TArray<UOptimusNodePin*> UOptimusNode_CustomComputeKernel::GetTargetParentPins(const UOptimusNodePin* InSourcePin) const
+{
+	TArray<UOptimusNodePin*> TargetParentPins;
+	
+	if (InSourcePin->GetDirection() == EOptimusNodePinDirection::Input)
+	{
+		return {};
+	}
+	
+	TArray<UOptimusNodePin*> GroupPinsToCheck;
+	// Nullptr represents the primary group
+	GroupPinsToCheck.Add(nullptr);
+
+	for (UOptimusNodePin* GroupPin: GetPins())
+	{
+		if (IsGroupingInputPin(GroupPin))
+		{
+			GroupPinsToCheck.Add(GroupPin);
+		}
+	}
+	
+	TSet<UOptimusComponentSourceBinding*> SourceComponentBindings = InSourcePin->GetComponentSourceBindings();
+	
+	for (UOptimusNodePin* GroupPin : GroupPinsToCheck)
+	{
+		TArray<const UOptimusNodePin*> PinsPerGroup;
+		
+		if (GroupPin)
+		{
+			PinsPerGroup = GroupPin->GetSubPins();
+		}
+		else
+		{
+			PinsPerGroup = GetPrimaryGroupInputPins();
+		}
+		
+		TSet<UOptimusComponentSourceBinding*> ComponentBindingsForGroup;
+		for (const UOptimusNodePin* Pin : PinsPerGroup)
+		{
+			ComponentBindingsForGroup.Append(Pin->GetComponentSourceBindings());
+		}
 
-UOptimusNodePin* UOptimusNode_CustomComputeKernel::TryAddPinFromPin(
-	UOptimusNodePin* InSourcePin,
-	FName InNewPinName
+		TSet<UOptimusComponentSourceBinding*> MatchingBindings = ComponentBindingsForGroup.Intersect(SourceComponentBindings);
+	
+		if (MatchingBindings.Num() > 0 ||
+			SourceComponentBindings.Num() == 0 || // Source-less pin can be added to any group
+			ComponentBindingsForGroup.Num() == 0) // Source-less group can be a target for any pin 
+		{
+			TargetParentPins.Add(GroupPin);
+		}
+	}
+
+	return TargetParentPins;
+}
+
+
+TArray<UOptimusNodePin*> UOptimusNode_CustomComputeKernel::TryAddPinFromPin(
+	UOptimusNodePin* InPreferredTargetParentPin,
+	UOptimusNodePin* InSourcePin
 	) 
 {
+	TArray<UOptimusNodePin*> AddedPins;
+	
 	const EOptimusNodePinDirection NewPinDirection =
 		InSourcePin->GetDirection() == EOptimusNodePinDirection::Input ?
 				EOptimusNodePinDirection::Output : EOptimusNodePinDirection::Input;
-	
-	TArray<FOptimusParameterBinding>& BindingArray = 
-		InSourcePin->GetDirection() == EOptimusNodePinDirection::Input ?
-			OutputBindingArray.InnerArray : InputBindingArray.InnerArray;
 
+	TArray<FOptimusParameterBinding>* BindingArray = nullptr;
+	UOptimusNodePin* ParentPin = nullptr;
+	UOptimusNodePin* BeforePin = nullptr;
+
+	if (NewPinDirection == EOptimusNodePinDirection::Output)
+	{
+		BindingArray = &OutputBindingArray.InnerArray;
+	}
+	else if (NewPinDirection == EOptimusNodePinDirection::Input)
+	{
+		if (InPreferredTargetParentPin)
+		{
+			// Making sure we are connecting from the input side of the kernel
+			check(InSourcePin->GetDirection() == EOptimusNodePinDirection::Output);
+		
+			FOptimusSecondaryInputBindingsGroup* Group = SecondaryInputBindingGroups.FindByPredicate(
+				[InPreferredTargetParentPin](const FOptimusSecondaryInputBindingsGroup& InGroup)
+				{
+					return InGroup.GroupName == InPreferredTargetParentPin->GetFName();
+				});
+
+			if (!ensure(Group))
+			{
+				return {};
+			}
+
+			ParentPin = InPreferredTargetParentPin;
+			BindingArray = &Group->BindingArray.InnerArray;
+		}
+		else
+		{
+			// In case there is no preferred parent pin, we infer the group pin from the source pin
+			// Two cases:
+			//		1. primary group
+			//		2. create a new group
+			TArray<UOptimusNodePin*> ParentPins = GetTargetParentPins(InSourcePin);
+			
+			if (ParentPins.Num() >= 1)
+			{
+				// multiple parent pins but no preferred parent pin implies we want to add to the primary group,
+				// nullptr group pin means the matching group is the primary group, which does not have a grouping pin
+				check(ParentPins[0] == nullptr);
+
+				BindingArray = &InputBindingArray.InnerArray;
+
+				for (UOptimusNodePin* Pin : GetPins())
+				{
+					if (Pin->IsGroupingPin())
+					{
+						BeforePin = Pin;
+						break;
+					}
+				}
+			}
+			else
+			{
+				// This means the source pin does not belong to any existing group, so create a new group
+				TSet<UOptimusComponentSourceBinding*> ComponentBindings = InSourcePin->GetComponentSourceBindings();
+				
+				check(ComponentBindings.Num() != 0);
+
+				UOptimusComponentSourceBinding* ComponentBinding = *ComponentBindings.begin();
+
+				FName GroupName = Optimus::GetUniqueNameForScope(this, ComponentBinding->GetFName());
+
+				SecondaryInputBindingGroups.AddDefaulted_GetRef().GroupName = GroupName;
+				BindingArray = &SecondaryInputBindingGroups.Last().BindingArray.InnerArray;
+				
+				AddedPins.Add(AddGroupingPinDirect(GroupName, EOptimusNodePinDirection::Input));
+				ParentPin = AddedPins.Last();
+			}
+		}	
+	}
+	else
+	{
+		checkNoEntry();
+	}
+
+	if (ParentPin)
+	{
+		ParentPin->SetIsExpanded(true);
+	}
+	
 	FOptimusParameterBinding Binding;
-	Binding.Name = InNewPinName;
+	Binding.Name = GetSanitizedNewPinName(ParentPin, InSourcePin->GetFName());
 	Binding.DataType = {InSourcePin->GetDataType()};
 	Binding.DataDomain = InSourcePin->GetDataDomain();
 
-	Modify();
-	BindingArray.Add(Binding);
+	
+	BindingArray->Add(Binding);
 	UpdatePreamble();
 	
-	UOptimusNodePin* NewPin	= AddPinDirect(InNewPinName, NewPinDirection, Binding.DataDomain, Binding.DataType);
+	AddedPins.Add(AddPinDirect(Binding.Name, NewPinDirection, Binding.DataDomain, Binding.DataType, BeforePin, ParentPin));
 
-	return NewPin;
+	return AddedPins;
 }
 
 
-bool UOptimusNode_CustomComputeKernel::RemoveAddedPin(
-	UOptimusNodePin* InAddedPinToRemove
-	)
+bool UOptimusNode_CustomComputeKernel::RemoveAddedPins(
+	TConstArrayView<UOptimusNodePin*> InAddedPinsToRemove
+)
 {
-	TArray<FOptimusParameterBinding>& BindingArray = 
-		InAddedPinToRemove->GetDirection() == EOptimusNodePinDirection::Input ?
-			InputBindingArray.InnerArray : OutputBindingArray.InnerArray;
+	if (InAddedPinsToRemove.Num() == 0)
+	{
+		return true;
+	}
 
-	Modify();
-	BindingArray.RemoveAll(
-		[InAddedPinToRemove](const FOptimusParameterBinding& Binding)
+	// Currently only two pins can be added in 1 go at the most, group pin + some data pin
+	if (!ensure(InAddedPinsToRemove.Num() <= 2))
+	{
+		return false;
+	}
+	
+	if (InAddedPinsToRemove.Num() == 1)
+	{
+		UOptimusNodePin* PinToRemove = InAddedPinsToRemove[0];
+		
+		TArray<FOptimusParameterBinding>* BindingArray = 
+			PinToRemove->GetDirection() == EOptimusNodePinDirection::Input ?
+				&InputBindingArray.InnerArray : &OutputBindingArray.InnerArray;
+	
+		if (UOptimusNodePin* GroupPin = PinToRemove->GetParentPin())
 		{
-			return Binding.Name == InAddedPinToRemove->GetFName(); 
+			FOptimusSecondaryInputBindingsGroup* Group = SecondaryInputBindingGroups.FindByPredicate(
+			   [GroupPin](const FOptimusSecondaryInputBindingsGroup& InGroup)
+			   {
+				   return InGroup.GroupName == GroupPin->GetFName();
+			   });
+
+			if (!ensure(Group))
+			{
+				return false;
+			}
+		
+			BindingArray = &Group->BindingArray.InnerArray;
+		}
+	
+
+		BindingArray->RemoveAll(
+			[PinToRemove](const FOptimusParameterBinding& Binding)
+			{
+				return Binding.Name == PinToRemove->GetFName(); 
+			});
+		UpdatePreamble();
+		return RemovePinDirect(PinToRemove);	
+	}
+
+	check(InAddedPinsToRemove.Num() == 2);
+
+	UOptimusNodePin* GroupPin = InAddedPinsToRemove[0];
+	check(GroupPin->GetDirection() == EOptimusNodePinDirection::Input);
+	check(GroupPin->GetSubPins().Num() == 1);
+	check(GroupPin->GetSubPins()[0] == InAddedPinsToRemove[1]);
+	
+	SecondaryInputBindingGroups.RemoveAll(
+		[GroupPin](const FOptimusSecondaryInputBindingsGroup& InGroup)
+		{
+			return InGroup.GroupName == GroupPin->GetFName();
 		});
+
 	UpdatePreamble();
-	return RemovePinDirect(InAddedPinToRemove);
+
+	return RemovePinDirect(GroupPin);
 }
 
 
 FName UOptimusNode_CustomComputeKernel::GetSanitizedNewPinName(
+	UOptimusNodePin* InTargetParentPin,
 	FName InPinName
 	)
 {
 	FName NewName = Optimus::GetSanitizedNameForHlsl(InPinName);
 
-	NewName = Optimus::GetUniqueNameForScope(this, NewName);
+	UObject* ScopeObject = this;
+	if (InTargetParentPin)
+	{
+		ScopeObject = InTargetParentPin;
+	}
+	
+	NewName = Optimus::GetUniqueNameForScope(ScopeObject, NewName);
 
 	return NewName;
 }
diff --git a/Engine/Plugins/Animation/DeformerGraph/Source/OptimusCore/Private/Nodes/OptimusNode_CustomComputeKernel.h b/Engine/Plugins/Animation/DeformerGraph/Source/OptimusCore/Private/Nodes/OptimusNode_CustomComputeKernel.h
--- a/Engine/Plugins/Animation/DeformerGraph/Source/OptimusCore/Private/Nodes/OptimusNode_CustomComputeKernel.h
+++ b/Engine/Plugins/Animation/DeformerGraph/Source/OptimusCore/Private/Nodes/OptimusNode_CustomComputeKernel.h
@@ -81,16 +81,19 @@ public:
 		FString* OutReason = nullptr
 		) const override;
 
-	UOptimusNodePin* TryAddPinFromPin(
-		UOptimusNodePin* InSourcePin,
-		FName InNewPinName
+	TArray<UOptimusNodePin*> GetTargetParentPins(const UOptimusNodePin* InSourcePin) const override;
+
+	TArray<UOptimusNodePin*> TryAddPinFromPin(
+		UOptimusNodePin* InPreferredTargetParentPin,
+		UOptimusNodePin* InSourcePin
 		) override;
 	
-	bool RemoveAddedPin(
-		UOptimusNodePin* InAddedPinToRemove
-		) override;
+	bool RemoveAddedPins(
+		TConstArrayView<UOptimusNodePin*> InAddedPinsToRemove
+	) override;
 	
 	FName GetSanitizedNewPinName(
+		UOptimusNodePin* InTargetParentPin,
 		FName InPinName
 		) override;
 
diff --git a/Engine/Plugins/Animation/DeformerGraph/Source/OptimusCore/Private/OptimusNodeGraph.cpp b/Engine/Plugins/Animation/DeformerGraph/Source/OptimusCore/Private/OptimusNodeGraph.cpp
--- a/Engine/Plugins/Animation/DeformerGraph/Source/OptimusCore/Private/OptimusNodeGraph.cpp
+++ b/Engine/Plugins/Animation/DeformerGraph/Source/OptimusCore/Private/OptimusNodeGraph.cpp
@@ -719,7 +719,11 @@ bool UOptimusNodeGraph::RemoveAllLinks(UOptimusNodePin* InNodePin)
 	return GetActionStack()->RunAction(Action);
 }
 
-bool UOptimusNodeGraph::AddPinAndLink(UOptimusNode* InTargetNode, UOptimusNodePin* InSourcePin)
+bool UOptimusNodeGraph::AddPinAndLink(
+	UOptimusNode* InTargetNode,
+	UOptimusNodePin* InPreferredTargetParentPin,
+	UOptimusNodePin* InSourcePin
+	)
 {
 	IOptimusNodeAdderPinProvider *AdderPinProvider = Cast<IOptimusNodeAdderPinProvider>(InTargetNode);
 	
@@ -746,43 +750,21 @@ bool UOptimusNodeGraph::AddPinAndLink(UOptimusNode* InTargetNode, UOptimusNodePi
 		return false;
 	}
 
-	// Add a pin according to adder pin provider and link to it
-	FOptimusCompoundAction *Action = new FOptimusCompoundAction(TEXT("Add Pin"));
-
-	// Create a name for the new pin up front, shared between two sub actions
-	FName PinName = AdderPinProvider->GetSanitizedNewPinName(InSourcePin->GetFName());
-
-	Action->AddSubAction<FOptimusNodeAction_ConnectAdderPin>(AdderPinProvider, InSourcePin, PinName);
-
-
-	FString OutputPinPath = FString::Printf(TEXT("%s.%s"),
-	*InTargetNode->GetNodePath(), *PinName.ToString());
-
-	FString InputPinPath = InSourcePin->GetPinPath();
-
-	if (InSourcePin->GetDirection() == EOptimusNodePinDirection::Output)
-	{
-		Swap(OutputPinPath, InputPinPath);
-	}
-	else
+	// For now lets default to not disconnect exisiting link on the source pin
+	// Additional logic be added in the future if we want to let the user decide
+	bool bShouldLink = true;
+	if (InSourcePin->GetDirection() == EOptimusNodePinDirection::Input)
 	{
-		// Check to see if there's an existing link on the _input_ pin. Output pins can have any
-		// number of connections coming out.
+		// Check to see if there's an existing link on the _input_ pin.
 		TArray<int32> PinLinks = GetAllLinkIndexesToPin(InSourcePin);
-
-		// This shouldn't happen, but we'll cover for it anyway.
-		checkSlow(PinLinks.Num() <= 1);
-
-		for (int32 LinkIndex : PinLinks)
+		if (PinLinks.Num() > 0)
 		{
-			Action->AddSubAction<FOptimusNodeGraphAction_RemoveLink>(Links[LinkIndex]);
+			
+			bShouldLink = false;
 		}
 	}
-	
-	Action->AddSubAction<FOptimusNodeGraphAction_AddLink>(OutputPinPath, InputPinPath);
 
-
-	return GetActionStack()->RunAction(Action);
+	return GetActionStack()->RunAction<FOptimusNodeGraphAction_ConnectAdderPin>(AdderPinProvider, InPreferredTargetParentPin, InSourcePin, bShouldLink);
 }
 
 UOptimusNode* UOptimusNodeGraph::ConvertCustomKernelToFunction(UOptimusNode* InCustomKernel)
diff --git a/Engine/Plugins/Animation/DeformerGraph/Source/OptimusCore/Public/IOptimusNodeAdderPinProvider.h b/Engine/Plugins/Animation/DeformerGraph/Source/OptimusCore/Public/IOptimusNodeAdderPinProvider.h
--- a/Engine/Plugins/Animation/DeformerGraph/Source/OptimusCore/Public/IOptimusNodeAdderPinProvider.h
+++ b/Engine/Plugins/Animation/DeformerGraph/Source/OptimusCore/Public/IOptimusNodeAdderPinProvider.h
@@ -22,11 +22,22 @@ class OPTIMUSCORE_API IOptimusNodeAdderPinProvider
 	GENERATED_BODY()
 
 public:
-	virtual bool CanAddPinFromPin(const UOptimusNodePin* InSourcePin, EOptimusNodePinDirection InNewPinDirection, FString* OutReason = nullptr) const = 0;
-
-	virtual UOptimusNodePin* TryAddPinFromPin(UOptimusNodePin* InSourcePin, FName InNewPinName) = 0;
+	virtual bool CanAddPinFromPin(
+		const UOptimusNodePin* InSourcePin,
+		EOptimusNodePinDirection InNewPinDirection,
+		FString* OutReason = nullptr
+		) const = 0;
+
+	// If there are multiple options for the target parent pin,
+	// a menu should show up for the user to choose one from these options
+	virtual TArray<UOptimusNodePin*> GetTargetParentPins(const UOptimusNodePin* InSourcePin) const = 0;
+	
+	virtual TArray<UOptimusNodePin*> TryAddPinFromPin(
+		UOptimusNodePin* InPreferredTargetParentPin,
+		UOptimusNodePin* InSourcePin
+		) = 0;
 	
-	virtual bool RemoveAddedPin(UOptimusNodePin* InAddedPinToRemove) = 0;
+	virtual bool RemoveAddedPins(TConstArrayView<UOptimusNodePin*> InAddedPinsToRemove) = 0;
 
-	virtual FName GetSanitizedNewPinName(FName InPinName)= 0;
+	virtual FName GetSanitizedNewPinName(UOptimusNodePin* InTargetParentPin, FName InPinName)= 0;
 };
diff --git a/Engine/Plugins/Animation/DeformerGraph/Source/OptimusCore/Public/OptimusNodeGraph.h b/Engine/Plugins/Animation/DeformerGraph/Source/OptimusCore/Public/OptimusNodeGraph.h
--- a/Engine/Plugins/Animation/DeformerGraph/Source/OptimusCore/Public/OptimusNodeGraph.h
+++ b/Engine/Plugins/Animation/DeformerGraph/Source/OptimusCore/Public/OptimusNodeGraph.h
@@ -297,12 +297,14 @@ public:
 	/// Add a new pin to the target node with the type of source pin
 	/// and connect the source pin to the new pin
 	/// @param InTargetNode The node to add the pin to, it has to have an adder pin
+	/// @param InPreferredTargetParentPin The preferred pin to add the new pin to, or null if it can be inferred from source
 	/// @param InSourcePin The pin to create the new pin and to connect to the new pin
 	/// @return True if new pin and the new link is created.
 	bool AddPinAndLink(
 		UOptimusNode* InTargetNode,
+		UOptimusNodePin* InPreferredTargetParentPin,
 		UOptimusNodePin* InSourcePin
-	);
+		);
 
 	const TArray<UOptimusNode*>& GetAllNodes() const { return Nodes; }
 	const TArray<UOptimusNodeLink*>& GetAllLinks() const { return Links; }
diff --git a/Engine/Plugins/Animation/DeformerGraph/Source/OptimusEditor/Private/OptimusEditorGraphSchema.cpp b/Engine/Plugins/Animation/DeformerGraph/Source/OptimusEditor/Private/OptimusEditorGraphSchema.cpp
--- a/Engine/Plugins/Animation/DeformerGraph/Source/OptimusEditor/Private/OptimusEditorGraphSchema.cpp
+++ b/Engine/Plugins/Animation/DeformerGraph/Source/OptimusEditor/Private/OptimusEditorGraphSchema.cpp
@@ -15,10 +15,15 @@
 
 #include "EdGraphSchema_K2.h"
 #include "Editor.h"
+#include "OptimusActionStack.h"
 #include "OptimusComputeDataInterface.h"
 #include "OptimusEditorGraphConnectionDrawingPolicy.h"
 #include "Styling/SlateIconFinder.h"
 
+#include "Framework/Application/SlateApplication.h"
+#include "Framework/Application/MenuStack.h"
+#include "Framework/MultiBox/MultiBoxBuilder.h"
+
 #define LOCTEXT_NAMESPACE "OptimusEditor"
 
 static bool IsValidAdderPinConnection(const UEdGraphPin* InPinA, const UEdGraphPin* InPinB, 
@@ -205,8 +210,57 @@ bool UOptimusEditorGraphSchema::TryCreateConnection(
 			TargetNode = OptimusEditor::GetModelNodeFromGraphPin(InPinA);
 			SourcePin = InputModelPin;
 		}
-		UOptimusNodeGraph *Graph = TargetNode->GetOwningGraph();
-		return Graph->AddPinAndLink(TargetNode, SourcePin);
+
+		IOptimusNodeAdderPinProvider *AdderPinProvider = Cast<IOptimusNodeAdderPinProvider>(TargetNode);
+		TArray<UOptimusNodePin*> ParentPins = AdderPinProvider->GetTargetParentPins(SourcePin);
+
+		if (ParentPins.Num() <= 1)
+		{
+			UOptimusNodeGraph *Graph = TargetNode->GetOwningGraph();
+			return Graph->AddPinAndLink(TargetNode, ParentPins.Num() == 0? nullptr : ParentPins[0], SourcePin);
+		}
+
+		FMenuBuilder MenuBuilder(true, NULL);
+
+		MenuBuilder.BeginSection(NAME_None, FText::FromString(TEXT("Add pin to")));
+
+		for (UOptimusNodePin* ParentPin : ParentPins)
+		{
+			FText EntryName;
+			if (!ParentPin)
+			{
+				EntryName = LOCTEXT("DefaultGroupName", "Default");
+			}
+			else
+			{
+				EntryName = FText::FromName(ParentPin->GetFName());
+			}
+			
+			MenuBuilder.AddMenuEntry(
+				EntryName,
+				FText::GetEmpty(),
+				FSlateIcon(),
+				FUIAction(
+				FExecuteAction::CreateLambda([TargetNode, ParentPin, SourcePin]()
+					{
+						UOptimusNodeGraph *Graph = TargetNode->GetOwningGraph();
+						Graph->AddPinAndLink(TargetNode, ParentPin, SourcePin);
+					}),
+					FCanExecuteAction()
+				));
+		}
+
+		MenuBuilder.EndSection();
+		
+		FSlateApplication::Get().PushMenu(
+			FSlateApplication::Get().GetUserFocusedWidget(0).ToSharedRef(),
+			FWidgetPath(),
+			MenuBuilder.MakeWidget(),
+			FSlateApplication::Get().GetCursorPos(),
+			FPopupTransitionEffect( FPopupTransitionEffect::ContextMenu)
+			);
+
+		return false;
 	}
 
 	return false;
diff --git a/Engine/Plugins/Animation/DeformerGraph/Source/OptimusEditor/Private/OptimusEditorGraphSchema.h b/Engine/Plugins/Animation/DeformerGraph/Source/OptimusEditor/Private/OptimusEditorGraphSchema.h
--- a/Engine/Plugins/Animation/DeformerGraph/Source/OptimusEditor/Private/OptimusEditorGraphSchema.h
+++ b/Engine/Plugins/Animation/DeformerGraph/Source/OptimusEditor/Private/OptimusEditorGraphSchema.h
@@ -34,7 +34,7 @@ public:
 	static FLinearColor GetColorFromPinType(const FEdGraphPinType& InPinType);
 
 	// UEdGraphSchema overrides
-	bool TryCreateConnection(UEdGraphPin* PinA, UEdGraphPin* PinB) const;
+	bool TryCreateConnection(UEdGraphPin* PinA, UEdGraphPin* PinB) const override;
 	const FPinConnectionResponse CanCreateConnection(const UEdGraphPin* A, const UEdGraphPin* B) const override;
 	void BreakPinLinks(UEdGraphPin& TargetPin, bool bSendsNodeNotifcation) const override;
 	void BreakSinglePinLink(UEdGraphPin* SourcePin, UEdGraphPin* TargetPin) const override;
