Deformer Graph: Fix multiple issues with variables.
- Variable default value was not being added on creation.
- The value container had the wrong size (was using property size, should have been shader value size).
- Ensured the default value and shader value containers were properly set on PostLoad.
- Undoing variable deletion would put the variable on the bottom of the list no matter where it was prior to deletion (applied the same to resources and bindings as well).
- Deleting a variable would leave dangling variable nodes. The variable nodes now get deleted as well.
- Fixed renaming of resources from the graph explorer.
- Fixed some erroneous BP markup for OptimusDeformerInstance.

#rnx

diff --git a/Engine/Plugins/Animation/DeformerGraph/Source/OptimusCore/Private/Actions/OptimusComponentBindingActions.cpp b/Engine/Plugins/Animation/DeformerGraph/Source/OptimusCore/Private/Actions/OptimusComponentBindingActions.cpp
--- a/Engine/Plugins/Animation/DeformerGraph/Source/OptimusCore/Private/Actions/OptimusComponentBindingActions.cpp
+++ b/Engine/Plugins/Animation/DeformerGraph/Source/OptimusCore/Private/Actions/OptimusComponentBindingActions.cpp
@@ -56,7 +56,7 @@ bool FOptimusComponentBindingAction_AddBinding::Do(IOptimusPathResolver* InRoot)
 		return false;
 	}
 
-	if (!Deformer->AddComponentBindingDirect(Binding))
+	if (!Deformer->AddComponentBindingDirect(Binding, INDEX_NONE))
 	{
 		Binding->Rename(nullptr, GetTransientPackage());
 		return false;
@@ -88,6 +88,7 @@ FOptimusComponentBindingAction_RemoveBinding::FOptimusComponentBindingAction_Rem
 		ComponentSourceClassPath = InBinding->GetComponentSource()->GetClass()->GetStructPathName();
 		ComponentBindingName = InBinding->BindingName;
 		ComponentTags = InBinding->ComponentTags;
+		ComponentBindingIndex = InBinding->GetIndex();
 
 		SetTitlef(TEXT("Remove component binding '%s'"), *ComponentBindingName.ToString());
 	}
@@ -130,7 +131,7 @@ bool FOptimusComponentBindingAction_RemoveBinding::Undo(IOptimusPathResolver* In
 
 	Binding->ComponentTags = ComponentTags;
 
-	if (!Deformer->AddComponentBindingDirect(Binding))
+	if (!Deformer->AddComponentBindingDirect(Binding, ComponentBindingIndex))
 	{
 		Binding->Rename(nullptr, GetTransientPackage());
 		return false;
diff --git a/Engine/Plugins/Animation/DeformerGraph/Source/OptimusCore/Private/Actions/OptimusComponentBindingActions.h b/Engine/Plugins/Animation/DeformerGraph/Source/OptimusCore/Private/Actions/OptimusComponentBindingActions.h
--- a/Engine/Plugins/Animation/DeformerGraph/Source/OptimusCore/Private/Actions/OptimusComponentBindingActions.h
+++ b/Engine/Plugins/Animation/DeformerGraph/Source/OptimusCore/Private/Actions/OptimusComponentBindingActions.h
@@ -64,6 +64,9 @@ private:
 
 	// Stored component tags
 	TArray<FName> ComponentTags;
+
+	// The index of the binding in the list.
+	int32 ComponentBindingIndex = 0;
 };
 
 
diff --git a/Engine/Plugins/Animation/DeformerGraph/Source/OptimusCore/Private/Actions/OptimusResourceActions.cpp b/Engine/Plugins/Animation/DeformerGraph/Source/OptimusCore/Private/Actions/OptimusResourceActions.cpp
--- a/Engine/Plugins/Animation/DeformerGraph/Source/OptimusCore/Private/Actions/OptimusResourceActions.cpp
+++ b/Engine/Plugins/Animation/DeformerGraph/Source/OptimusCore/Private/Actions/OptimusResourceActions.cpp
@@ -52,7 +52,7 @@ bool FOptimusResourceAction_AddResource::Do(
 	Resource->DataType = DataType;
 	Resource->DataDomain = DataDomain;
 
-	if (!Deformer->AddResourceDirect(Resource))
+	if (!Deformer->AddResourceDirect(Resource, INDEX_NONE))
 	{
 		Resource->Rename(nullptr, GetTransientPackage());
 		return false;
@@ -86,6 +86,7 @@ FOptimusResourceAction_RemoveResource::FOptimusResourceAction_RemoveResource(
 	{
 		ResourceName = InResource->GetFName();
 		DataType = InResource->DataType;
+		ResourceIndex = InResource->GetIndex();
 
 		SetTitlef(TEXT("Remove resource '%s'"), *InResource->GetName());
 	}
@@ -131,7 +132,7 @@ bool FOptimusResourceAction_RemoveResource::Undo(
 		Optimus::FBinaryObjectReader ResourceArchive(Resource, ResourceData);
 	}
 
-	if (!Deformer->AddResourceDirect(Resource))
+	if (!Deformer->AddResourceDirect(Resource, ResourceIndex))
 	{
 		Resource->Rename(nullptr, GetTransientPackage());
 		return false;
diff --git a/Engine/Plugins/Animation/DeformerGraph/Source/OptimusCore/Private/Actions/OptimusResourceActions.h b/Engine/Plugins/Animation/DeformerGraph/Source/OptimusCore/Private/Actions/OptimusResourceActions.h
--- a/Engine/Plugins/Animation/DeformerGraph/Source/OptimusCore/Private/Actions/OptimusResourceActions.h
+++ b/Engine/Plugins/Animation/DeformerGraph/Source/OptimusCore/Private/Actions/OptimusResourceActions.h
@@ -71,6 +71,9 @@ private:
 	// The data type of the resource
 	FOptimusDataTypeRef DataType;
 
+	// The index of the removed resource within the resource list
+	int32 ResourceIndex = 0;
+
 	// The stored resource data.
 	TArray<uint8> ResourceData;
 };
diff --git a/Engine/Plugins/Animation/DeformerGraph/Source/OptimusCore/Private/Actions/OptimusVariableActions.cpp b/Engine/Plugins/Animation/DeformerGraph/Source/OptimusCore/Private/Actions/OptimusVariableActions.cpp
--- a/Engine/Plugins/Animation/DeformerGraph/Source/OptimusCore/Private/Actions/OptimusVariableActions.cpp
+++ b/Engine/Plugins/Animation/DeformerGraph/Source/OptimusCore/Private/Actions/OptimusVariableActions.cpp
@@ -5,6 +5,8 @@
 #include "OptimusDeformer.h"
 #include "OptimusHelpers.h"
 #include "OptimusVariableDescription.h"
+#include "OptimusValueContainer.h"
+
 
 #include UE_INLINE_GENERATED_CPP_BY_NAME(OptimusVariableActions)
 
@@ -46,7 +48,9 @@ bool FOptimusVariableAction_AddVariable::Do(IOptimusPathResolver* InRoot)
 	Variable->VariableName = Variable->GetFName();
 	Variable->DataType = DataType;
 
-	if (!Deformer->AddVariableDirect(Variable))
+	Variable->EnsureValueContainer();
+
+	if (!Deformer->AddVariableDirect(Variable, INDEX_NONE))
 	{
 		Variable->Rename(nullptr, GetTransientPackage());
 		return false;
@@ -78,6 +82,7 @@ FOptimusVariableAction_RemoveVariable::FOptimusVariableAction_RemoveVariable(
 	{
 		VariableName = InVariable->GetFName();
 		DataType = InVariable->DataType;
+		VariableIndex = InVariable->GetIndex();
 
 		SetTitlef(TEXT("Remove variable '%s'"), *InVariable->GetName());
 	}
@@ -97,6 +102,11 @@ bool FOptimusVariableAction_RemoveVariable::Do(IOptimusPathResolver* InRoot)
 	{
 		Optimus::FBinaryObjectWriter VarArchive(Variable, VariableData);
 	}
+	
+	if (ensure(Variable->DefaultValue))
+	{
+		Optimus::FBinaryObjectWriter ValueArchive(Variable->DefaultValue, DefaultValueData);
+	}
 
 	return Deformer->RemoveVariableDirect(Variable);
 }
@@ -120,7 +130,14 @@ bool FOptimusVariableAction_RemoveVariable::Undo(IOptimusPathResolver* InRoot)
 		Optimus::FBinaryObjectReader VarArchive(Variable, VariableData);
 	}
 
-	if (!Deformer->AddVariableDirect(Variable))
+	Variable->EnsureValueContainer();
+
+	if (!DefaultValueData.IsEmpty() && Variable->DefaultValue)
+	{
+		Optimus::FBinaryObjectReader VarArchive(Variable->DefaultValue, DefaultValueData);
+	}
+	
+	if (!Deformer->AddVariableDirect(Variable, VariableIndex))
 	{
 		Variable->Rename(nullptr, GetTransientPackage());
 		return false;
diff --git a/Engine/Plugins/Animation/DeformerGraph/Source/OptimusCore/Private/Actions/OptimusVariableActions.h b/Engine/Plugins/Animation/DeformerGraph/Source/OptimusCore/Private/Actions/OptimusVariableActions.h
--- a/Engine/Plugins/Animation/DeformerGraph/Source/OptimusCore/Private/Actions/OptimusVariableActions.h
+++ b/Engine/Plugins/Animation/DeformerGraph/Source/OptimusCore/Private/Actions/OptimusVariableActions.h
@@ -65,8 +65,14 @@ private:
 	// The data type of the variable
 	FOptimusDataTypeRef DataType;
 
+	// THe index of the variable in the list
+	int32 VariableIndex = 0;
+	
 	// The stored variable data.
 	TArray<uint8> VariableData;
+
+	// The stored default value data
+	TArray<uint8> DefaultValueData;
 };
 
 
diff --git a/Engine/Plugins/Animation/DeformerGraph/Source/OptimusCore/Private/Nodes/OptimusNode_Resource.cpp b/Engine/Plugins/Animation/DeformerGraph/Source/OptimusCore/Private/Nodes/OptimusNode_Resource.cpp
--- a/Engine/Plugins/Animation/DeformerGraph/Source/OptimusCore/Private/Nodes/OptimusNode_Resource.cpp
+++ b/Engine/Plugins/Animation/DeformerGraph/Source/OptimusCore/Private/Nodes/OptimusNode_Resource.cpp
@@ -26,15 +26,20 @@ int32 UOptimusNode_Resource::GetDataFunctionIndexFromPin(const UOptimusNodePin*
 }
 
 
-FName UOptimusNode_Resource::GetResourcePinName(int32 InPinIndex) const
+FName UOptimusNode_Resource::GetResourcePinName(int32 InPinIndex, FName InNameOverride) const
 {
-	if (const UOptimusResourceDescription* Res = GetResourceDescription())
+	if (InNameOverride.IsNone())
 	{
-		if (ensure(InPinIndex >= 0 && InPinIndex < 2))
+		if (const UOptimusResourceDescription* Res = GetResourceDescription())
 		{
-			FString Suffix = InPinIndex == 0 ? TEXT("Set") : TEXT("Get");
-			return FName(Suffix + Res->ResourceName.ToString());
+			InNameOverride = Res->ResourceName;
 		}
+		
+	}
+	if (ensure(InPinIndex >= 0 && InPinIndex < 2))
+	{
+		FString Suffix = InPinIndex == 0 ? TEXT("Set") : TEXT("Get");
+		return FName(Suffix + InNameOverride.ToString());
 	}
 
 	return {};
diff --git a/Engine/Plugins/Animation/DeformerGraph/Source/OptimusCore/Private/Nodes/OptimusNode_Resource.h b/Engine/Plugins/Animation/DeformerGraph/Source/OptimusCore/Private/Nodes/OptimusNode_Resource.h
--- a/Engine/Plugins/Animation/DeformerGraph/Source/OptimusCore/Private/Nodes/OptimusNode_Resource.h
+++ b/Engine/Plugins/Animation/DeformerGraph/Source/OptimusCore/Private/Nodes/OptimusNode_Resource.h
@@ -14,7 +14,7 @@ class UOptimusNode_Resource
 	GENERATED_BODY()
 
 	// UOptimusNode_Resource implementations
-	FName GetResourcePinName(int32 InPinIndex) const override;
+	FName GetResourcePinName(int32 InPinIndex, FName InNameOverride = NAME_None) const override;
 	// IOptimusDataInterfaceProvider implementations
 	int32 GetDataFunctionIndexFromPin(const UOptimusNodePin* InPin) const override;
 
diff --git a/Engine/Plugins/Animation/DeformerGraph/Source/OptimusCore/Private/Nodes/OptimusNode_ResourceAccessorBase.cpp b/Engine/Plugins/Animation/DeformerGraph/Source/OptimusCore/Private/Nodes/OptimusNode_ResourceAccessorBase.cpp
--- a/Engine/Plugins/Animation/DeformerGraph/Source/OptimusCore/Private/Nodes/OptimusNode_ResourceAccessorBase.cpp
+++ b/Engine/Plugins/Animation/DeformerGraph/Source/OptimusCore/Private/Nodes/OptimusNode_ResourceAccessorBase.cpp
@@ -177,15 +177,22 @@ UOptimusResourceDescription* UOptimusNode_ResourceAccessorBase::GetResourceDescr
 }
 
 
-FName UOptimusNode_ResourceAccessorBase::GetResourcePinName(int32 InPinIndex) const
+FName UOptimusNode_ResourceAccessorBase::GetResourcePinName(int32 InPinIndex, FName InNameOverride) const
 {
-	UOptimusResourceDescription const* Description = GetResourceDescription();
-	if (!ensure(Description))
+	if (InNameOverride.IsNone())
 	{
-		return {};
-	}
+		UOptimusResourceDescription const* Description = GetResourceDescription();
+		if (!ensure(Description))
+		{
+			return {};
+		}
 	
-	return Description->ResourceName;
+		return Description->ResourceName;
+	}
+	else
+	{
+		return InNameOverride;
+	}
 }
 
 
diff --git a/Engine/Plugins/Animation/DeformerGraph/Source/OptimusCore/Private/Nodes/OptimusNode_ResourceAccessorBase.h b/Engine/Plugins/Animation/DeformerGraph/Source/OptimusCore/Private/Nodes/OptimusNode_ResourceAccessorBase.h
--- a/Engine/Plugins/Animation/DeformerGraph/Source/OptimusCore/Private/Nodes/OptimusNode_ResourceAccessorBase.h
+++ b/Engine/Plugins/Animation/DeformerGraph/Source/OptimusCore/Private/Nodes/OptimusNode_ResourceAccessorBase.h
@@ -44,7 +44,7 @@ public:
 	UOptimusResourceDescription* GetResourceDescription() const;
 
 	/** Returns name to use when creating each pin on the node. */
-	virtual FName GetResourcePinName(int32 InPinIndex) const;
+	virtual FName GetResourcePinName(int32 InPinIndex, FName InNameOverride = NAME_None) const;
 
 	// UOptimusNode overrides
 	FName GetNodeCategory() const override 
diff --git a/Engine/Plugins/Animation/DeformerGraph/Source/OptimusCore/Private/OptimusComponentSource.cpp b/Engine/Plugins/Animation/DeformerGraph/Source/OptimusCore/Private/OptimusComponentSource.cpp
--- a/Engine/Plugins/Animation/DeformerGraph/Source/OptimusCore/Private/OptimusComponentSource.cpp
+++ b/Engine/Plugins/Animation/DeformerGraph/Source/OptimusCore/Private/OptimusComponentSource.cpp
@@ -66,6 +66,12 @@ UOptimusDeformer* UOptimusComponentSourceBinding::GetOwningDeformer() const
 	return Container ? CastChecked<UOptimusDeformer>(Container->GetOuter()) : nullptr;
 }
 
+int32 UOptimusComponentSourceBinding::GetIndex() const
+{
+	const UOptimusComponentSourceBindingContainer* Container = CastChecked<UOptimusComponentSourceBindingContainer>(GetOuter());
+	return Container->Bindings.IndexOfByKey(this);
+}
+
 const UOptimusComponentSource* UOptimusComponentSourceBinding::GetComponentSource() const
 {
 	if (ComponentType)
diff --git a/Engine/Plugins/Animation/DeformerGraph/Source/OptimusCore/Private/OptimusDeformer.cpp b/Engine/Plugins/Animation/DeformerGraph/Source/OptimusCore/Private/OptimusDeformer.cpp
--- a/Engine/Plugins/Animation/DeformerGraph/Source/OptimusCore/Private/OptimusDeformer.cpp
+++ b/Engine/Plugins/Animation/DeformerGraph/Source/OptimusCore/Private/OptimusDeformer.cpp
@@ -194,7 +194,24 @@ bool UOptimusDeformer::RemoveVariable(
 		return false;
 	}
 
-	return GetActionStack()->RunAction<FOptimusVariableAction_RemoveVariable>(InVariableDesc);
+	FOptimusCompoundAction* Action = new FOptimusCompoundAction(TEXT("Remove Variable"));
+	
+	TMap<const UOptimusNodeGraph*, TArray<UOptimusNode*>> NodesByGraph;
+	for (UOptimusNode* Node: GetNodesUsingVariable(InVariableDesc))
+	{
+		UOptimusNode_GetVariable* VariableNode = Cast<UOptimusNode_GetVariable>(Node);
+		NodesByGraph.FindOrAdd(VariableNode->GetOwningGraph()).Add(VariableNode);
+	}
+
+	for (const TTuple<const UOptimusNodeGraph*, TArray<UOptimusNode*>>& GraphNodes: NodesByGraph)
+	{
+		const UOptimusNodeGraph* Graph = GraphNodes.Key;
+		Graph->RemoveNodesToAction(Action, GraphNodes.Value);
+	}
+
+	Action->AddSubAction<FOptimusVariableAction_RemoveVariable>(InVariableDesc);
+
+	return GetActionStack()->RunAction(Action);
 }
 
 
@@ -232,17 +249,13 @@ bool UOptimusDeformer::RenameVariable(
 	
 	FOptimusCompoundAction* Action = new FOptimusCompoundAction(TEXT("Rename Variable"));
 	
-	TArray<UOptimusNode*> AllVariableNodes = GetAllNodesOfClass(UOptimusNode_GetVariable::StaticClass());
-	for (UOptimusNode* Node: AllVariableNodes)
+	for (UOptimusNode* Node: GetNodesUsingVariable(InVariableDesc))
 	{
 		const UOptimusNode_GetVariable* VariableNode = Cast<UOptimusNode_GetVariable>(Node);
-		if (VariableNode->GetVariableDescription() == InVariableDesc)
+		if (ensure(VariableNode->GetPins().Num() == 1))
 		{
-			if (ensure(VariableNode->GetPins().Num() == 1))
-			{
-				Action->AddSubAction<FOptimusNodeAction_SetPinName>(VariableNode->GetPins()[0], InNewName);
-			}
-		}	
+			Action->AddSubAction<FOptimusNodeAction_SetPinName>(VariableNode->GetPins()[0], InNewName);
+		}
 	}
 
 	if (InNewName == InVariableDesc->VariableName)
@@ -289,36 +302,32 @@ bool UOptimusDeformer::SetVariableDataType(
 
 	TSet<TTuple<UOptimusNodePin* /*OutputPin*/, UOptimusNodePin* /*InputPin*/>> Links;
 	
-	TArray<UOptimusNode*> AllVariableNodes = GetAllNodesOfClass(UOptimusNode_GetVariable::StaticClass());
-	for (UOptimusNode* Node: AllVariableNodes)
+	for (UOptimusNode* Node: GetNodesUsingVariable(InVariableDesc))
 	{
-		const UOptimusNode_GetVariable* VariableNode = Cast<UOptimusNode_GetVariable>(Node);
-		if (VariableNode->GetVariableDescription() == InVariableDesc)
+		const UOptimusNode_GetVariable* VariableNode = CastChecked<UOptimusNode_GetVariable>(Node);
+		if (ensure(VariableNode->GetPins().Num() == 1))
 		{
-			if (ensure(VariableNode->GetPins().Num() == 1))
-			{
-				UOptimusNodePin* Pin = VariableNode->GetPins()[0];
+			UOptimusNodePin* Pin = VariableNode->GetPins()[0];
 
-				// Update the pin type to match.
-				Action->AddSubAction<FOptimusNodeAction_SetPinType>(VariableNode->GetPins()[0], InDataType);
+			// Update the pin type to match.
+			Action->AddSubAction<FOptimusNodeAction_SetPinType>(VariableNode->GetPins()[0], InDataType);
 
-				// Collect _unique_ links (in case there's a resource->resource link, since that would otherwise
-				// show up twice).
-				const UOptimusNodeGraph* Graph = Pin->GetOwningNode()->GetOwningGraph();
+			// Collect _unique_ links (in case there's a resource->resource link, since that would otherwise
+			// show up twice).
+			const UOptimusNodeGraph* Graph = Pin->GetOwningNode()->GetOwningGraph();
 
-				for (UOptimusNodePin* ConnectedPin: Graph->GetConnectedPins(Pin))
+			for (UOptimusNodePin* ConnectedPin: Graph->GetConnectedPins(Pin))
+			{
+				if (Pin->GetDirection() == EOptimusNodePinDirection::Output)
 				{
-					if (Pin->GetDirection() == EOptimusNodePinDirection::Output)
-					{
-						Links.Add({Pin, ConnectedPin});
-					}
-					else
-					{
-						Links.Add({ConnectedPin, Pin});
-					}
+					Links.Add({Pin, ConnectedPin});
+				}
+				else
+				{
+					Links.Add({ConnectedPin, Pin});
 				}
 			}
-		}	
+		}
 	}
 
 	for (auto [OutputPin, InputPin]: Links)
@@ -342,10 +351,27 @@ bool UOptimusDeformer::SetVariableDataType(
 	
 	// Make sure the value data container is still large enough to hold the property value.
 	InVariableDesc->ValueData.Reset();
-	InVariableDesc->ResetValueDataSize();
+	InVariableDesc->EnsureValueContainer();
 	return true;
 }
 
+TArray<UOptimusNode*> UOptimusDeformer::GetNodesUsingVariable(
+	const UOptimusVariableDescription* InVariableDesc
+	) const
+{
+	TArray<UOptimusNode*> UsedNodes;
+	TArray<UOptimusNode*> AllVariableNodes = GetAllNodesOfClass(UOptimusNode_GetVariable::StaticClass());
+	for (UOptimusNode* Node: AllVariableNodes)
+	{
+		const UOptimusNode_GetVariable* VariableNode = Cast<UOptimusNode_GetVariable>(Node);
+		if (VariableNode->GetVariableDescription() == InVariableDesc)
+		{
+			UsedNodes.Add(Node);
+		}
+	}
+	return UsedNodes;
+}
+
 
 UOptimusVariableDescription* UOptimusDeformer::ResolveVariable(
 	FName InVariableName
@@ -388,7 +414,8 @@ UOptimusVariableDescription* UOptimusDeformer::CreateVariableDirect(
 
 
 bool UOptimusDeformer::AddVariableDirect(
-	UOptimusVariableDescription* InVariableDesc
+	UOptimusVariableDescription* InVariableDesc,
+	const int32 InIndex
 	)
 {
 	if (!ensure(InVariableDesc))
@@ -401,7 +428,15 @@ bool UOptimusDeformer::AddVariableDirect(
 		return false;
 	}
 
-	Variables->Descriptions.Add(InVariableDesc);
+	if (Variables->Descriptions.IsValidIndex(InIndex))
+	{
+		Variables->Descriptions.Insert(InVariableDesc, InIndex);
+	}
+	else
+	{
+		Variables->Descriptions.Add(InVariableDesc);
+	}
+
 
 	Notify(EOptimusGlobalNotifyType::VariableAdded, InVariableDesc);
 
@@ -534,15 +569,10 @@ bool UOptimusDeformer::RemoveResource(UOptimusResourceDescription* InResourceDes
 	FOptimusCompoundAction* Action = new FOptimusCompoundAction(TEXT("Remove Resource"));
 	
 	TMap<const UOptimusNodeGraph*, TArray<UOptimusNode*>> NodesByGraph;
-	
-	TArray<UOptimusNode*> AllResourceNodes = GetAllNodesOfClass(UOptimusNode_ResourceAccessorBase::StaticClass());
-	for (UOptimusNode* Node: AllResourceNodes)
+	for (UOptimusNode* Node: GetNodesUsingResource(InResourceDesc))
 	{
 		UOptimusNode_ResourceAccessorBase* ResourceNode = Cast<UOptimusNode_ResourceAccessorBase>(Node);
-		if (ResourceNode->GetResourceDescription() == InResourceDesc)
-		{
-			NodesByGraph.FindOrAdd(ResourceNode->GetOwningGraph()).Add(ResourceNode);
-		}	
+		NodesByGraph.FindOrAdd(ResourceNode->GetOwningGraph()).Add(ResourceNode);
 	}
 
 	for (const TTuple<const UOptimusNodeGraph*, TArray<UOptimusNode*>>& GraphNodes: NodesByGraph)
@@ -592,17 +622,13 @@ bool UOptimusDeformer::RenameResource(
 
 	FOptimusCompoundAction* Action = new FOptimusCompoundAction(TEXT("Rename Resource"));
 
-	TArray<UOptimusNode*> AllResourceNodes = GetAllNodesOfClass(UOptimusNode_ResourceAccessorBase::StaticClass());
-	for (UOptimusNode* Node: AllResourceNodes)
+	for (UOptimusNode* Node: GetNodesUsingResource(InResourceDesc))
 	{
-		const UOptimusNode_ResourceAccessorBase* ResourceNode = Cast<UOptimusNode_ResourceAccessorBase>(Node);
-		if (ResourceNode->GetResourceDescription() == InResourceDesc)
+		const UOptimusNode_ResourceAccessorBase* ResourceNode = CastChecked<UOptimusNode_ResourceAccessorBase>(Node);
+		for (int32 PinIndex = 0; PinIndex < ResourceNode->GetPins().Num(); ++PinIndex)
 		{
-			for (int32 PinIndex = 0; PinIndex < ResourceNode->GetPins().Num(); ++PinIndex)
-			{
-				Action->AddSubAction<FOptimusNodeAction_SetPinName>(ResourceNode->GetPins()[PinIndex], ResourceNode->GetResourcePinName(PinIndex));
-			}
-		}	
+			Action->AddSubAction<FOptimusNodeAction_SetPinName>(ResourceNode->GetPins()[PinIndex], ResourceNode->GetResourcePinName(PinIndex, InNewName));
+		}
 	}	
 
 	if (InNewName == InResourceDesc->ResourceName)
@@ -649,33 +675,29 @@ bool UOptimusDeformer::SetResourceDataType(
 
 	TSet<TTuple<UOptimusNodePin* /*OutputPin*/, UOptimusNodePin* /*InputPin*/>> Links;
 	
-	TArray<UOptimusNode*> AllResourceNodes = GetAllNodesOfClass(UOptimusNode_ResourceAccessorBase::StaticClass());
-	for (UOptimusNode* Node: AllResourceNodes)
+	for (UOptimusNode* Node: GetNodesUsingResource(InResourceDesc))
 	{
-		const UOptimusNode_ResourceAccessorBase* ResourceNode = Cast<UOptimusNode_ResourceAccessorBase>(Node);
-		if (ResourceNode->GetResourceDescription() == InResourceDesc)
+		const UOptimusNode_ResourceAccessorBase* ResourceNode = CastChecked<UOptimusNode_ResourceAccessorBase>(Node);
+		for (UOptimusNodePin* Pin : ResourceNode->GetPins())
 		{
-			for (UOptimusNodePin* Pin : ResourceNode->GetPins())
-			{
-				// Update the pin type to match.
-				Action->AddSubAction<FOptimusNodeAction_SetPinType>(Pin, InDataType);
+			// Update the pin type to match.
+			Action->AddSubAction<FOptimusNodeAction_SetPinType>(Pin, InDataType);
 
-				// Collect _unique_ links (in case there's a resource->resource link, since that would otherwise show up twice).
-				const UOptimusNodeGraph* Graph = Pin->GetOwningNode()->GetOwningGraph();
+			// Collect _unique_ links (in case there's a resource->resource link, since that would otherwise show up twice).
+			const UOptimusNodeGraph* Graph = Pin->GetOwningNode()->GetOwningGraph();
 
-				for (UOptimusNodePin* ConnectedPin: Graph->GetConnectedPins(Pin))
+			for (UOptimusNodePin* ConnectedPin: Graph->GetConnectedPins(Pin))
+			{
+				if (Pin->GetDirection() == EOptimusNodePinDirection::Output)
 				{
-					if (Pin->GetDirection() == EOptimusNodePinDirection::Output)
-					{
-						Links.Add({Pin, ConnectedPin});
-					}
-					else
-					{
-						Links.Add({ConnectedPin, Pin});
-					}
+					Links.Add({Pin, ConnectedPin});
+				}
+				else
+				{
+					Links.Add({ConnectedPin, Pin});
 				}
 			}
-		}	
+		}
 	}
 
 	for (auto [OutputPin, InputPin]: Links)
@@ -716,31 +738,27 @@ bool UOptimusDeformer::SetResourceDataDomain(
 
 	TSet<TTuple<UOptimusNodePin* /*OutputPin*/, UOptimusNodePin* /*InputPin*/>> Links;
 	
-	TArray<UOptimusNode*> AllResourceNodes = GetAllNodesOfClass(UOptimusNode_ResourceAccessorBase::StaticClass());
-	for (UOptimusNode* Node: AllResourceNodes)
+	for (UOptimusNode* Node: GetNodesUsingResource(InResourceDesc))
 	{
-		const UOptimusNode_ResourceAccessorBase* ResourceNode = Cast<UOptimusNode_ResourceAccessorBase>(Node);
-		if (ResourceNode->GetResourceDescription() == InResourceDesc)
+		const UOptimusNode_ResourceAccessorBase* ResourceNode = CastChecked<UOptimusNode_ResourceAccessorBase>(Node);
+		for (UOptimusNodePin* Pin : ResourceNode->GetPins())
 		{
-			for (UOptimusNodePin* Pin : ResourceNode->GetPins())
-			{
-				// Update the pin type to match.
-				Action->AddSubAction<FOptimusNodeAction_SetPinDataDomain>(Pin, InDataDomain);
+			// Update the pin type to match.
+			Action->AddSubAction<FOptimusNodeAction_SetPinDataDomain>(Pin, InDataDomain);
 
-				// Collect _unique_ links (in case there's a resource->resource link, since that would otherwise
-				// show up twice).
-				const UOptimusNodeGraph* Graph = Pin->GetOwningNode()->GetOwningGraph();
+			// Collect _unique_ links (in case there's a resource->resource link, since that would otherwise
+			// show up twice).
+			const UOptimusNodeGraph* Graph = Pin->GetOwningNode()->GetOwningGraph();
 
-				for (UOptimusNodePin* ConnectedPin: Graph->GetConnectedPins(Pin))
+			for (UOptimusNodePin* ConnectedPin: Graph->GetConnectedPins(Pin))
+			{
+				if (Pin->GetDirection() == EOptimusNodePinDirection::Output)
 				{
-					if (Pin->GetDirection() == EOptimusNodePinDirection::Output)
-					{
-						Links.Add({Pin, ConnectedPin});
-					}
-					else
-					{
-						Links.Add({ConnectedPin, Pin});
-					}
+					Links.Add({Pin, ConnectedPin});
+				}
+				else
+				{
+					Links.Add({ConnectedPin, Pin});
 				}
 			}
 		}	
@@ -760,6 +778,24 @@ bool UOptimusDeformer::SetResourceDataDomain(
 }
 
 
+TArray<UOptimusNode*> UOptimusDeformer::GetNodesUsingResource(
+	const UOptimusResourceDescription* InResourceDesc
+	) const
+{
+	TArray<UOptimusNode*> UsedNodes;
+	TArray<UOptimusNode*> AllResourceNodes = GetAllNodesOfClass(UOptimusNode_ResourceAccessorBase::StaticClass());
+	for (UOptimusNode* Node: AllResourceNodes)
+	{
+		const UOptimusNode_ResourceAccessorBase* ResourceNode = Cast<UOptimusNode_ResourceAccessorBase>(Node);
+		if (ResourceNode->GetResourceDescription() == InResourceDesc)
+		{
+			UsedNodes.Add(Node);
+		}
+	}
+	return UsedNodes;
+}
+
+
 UOptimusResourceDescription* UOptimusDeformer::ResolveResource(
 	FName InResourceName
 	) const
@@ -797,7 +833,8 @@ UOptimusResourceDescription* UOptimusDeformer::CreateResourceDirect(
 
 
 bool UOptimusDeformer::AddResourceDirect(
-	UOptimusResourceDescription* InResourceDesc
+	UOptimusResourceDescription* InResourceDesc,
+	const int32 InIndex
 	)
 {
 	if (!ensure(InResourceDesc))
@@ -810,7 +847,15 @@ bool UOptimusDeformer::AddResourceDirect(
 		return false;
 	}
 
-	Resources->Descriptions.Add(InResourceDesc);
+	if (Resources->Descriptions.IsValidIndex(InIndex))
+	{
+		Resources->Descriptions.Insert(InResourceDesc, InIndex);
+	}
+	else
+	{
+		Resources->Descriptions.Add(InResourceDesc);
+	}
+	
 	Notify(EOptimusGlobalNotifyType::ResourceAdded, InResourceDesc);
 	(void)MarkPackageDirty();
 	
@@ -967,7 +1012,8 @@ UOptimusComponentSourceBinding* UOptimusDeformer::CreateComponentBindingDirect(
 
 
 bool UOptimusDeformer::AddComponentBindingDirect(
-	UOptimusComponentSourceBinding* InComponentBinding
+	UOptimusComponentSourceBinding* InComponentBinding,
+	const int32 InIndex
 	)
 {
 	if (!ensure(InComponentBinding))
@@ -983,8 +1029,15 @@ bool UOptimusDeformer::AddComponentBindingDirect(
 	{
 		InComponentBinding->bIsPrimaryBinding = true;
 	}
-	
-	Bindings->Bindings.Add(InComponentBinding);
+
+	if (Bindings->Bindings.IsValidIndex(InIndex))
+	{
+		Bindings->Bindings.Insert(InComponentBinding, InIndex);
+	}
+	else
+	{
+		Bindings->Bindings.Add(InComponentBinding);
+	}
 
 	Notify(EOptimusGlobalNotifyType::ComponentBindingAdded, InComponentBinding);
 
@@ -1017,14 +1070,9 @@ bool UOptimusDeformer::RemoveComponentBinding(
 	
 	TMap<const UOptimusNodeGraph*, TArray<UOptimusNode*>> NodesByGraph;
 	
-	TArray<UOptimusNode*> AllComponentSourceNodes = GetAllNodesOfClass(UOptimusNode_ComponentSource::StaticClass());
-	for (UOptimusNode* Node: AllComponentSourceNodes)
+	for (UOptimusNode* Node: GetNodesUsingComponentBinding(InBinding))
 	{
-		UOptimusNode_ComponentSource* ComponentNode = Cast<UOptimusNode_ComponentSource>(Node);
-		if (ComponentNode->GetComponentBinding() == InBinding)
-		{
-			NodesByGraph.FindOrAdd(ComponentNode->GetOwningGraph()).Add(ComponentNode);
-		}	
+		NodesByGraph.FindOrAdd(Node->GetOwningGraph()).Add(Node);
 	}
 
 	for (const TTuple<const UOptimusNodeGraph*, TArray<UOptimusNode*>>& GraphNodes: NodesByGraph)
@@ -1096,14 +1144,9 @@ bool UOptimusDeformer::RenameComponentBinding(
 
 	FOptimusCompoundAction* Action = new FOptimusCompoundAction(TEXT("Rename Component Binding"));
 
-	TArray<UOptimusNode*> AllComponentSourceNode = GetAllNodesOfClass(UOptimusNode_ComponentSource::StaticClass());
-	for (UOptimusNode* Node: AllComponentSourceNode)
+	for (UOptimusNode* Node: GetNodesUsingComponentBinding(InBinding))
 	{
-		UOptimusNode_ComponentSource* ComponentSourceNode = Cast<UOptimusNode_ComponentSource>(Node);
-		if (ComponentSourceNode->GetComponentBinding() == InBinding)
-		{
-			Action->AddSubAction<FOptimusNodeAction_RenameNode>(ComponentSourceNode, FText::FromName(InNewName));
-		}	
+		Action->AddSubAction<FOptimusNodeAction_RenameNode>(Node, FText::FromName(InNewName));
 	}	
 
 	if (InNewName == InBinding->BindingName)
@@ -1172,32 +1215,28 @@ bool UOptimusDeformer::SetComponentBindingSource(
 
 	FOptimusCompoundAction* Action = new FOptimusCompoundAction(TEXT("Set Component Binding Source"));
 	TSet<TTuple<UOptimusNodePin* /*OutputPin*/, UOptimusNodePin* /*InputPin*/>> Links;
-	TArray<UOptimusNode*> AllComponentSourceNode = GetAllNodesOfClass(UOptimusNode_ComponentSource::StaticClass());
-	for (UOptimusNode* Node: AllComponentSourceNode)
+	for (UOptimusNode* Node: GetNodesUsingComponentBinding(InBinding))
 	{
-		const UOptimusNode_ComponentSource* ComponentSourceNode = Cast<UOptimusNode_ComponentSource>(Node);
-		if (ComponentSourceNode->GetComponentBinding() == InBinding)
-		{
-			UOptimusNodePin* ComponentSourcePin = ComponentSourceNode->GetComponentPin();
+		const UOptimusNode_ComponentSource* ComponentSourceNode = CastChecked<UOptimusNode_ComponentSource>(Node);
+		UOptimusNodePin* ComponentSourcePin = ComponentSourceNode->GetComponentPin();
 
-			const UOptimusNodeGraph* Graph = ComponentSourceNode->GetOwningGraph();
+		const UOptimusNodeGraph* Graph = ComponentSourceNode->GetOwningGraph();
 
-			for (UOptimusNodePin* ConnectedPin: Graph->GetConnectedPins(ComponentSourcePin))
+		for (UOptimusNodePin* ConnectedPin: Graph->GetConnectedPins(ComponentSourcePin))
+		{
+			// Will this connection be invalid once the source is changed?
+			if (const UOptimusNode_DataInterface* DataInterfaceNode = Cast<UOptimusNode_DataInterface>(ConnectedPin->GetOwningNode()))
 			{
-				// Will this connection be invalid once the source is changed?
-				if (const UOptimusNode_DataInterface* DataInterfaceNode = Cast<UOptimusNode_DataInterface>(ConnectedPin->GetOwningNode()))
+				if (!DataInterfaceNode->IsComponentSourceCompatible(InComponentSource))
 				{
-					if (!DataInterfaceNode->IsComponentSourceCompatible(InComponentSource))
-					{
-						Action->AddSubAction<FOptimusNodeGraphAction_RemoveLink>(ComponentSourcePin, ConnectedPin);
-					}
+					Action->AddSubAction<FOptimusNodeGraphAction_RemoveLink>(ComponentSourcePin, ConnectedPin);
 				}
 			}
+		}
 
-			// Change the pin name _after_ the links are removed, since the link remove action uses the pin path, including
-			// the old name to resolve the pin.
-			Action->AddSubAction<FOptimusNodeAction_SetPinName>(ComponentSourcePin, InComponentSource->GetBindingName());
-		}	
+		// Change the pin name _after_ the links are removed, since the link remove action uses the pin path, including
+		// the old name to resolve the pin.
+		Action->AddSubAction<FOptimusNodeAction_SetPinName>(ComponentSourcePin, InComponentSource->GetBindingName());
 	}
 	
 	if (InBinding->ComponentType != InComponentSource->GetClass())
@@ -1213,6 +1252,24 @@ bool UOptimusDeformer::SetComponentBindingSource(
 }
 
 
+TArray<UOptimusNode*> UOptimusDeformer::GetNodesUsingComponentBinding(
+	const UOptimusComponentSourceBinding* InBinding
+	) const
+{
+	TArray<UOptimusNode*> UsedNodes;
+	TArray<UOptimusNode*> AllComponentSourceNode = GetAllNodesOfClass(UOptimusNode_ComponentSource::StaticClass());
+	for (UOptimusNode* Node: AllComponentSourceNode)
+	{
+		const UOptimusNode_ComponentSource* ComponentSourceNode = Cast<UOptimusNode_ComponentSource>(Node);
+		if (ComponentSourceNode->GetComponentBinding() == InBinding)
+		{
+			UsedNodes.Add(Node);
+		}
+	}
+	return UsedNodes;
+}
+
+
 bool UOptimusDeformer::SetComponentBindingSourceDirect(
 	UOptimusComponentSourceBinding* InBinding,
 	const UOptimusComponentSource* InComponentSource
diff --git a/Engine/Plugins/Animation/DeformerGraph/Source/OptimusCore/Private/OptimusDeformerInstance.cpp b/Engine/Plugins/Animation/DeformerGraph/Source/OptimusCore/Private/OptimusDeformerInstance.cpp
--- a/Engine/Plugins/Animation/DeformerGraph/Source/OptimusCore/Private/OptimusDeformerInstance.cpp
+++ b/Engine/Plugins/Animation/DeformerGraph/Source/OptimusCore/Private/OptimusDeformerInstance.cpp
@@ -474,8 +474,6 @@ namespace
 	template <typename T>
 	bool SetVariableValue(UOptimusVariableContainer const* InVariables, FName InVariableName, FName InTypeName, T const& InValue)
 	{
-		const uint8* ValueBytes = reinterpret_cast<const uint8*>(&InValue);
-
 		FOptimusDataTypeHandle WantedType = FOptimusDataTypeRegistry::Get().FindType(InTypeName);
 		for (UOptimusVariableDescription* VariableDesc : InVariables->Descriptions)
 		{
@@ -484,8 +482,9 @@ namespace
 				TUniquePtr<FProperty> Property(WantedType->CreateProperty(nullptr, NAME_None));
 				if (ensure(Property->GetSize() == sizeof(T)))
 				{
+					const uint8* ValueBytes = reinterpret_cast<const uint8*>(&InValue);
 					FShaderValueType::FValue ValueResult = WantedType->MakeShaderValue();
-					WantedType->ConvertPropertyValueToShader(TArrayView<const uint8>((const uint8*)&InValue, sizeof(T)), ValueResult);
+					WantedType->ConvertPropertyValueToShader(TArrayView<const uint8>(ValueBytes, sizeof(T)), ValueResult);
 					VariableDesc->ValueData = MoveTemp(ValueResult.ShaderValue);
 				}
 
diff --git a/Engine/Plugins/Animation/DeformerGraph/Source/OptimusCore/Private/OptimusResourceDescription.cpp b/Engine/Plugins/Animation/DeformerGraph/Source/OptimusCore/Private/OptimusResourceDescription.cpp
--- a/Engine/Plugins/Animation/DeformerGraph/Source/OptimusCore/Private/OptimusResourceDescription.cpp
+++ b/Engine/Plugins/Animation/DeformerGraph/Source/OptimusCore/Private/OptimusResourceDescription.cpp
@@ -17,6 +17,13 @@ UOptimusDeformer* UOptimusResourceDescription::GetOwningDeformer() const
 }
 
 
+int32 UOptimusResourceDescription::GetIndex() const
+{
+	const UOptimusResourceContainer* Container = CastChecked<UOptimusResourceContainer>(GetOuter());
+	return Container->Descriptions.IndexOfByKey(this);
+}
+
+
 void UOptimusResourceDescription::PostLoad()
 {
 	Super::PostLoad();
diff --git a/Engine/Plugins/Animation/DeformerGraph/Source/OptimusCore/Private/OptimusVariableDescription.cpp b/Engine/Plugins/Animation/DeformerGraph/Source/OptimusCore/Private/OptimusVariableDescription.cpp
--- a/Engine/Plugins/Animation/DeformerGraph/Source/OptimusCore/Private/OptimusVariableDescription.cpp
+++ b/Engine/Plugins/Animation/DeformerGraph/Source/OptimusCore/Private/OptimusVariableDescription.cpp
@@ -10,21 +10,39 @@
 #include UE_INLINE_GENERATED_CPP_BY_NAME(OptimusVariableDescription)
 
 
-void UOptimusVariableDescription::ResetValueDataSize()
+bool UOptimusVariableDescription::EnsureValueContainer()
 {
-	DefaultValue = UOptimusValueContainer::MakeValueContainer(this, DataType);
+	bool bValueContainerChanged = false;
+	
+	// Check if the current default value storage matches, otherwise create a matching default value storage, otherwise
+	// if the variable type changes, we end up with mismatch in storage vs type.
+	const UClass* RequiredClass = UOptimusValueContainerGeneratorClass::GetClassForType(GetPackage(), DataType);
+
+	bool bValueContainerNeedsZeroing = false;
+	if (!DefaultValue || DefaultValue->GetClass() != RequiredClass)
+	{
+		DefaultValue = UOptimusValueContainer::MakeValueContainer(this, DataType);
+		bValueContainerChanged = true;
+		bValueContainerNeedsZeroing = true;
+	}
 
 	if (DataType->CanCreateProperty())
 	{
-		// Create a temporary property from the type so that we can get the size of the
-		// type for properly resizing the storage.
-		TUniquePtr<FProperty> TempProperty(DataType->CreateProperty(nullptr, NAME_None));
+		const FShaderValueType::FValue ShaderValue = DataType->MakeShaderValue();
 		
-		if (ValueData.Num() != TempProperty->GetSize())
+		if (bValueContainerNeedsZeroing || ValueData.Num() != ShaderValue.ShaderValue.Num())
 		{
-			ValueData.SetNumZeroed(TempProperty->GetSize());
+			ValueData.SetNumZeroed(ShaderValue.ShaderValue.Num());
+			bValueContainerChanged = true;
 		}
 	}
+	else if (!ValueData.IsEmpty())
+	{
+		ValueData.Reset();
+		bValueContainerChanged = true;
+	}
+
+	return bValueContainerChanged;
 }
 
 
@@ -34,6 +52,14 @@ UOptimusDeformer* UOptimusVariableDescription::GetOwningDeformer() const
 	return Container ? CastChecked<UOptimusDeformer>(Container->GetOuter()) : nullptr;
 }
 
+
+int32 UOptimusVariableDescription::GetIndex() const
+{
+	const UOptimusVariableContainer* Container = CastChecked<UOptimusVariableContainer>(GetOuter());
+	return Container->Descriptions.IndexOfByKey(this);
+}
+
+
 void UOptimusVariableDescription::PostLoad()
 {
 	Super::PostLoad();
@@ -45,7 +71,10 @@ void UOptimusVariableDescription::PostLoad()
 	if (DataType == FloatDataType)
 	{
 		DataType = DoubleDataType;
-		ResetValueDataSize();
+	}
+
+	if (EnsureValueContainer())
+	{
 		(void)MarkPackageDirty();
 	}
 }
diff --git a/Engine/Plugins/Animation/DeformerGraph/Source/OptimusCore/Public/OptimusComponentSource.h b/Engine/Plugins/Animation/DeformerGraph/Source/OptimusCore/Public/OptimusComponentSource.h
--- a/Engine/Plugins/Animation/DeformerGraph/Source/OptimusCore/Public/OptimusComponentSource.h
+++ b/Engine/Plugins/Animation/DeformerGraph/Source/OptimusCore/Public/OptimusComponentSource.h
@@ -82,6 +82,9 @@ public:
 	// FIXME: Move to interface-based system.
 	UOptimusDeformer* GetOwningDeformer() const;
 
+	/** Get the index of this binding within its container */
+	int32 GetIndex() const;
+
 	/** The name to give the binding, to disambiguate it from other bindings of same component type. */
 	UPROPERTY(EditAnywhere, Category=Binding, meta = (EditCondition = "!bIsPrimaryBinding", HideEditConditionToggle))
 	FName BindingName;
diff --git a/Engine/Plugins/Animation/DeformerGraph/Source/OptimusCore/Public/OptimusDeformer.h b/Engine/Plugins/Animation/DeformerGraph/Source/OptimusCore/Public/OptimusDeformer.h
--- a/Engine/Plugins/Animation/DeformerGraph/Source/OptimusCore/Public/OptimusDeformer.h
+++ b/Engine/Plugins/Animation/DeformerGraph/Source/OptimusCore/Public/OptimusDeformer.h
@@ -162,6 +162,10 @@ public:
 		bool bInForceChange = false
 		);
 
+	TArray<UOptimusNode*> GetNodesUsingVariable(
+	    const UOptimusVariableDescription* InVariableDesc
+		) const;
+
 	UFUNCTION(BlueprintGetter)
 	const TArray<UOptimusVariableDescription*>& GetVariables() const { return Variables->Descriptions; }
 
@@ -197,6 +201,10 @@ public:
 		bool bInForceChange = false
 		);
 	
+	TArray<UOptimusNode*> GetNodesUsingResource(
+		const UOptimusResourceDescription* InResourceDesc
+		) const;
+	
 	UFUNCTION(BlueprintGetter)
 	const TArray<UOptimusResourceDescription*>& GetResources() const { return Resources->Descriptions; }
 
@@ -225,6 +233,10 @@ public:
 		const UOptimusComponentSource *InComponentSource,
 		bool bInForceChange = false
 		);
+
+	TArray<UOptimusNode*> GetNodesUsingComponentBinding(
+		const UOptimusComponentSourceBinding* InBinding
+		) const;
 	
 	UFUNCTION(BlueprintGetter)
 	const TArray<UOptimusComponentSourceBinding*>& GetComponentBindings() const { return Bindings->Bindings; }
@@ -295,23 +307,28 @@ public:
 	UOptimusNodeGraph* CreateGraph(
 	    EOptimusNodeGraphType InType,
 	    FName InName,
-	    TOptional<int32> InInsertBefore) override;
+	    TOptional<int32> InInsertBefore
+	    ) override;
 	
 	bool AddGraph(
 	    UOptimusNodeGraph* InGraph,
-		int32 InInsertBefore) override;
+		int32 InInsertBefore
+		) override;
 	
 	bool RemoveGraph(
 	    UOptimusNodeGraph* InGraph,
-		bool bDeleteGraph) override;
+		bool bDeleteGraph
+		) override;
 
 	bool MoveGraph(
 	    UOptimusNodeGraph* InGraph,
-	    int32 InInsertBefore) override;
+	    int32 InInsertBefore
+	    ) override;
 
 	bool RenameGraph(
 	    UOptimusNodeGraph* InGraph,
-	    const FString& InNewName) override;
+	    const FString& InNewName
+	    ) override;
 	
 	UPROPERTY(EditAnywhere, Category=Preview)
 	TObjectPtr<USkeletalMesh> Mesh = nullptr;
@@ -344,7 +361,8 @@ protected:
 
 	/** Adds a resource that was created by this deformer and is owned by it. */
 	bool AddResourceDirect(
-		UOptimusResourceDescription* InResourceDesc
+		UOptimusResourceDescription* InResourceDesc,
+		const int32 InIndex
 		);
 
 	bool RemoveResourceDirect(
@@ -375,7 +393,8 @@ protected:
 
 	/** Adds a resource that was created by this deformer and is owned by it. */
 	bool AddVariableDirect(
-		UOptimusVariableDescription* InVariableDesc
+		UOptimusVariableDescription* InVariableDesc,
+		const int32 InIndex
 		);
 
 	bool RemoveVariableDirect(
@@ -398,7 +417,8 @@ protected:
 		);
 
 	bool AddComponentBindingDirect(
-		UOptimusComponentSourceBinding* InComponentBinding
+		UOptimusComponentSourceBinding* InComponentBinding,
+		const int32 InIndex
 		);
 
 	bool RemoveComponentBindingDirect(
diff --git a/Engine/Plugins/Animation/DeformerGraph/Source/OptimusCore/Public/OptimusDeformerInstance.h b/Engine/Plugins/Animation/DeformerGraph/Source/OptimusCore/Public/OptimusDeformerInstance.h
--- a/Engine/Plugins/Animation/DeformerGraph/Source/OptimusCore/Public/OptimusDeformerInstance.h
+++ b/Engine/Plugins/Animation/DeformerGraph/Source/OptimusCore/Public/OptimusDeformerInstance.h
@@ -133,7 +133,7 @@ protected:
  * This implements the UMeshDeformerInstance interface to enqueue the graph execution.
  * It also contains the per instance deformer variable state and local state for each of the graphs in the deformer.
  */
-UCLASS(Blueprintable, BlueprintType)
+UCLASS(BlueprintType)
 class UOptimusDeformerInstance :
 	public UMeshDeformerInstance
 {
@@ -160,7 +160,7 @@ public:
 	void SetupFromDeformer(UOptimusDeformer* InDeformer);
 
 	/** Set the value of a boolean variable. */
-	UFUNCTION(BlueprintPure, Category="Deformer", meta=(DisplayName="Set Variable (Boolean)"))
+	UFUNCTION(BlueprintCallable, Category="Deformer", meta=(DisplayName="Set Variable (Boolean)"))
 	bool SetBoolVariable(FName InVariableName, bool InValue);
 
 	/** Set the value of an integer variable. */
diff --git a/Engine/Plugins/Animation/DeformerGraph/Source/OptimusCore/Public/OptimusResourceDescription.h b/Engine/Plugins/Animation/DeformerGraph/Source/OptimusCore/Public/OptimusResourceDescription.h
--- a/Engine/Plugins/Animation/DeformerGraph/Source/OptimusCore/Public/OptimusResourceDescription.h
+++ b/Engine/Plugins/Animation/DeformerGraph/Source/OptimusCore/Public/OptimusResourceDescription.h
@@ -30,6 +30,9 @@ public:
 	// FIXME: Move to interface-based system.
 	UOptimusDeformer* GetOwningDeformer() const;
 
+	/** Returns the index of the variable within the container */ 
+	int32 GetIndex() const;
+
 	UPROPERTY(EditAnywhere, BlueprintReadOnly, Category = ResourceDescription)
 	FName ResourceName;
 
diff --git a/Engine/Plugins/Animation/DeformerGraph/Source/OptimusCore/Public/OptimusVariableDescription.h b/Engine/Plugins/Animation/DeformerGraph/Source/OptimusCore/Public/OptimusVariableDescription.h
--- a/Engine/Plugins/Animation/DeformerGraph/Source/OptimusCore/Public/OptimusVariableDescription.h
+++ b/Engine/Plugins/Animation/DeformerGraph/Source/OptimusCore/Public/OptimusVariableDescription.h
@@ -46,12 +46,15 @@ public:
 	 * Set the value data storage to match the size required by the DataType. 
 	 * If a reallocation is required then the value data will be zeroed.
 	 */
-	void ResetValueDataSize();
+	bool EnsureValueContainer();
 
 	/** Returns the owning deformer to operate on this variable */
 	// FIXME: Move to interface-based system.
 	UOptimusDeformer* GetOwningDeformer() const;
 
+	/** Returns the index of the variable within the container */ 
+	int32 GetIndex() const;
+
 	/** An identifier that uniquely identifies this variable */
 	UPROPERTY()
 	FGuid Guid;
