Optimus: Disable connection from nonunified data interface to seconary group on a kernel within the graph

minor adder pin refactor, added the ability to create pins from any pin in the graph, and auto link only if the component bindings match & allowed

diff --git a/Engine/Plugins/Animation/DeformerGraph/Source/OptimusCore/Private/Actions/OptimusNodeGraphActions.cpp b/Engine/Plugins/Animation/DeformerGraph/Source/OptimusCore/Private/Actions/OptimusNodeGraphActions.cpp
--- a/Engine/Plugins/Animation/DeformerGraph/Source/OptimusCore/Private/Actions/OptimusNodeGraphActions.cpp
+++ b/Engine/Plugins/Animation/DeformerGraph/Source/OptimusCore/Private/Actions/OptimusNodeGraphActions.cpp
@@ -620,13 +620,12 @@ FOptimusNodeGraphAction_RemoveLink::FOptimusNodeGraphAction_RemoveLink(
 
 FOptimusNodeGraphAction_ConnectAdderPin::FOptimusNodeGraphAction_ConnectAdderPin(
 	IOptimusNodeAdderPinProvider* InAdderPinProvider,
-	UOptimusNodePin* InPreferredTargetParentPin,
-	UOptimusNodePin* InSourcePin,
-	bool bInShouldLink
+	const IOptimusNodeAdderPinProvider::FAdderPinAction& InAction,
+	UOptimusNodePin* InSourcePin
 	) :
-	bShouldLink(bInShouldLink)
+	Action(InAction)
 {
-	if (bInShouldLink)
+	if (InAction.bCanAutoLink)
 	{
 		SetTitlef(TEXT("Add Pin & Link"));
 	}
@@ -642,11 +641,6 @@ FOptimusNodeGraphAction_ConnectAdderPin::FOptimusNodeGraphAction_ConnectAdderPin
 		NodePath = Node->GetNodePath();
 	}
 
-	if (InPreferredTargetParentPin)
-	{
-		PreferredParentPinPath = InPreferredTargetParentPin->GetPinPath();
-	}
-	
 	SourcePinPath = InSourcePin->GetPinPath();
 }
 
@@ -660,11 +654,9 @@ bool FOptimusNodeGraphAction_ConnectAdderPin::Do(IOptimusPathResolver* InRoot)
 		return false;
 	}
 
-	UOptimusNodePin* PreferredTargetParentPin = InRoot->ResolvePinPath(PreferredParentPinPath);
-
 	UOptimusNodePin* SourcePin = InRoot->ResolvePinPath(SourcePinPath);
 
-	TArray<UOptimusNodePin*> AddedPins = AdderPinProvider->TryAddPinFromPin(PreferredTargetParentPin, SourcePin);
+	TArray<UOptimusNodePin*> AddedPins = AdderPinProvider->TryAddPinFromPin(Action, SourcePin);
 
 	if (AddedPins.Num() == 0)
 	{
@@ -686,7 +678,7 @@ bool FOptimusNodeGraphAction_ConnectAdderPin::Do(IOptimusPathResolver* InRoot)
 		return false;
 	}
 
-	if (bShouldLink)
+	if (Action.bCanAutoLink)
 	{
 		if (SourcePin->GetDirection() == EOptimusNodePinDirection::Output)
 		{
@@ -720,7 +712,7 @@ bool FOptimusNodeGraphAction_ConnectAdderPin::Undo(IOptimusPathResolver* InRoot)
 		return false;
 	}
 
-	if (bShouldLink)
+	if (Action.bCanAutoLink)
 	{
 		if (!ensure(RemoveLink(InRoot)))
 		{
diff --git a/Engine/Plugins/Animation/DeformerGraph/Source/OptimusCore/Private/Actions/OptimusNodeGraphActions.h b/Engine/Plugins/Animation/DeformerGraph/Source/OptimusCore/Private/Actions/OptimusNodeGraphActions.h
--- a/Engine/Plugins/Animation/DeformerGraph/Source/OptimusCore/Private/Actions/OptimusNodeGraphActions.h
+++ b/Engine/Plugins/Animation/DeformerGraph/Source/OptimusCore/Private/Actions/OptimusNodeGraphActions.h
@@ -6,6 +6,7 @@
 #include "OptimusDataDomain.h"
 #include "OptimusBindingTypes.h"
 #include "Nodes/OptimusNode_ComputeKernelBase.h"
+#include "IOptimusNodeAdderPinProvider.h"
 
 #include "UObject/UnrealNames.h"
 
@@ -366,9 +367,8 @@ struct FOptimusNodeGraphAction_ConnectAdderPin
 
 	FOptimusNodeGraphAction_ConnectAdderPin(
 		IOptimusNodeAdderPinProvider* InAdderPinProvider,
-		UOptimusNodePin* InPreferredTargetParentPin,
-		UOptimusNodePin* InSourcePin,
-		bool bInShouldLink
+		const IOptimusNodeAdderPinProvider::FAdderPinAction& InAction,
+		UOptimusNodePin* InSourcePin
 		);
 
 protected:
@@ -378,13 +378,11 @@ protected:
 private:
 	FString NodePath;
 
-	FString PreferredParentPinPath;
+	IOptimusNodeAdderPinProvider::FAdderPinAction Action;	
 	
 	FString SourcePinPath;
 
 	TArray<FString> AddedPinPaths;
-
-	bool bShouldLink;
 };
 
 USTRUCT()
diff --git a/Engine/Plugins/Animation/DeformerGraph/Source/OptimusCore/Private/Nodes/OptimusNode_CustomComputeKernel.cpp b/Engine/Plugins/Animation/DeformerGraph/Source/OptimusCore/Private/Nodes/OptimusNode_CustomComputeKernel.cpp
--- a/Engine/Plugins/Animation/DeformerGraph/Source/OptimusCore/Private/Nodes/OptimusNode_CustomComputeKernel.cpp
+++ b/Engine/Plugins/Animation/DeformerGraph/Source/OptimusCore/Private/Nodes/OptimusNode_CustomComputeKernel.cpp
@@ -19,12 +19,17 @@
 
 #include UE_INLINE_GENERATED_CPP_BY_NAME(OptimusNode_CustomComputeKernel)
 
+#define LOCTEXT_NAMESPACE "OptimusNode_CustomComputeKernel"
+
 static const FName DefaultKernelName("MyKernel");
 static const FName DefaultGroupName("Group"); 
 static const FName InputBindingsName= GET_MEMBER_NAME_CHECKED(UOptimusNode_CustomComputeKernel, InputBindingArray);
 static const FName OutputBindingsName = GET_MEMBER_NAME_CHECKED(UOptimusNode_CustomComputeKernel, OutputBindingArray);
 static const FName ExtraInputBindingGroupsName = GET_MEMBER_NAME_CHECKED(UOptimusNode_CustomComputeKernel, SecondaryInputBindingGroups);
 
+static const FName AdderPinActionKeyAddToPrimaryGroup("@Primary");
+static const FName AdderPinActionKeyAddNewGroup("@NewGroup");
+
 static FString GetShaderValueTypeFriendlyName(const FOptimusDataTypeRef& InDataType)
 {
 	FName FriendlyName = NAME_None;
@@ -51,6 +56,24 @@ static bool IsGroupingInputPin(const UOptimusNodePin *InPin)
 	return InPin->GetDirection() == EOptimusNodePinDirection::Input && InPin->IsGroupingPin();
 }
 
+bool UOptimusNode_CustomComputeKernel::DoesSourceSupportUnifiedDispatch(const UOptimusNodePin& InOtherNodesPin)
+{
+	if (InOtherNodesPin.GetDirection() == EOptimusNodePinDirection::Output )
+	{
+		if (const UOptimusNode_DataInterface* DataInterfaceNode = Cast<UOptimusNode_DataInterface>(InOtherNodesPin.GetOwningNode()))
+		{
+			if (const UOptimusComputeDataInterface* DataInterface = DataInterfaceNode->GetDataInterface(GetTransientPackage()))
+			{
+				if (!DataInterface->CanSupportUnifiedDispatch())
+				{
+					return false;
+				}
+			}
+		}	
+	}
+
+	return true;
+}
 
 
 UOptimusNode_CustomComputeKernel::UOptimusNode_CustomComputeKernel()
@@ -132,18 +155,27 @@ FString UOptimusNode_CustomComputeKernel::GetBindingDeclaration(
 }
 
 
-bool UOptimusNode_CustomComputeKernel::CanAddPinFromPin(
+TArray<IOptimusNodeAdderPinProvider::FAdderPinAction> UOptimusNode_CustomComputeKernel::GetAvailableAdderPinActions(
 	const UOptimusNodePin* InSourcePin,
 	EOptimusNodePinDirection InNewPinDirection,
 	FString* OutReason
 	) const
 {
-	if (!CanConnectPinToNode(InSourcePin, InNewPinDirection, OutReason))
+	TSet<UOptimusComponentSourceBinding*> SourceComponentBindings = InSourcePin->GetComponentSourceBindings();
+	
+	TArray<UOptimusNodePin*> GroupPins;
+	// Nullptr represents the primary group
+	GroupPins.Add(nullptr);
+
+	for (UOptimusNodePin* GroupPin: GetPins())
 	{
-		return false;
+		if (IsGroupingInputPin(GroupPin))
+		{
+			GroupPins.Add(GroupPin);
+		}
 	}
-	
-	if (InSourcePin->GetDirection() == EOptimusNodePinDirection::Input)
+
+	if (InNewPinDirection == EOptimusNodePinDirection::Output)
 	{
 		if (InSourcePin->GetDataDomain().IsSingleton())
 		{
@@ -152,157 +184,250 @@ bool UOptimusNode_CustomComputeKernel::CanAddPinFromPin(
 				*OutReason = TEXT("Can't add parameter pin as output");
 			}
 
-			return false;
+			return {};
 		}
-	}
 
-	return true;
-}
-
-TArray<UOptimusNodePin*> UOptimusNode_CustomComputeKernel::GetTargetParentPins(const UOptimusNodePin* InSourcePin) const
-{
-	TArray<UOptimusNodePin*> TargetParentPins;
-	
-	if (InSourcePin->GetDirection() == EOptimusNodePinDirection::Input)
-	{
-		return {};
-	}
-	
-	TArray<UOptimusNodePin*> GroupPinsToCheck;
-	// Nullptr represents the primary group
-	GroupPinsToCheck.Add(nullptr);
+		FAdderPinAction	Action;
+		Action.DisplayName = TEXT("Primary Group");
+		Action.NewPinDirection = InNewPinDirection;
+		Action.Key = AdderPinActionKeyAddToPrimaryGroup;
+		Action.ToolTip = FText::GetEmpty();
+		Action.bCanAutoLink = false;
 
-	for (UOptimusNodePin* GroupPin: GetPins())
-	{
-		if (IsGroupingInputPin(GroupPin))
+		if (InSourcePin->GetDirection() != InNewPinDirection)
 		{
-			GroupPinsToCheck.Add(GroupPin);
+			TArray<const UOptimusNodePin*> PrimaryGroupPins = GetPrimaryGroupInputPins();
+			TSet<UOptimusComponentSourceBinding*> PrimaryGroupComponentBindings;
+			for (const UOptimusNodePin* Pin : PrimaryGroupPins)
+			{
+				PrimaryGroupComponentBindings.Append(Pin->GetComponentSourceBindings());
+			}
+		
+			TSet<UOptimusComponentSourceBinding*> MatchingBindings = PrimaryGroupComponentBindings.Intersect(SourceComponentBindings);
+		
+			if (MatchingBindings.Num() > 0 ||
+				SourceComponentBindings.Num() == 0 || // Source-less pin can be added to kernel
+				PrimaryGroupComponentBindings.Num() == 0) // Source-less kernel can be a target for any pin 
+			{
+				Action.bCanAutoLink = true;
+			}
 		}
+
+		return { Action };
 	}
-	
-	TSet<UOptimusComponentSourceBinding*> SourceComponentBindings = InSourcePin->GetComponentSourceBindings();
-	
-	for (UOptimusNodePin* GroupPin : GroupPinsToCheck)
+	else if (InNewPinDirection == EOptimusNodePinDirection::Input)
 	{
-		TArray<const UOptimusNodePin*> PinsPerGroup;
-		
-		if (GroupPin)
-		{
-			PinsPerGroup = GroupPin->GetSubPins();
-		}
-		else
+		if (InSourcePin->GetDirection() == EOptimusNodePinDirection::Input)
 		{
-			PinsPerGroup = GetPrimaryGroupInputPins();
+			// For convenience lets allow input pins to be created from other input pins 
+			// of course linking is not allowed
+			TArray<FAdderPinAction> Actions;
+			
+			for (UOptimusNodePin* GroupPin: GroupPins)
+			{
+				FAdderPinAction Action;
+				Action.NewPinDirection = InNewPinDirection;
+				if (GroupPin == nullptr)
+				{
+					Action.DisplayName = TEXT("Primary Group");
+					Action.Key = AdderPinActionKeyAddToPrimaryGroup;
+				}
+				else
+				{
+					Action.DisplayName = GroupPin->GetFName();
+					Action.Key = GroupPin->GetFName();
+				}
+
+				Action.bCanAutoLink = false;
+				Action.ToolTip = LOCTEXT("AdderPinToolTipDirectionMismatch", "Add pin only, can't link input to input");
+				Actions.Add(Action);
+			}
+
+			return Actions;
 		}
-		
-		TSet<UOptimusComponentSourceBinding*> ComponentBindingsForGroup;
-		for (const UOptimusNodePin* Pin : PinsPerGroup)
+		else if (InSourcePin->GetDirection() == EOptimusNodePinDirection::Output)
 		{
-			ComponentBindingsForGroup.Append(Pin->GetComponentSourceBindings());
-		}
+			TArray<FAdderPinAction> Actions;
+			
+			const bool bSourceSupportUnifiedDispatch = DoesSourceSupportUnifiedDispatch(*InSourcePin);
+			
+			for (UOptimusNodePin* GroupPin: GroupPins)
+			{
+				FAdderPinAction Action;
+				Action.NewPinDirection = InNewPinDirection;
+				TArray<const UOptimusNodePin*> PinsPerGroup;
+
+				if (GroupPin == nullptr)
+				{
+					Action.DisplayName = TEXT("Primary Group");
+					Action.Key = AdderPinActionKeyAddToPrimaryGroup;
+					
+					PinsPerGroup = GetPrimaryGroupInputPins();
+				}
+				else 
+				{
+					Action.DisplayName = GroupPin->GetFName();
+					Action.Key = GroupPin->GetFName();
+					
+					PinsPerGroup = GroupPin->GetSubPins();
+				}
+
+				TSet<UOptimusComponentSourceBinding*> ComponentBindingsForGroup;
+				for (const UOptimusNodePin* Pin : PinsPerGroup)
+				{
+					ComponentBindingsForGroup.Append(Pin->GetComponentSourceBindings());
+				}
 
-		TSet<UOptimusComponentSourceBinding*> MatchingBindings = ComponentBindingsForGroup.Intersect(SourceComponentBindings);
+				TSet<UOptimusComponentSourceBinding*> MatchingBindings = ComponentBindingsForGroup.Intersect(SourceComponentBindings);
 	
-		if (MatchingBindings.Num() > 0 ||
-			SourceComponentBindings.Num() == 0 || // Source-less pin can be added to any group
-			ComponentBindingsForGroup.Num() == 0) // Source-less group can be a target for any pin 
-		{
-			TargetParentPins.Add(GroupPin);
+				if (MatchingBindings.Num() > 0 ||
+					SourceComponentBindings.Num() == 0 || // Source-less pin can be added to any group
+					ComponentBindingsForGroup.Num() == 0) // Source-less group can be a target for any pin 
+				{
+					Action.bCanAutoLink = true;
+				}
+				else
+				{
+					Action.bCanAutoLink = false;
+					Action.ToolTip = LOCTEXT("AdderPinToolTipComponentMismatch", "Add pin only, can't auto-link due to group's component binding not matching that of the source pin");
+				}
+
+				// Extra restrictions
+				// Secondary group does not accept data interface that does not support non unified dispatch
+				if (GroupPin != nullptr && !bSourceSupportUnifiedDispatch)
+				{
+					Action.bCanAutoLink = false;
+					Action.ToolTip = LOCTEXT("AdderPinToolTipCannotLinkNonUnified", "Add pin only, can't link due to data interface not supporting unified dispatch, consider insert kernel inbetween");
+				}
+				
+				Actions.Add(Action);
+			}
+
+			bool bHasAutoLinkAction = false;
+			for (const FAdderPinAction& Action : Actions)
+			{
+				if (Action.bCanAutoLink)
+				{
+					bHasAutoLinkAction = true;
+					break;
+				}
+			}
+			
+			if (!bHasAutoLinkAction && bSourceSupportUnifiedDispatch)
+			{
+				FAdderPinAction AddNewGroupAction;
+				AddNewGroupAction.DisplayName = TEXT("Add a New Group");
+				AddNewGroupAction.NewPinDirection = InNewPinDirection;
+				AddNewGroupAction.Key = AdderPinActionKeyAddNewGroup;
+				AddNewGroupAction.bCanAutoLink = true;
+				
+				Actions.Add(AddNewGroupAction);
+			}
+
+			for (FAdderPinAction& Action : Actions)
+			{
+				if (Action.bCanAutoLink)
+				{
+					Action.ToolTip = LOCTEXT("AdderPinToolTipAutoLink", "Group accepts pin connection");
+					break;
+				}
+			}
+			
+			return Actions;	
 		}
 	}
 
-	return TargetParentPins;
+	return {};
 }
 
-
 TArray<UOptimusNodePin*> UOptimusNode_CustomComputeKernel::TryAddPinFromPin(
-	UOptimusNodePin* InPreferredTargetParentPin,
+	const FAdderPinAction& InSelectedAction,
 	UOptimusNodePin* InSourcePin
-	) 
+	)
 {
 	TArray<UOptimusNodePin*> AddedPins;
 	
-	const EOptimusNodePinDirection NewPinDirection =
-		InSourcePin->GetDirection() == EOptimusNodePinDirection::Input ?
-				EOptimusNodePinDirection::Output : EOptimusNodePinDirection::Input;
-
 	TArray<FOptimusParameterBinding>* BindingArray = nullptr;
 	UOptimusNodePin* ParentPin = nullptr;
 	UOptimusNodePin* BeforePin = nullptr;
 
-	if (NewPinDirection == EOptimusNodePinDirection::Output)
+	if (InSelectedAction.NewPinDirection == EOptimusNodePinDirection::Output)
 	{
 		BindingArray = &OutputBindingArray.InnerArray;
 	}
-	else if (NewPinDirection == EOptimusNodePinDirection::Input)
+	else if (InSelectedAction.NewPinDirection == EOptimusNodePinDirection::Input)
 	{
-		if (InPreferredTargetParentPin)
+		if (InSelectedAction.Key == AdderPinActionKeyAddNewGroup)
 		{
-			// Making sure we are connecting from the input side of the kernel
-			check(InSourcePin->GetDirection() == EOptimusNodePinDirection::Output);
-		
-			FOptimusSecondaryInputBindingsGroup* Group = SecondaryInputBindingGroups.FindByPredicate(
-				[InPreferredTargetParentPin](const FOptimusSecondaryInputBindingsGroup& InGroup)
-				{
-					return InGroup.GroupName == InPreferredTargetParentPin->GetFName();
-				});
+			//Creating a new group
+			
+			TSet<UOptimusComponentSourceBinding*> ComponentBindings = InSourcePin->GetComponentSourceBindings();
+			
+			check(ComponentBindings.Num() != 0);
 
-			if (!ensure(Group))
-			{
-				return {};
-			}
+			UOptimusComponentSourceBinding* ComponentBinding = *ComponentBindings.begin();
 
-			ParentPin = InPreferredTargetParentPin;
-			BindingArray = &Group->BindingArray.InnerArray;
+			FName GroupName = Optimus::GetUniqueNameForScope(this, ComponentBinding->GetFName());
+
+			SecondaryInputBindingGroups.AddDefaulted_GetRef().GroupName = GroupName;
+			BindingArray = &SecondaryInputBindingGroups.Last().BindingArray.InnerArray;
+			
+			AddedPins.Add(AddGroupingPinDirect(GroupName, EOptimusNodePinDirection::Input));
+			ParentPin = AddedPins.Last();	
 		}
 		else
 		{
-			// In case there is no preferred parent pin, we infer the group pin from the source pin
-			// Two cases:
-			//		1. primary group
-			//		2. create a new group
-			TArray<UOptimusNodePin*> ParentPins = GetTargetParentPins(InSourcePin);
+			// Adding to existing group
 			
-			if (ParentPins.Num() >= 1)
+			TArray<UOptimusNodePin*> GroupPins;
+			for (UOptimusNodePin* GroupPin: GetPins())
+			{
+				if (IsGroupingInputPin(GroupPin))
+				{
+					GroupPins.Add(GroupPin);
+				}
+			}	
+			
+			if (InSelectedAction.Key == AdderPinActionKeyAddToPrimaryGroup)
 			{
-				// multiple parent pins but no preferred parent pin implies we want to add to the primary group,
-				// nullptr group pin means the matching group is the primary group, which does not have a grouping pin
-				check(ParentPins[0] == nullptr);
-
 				BindingArray = &InputBindingArray.InnerArray;
-
-				for (UOptimusNodePin* Pin : GetPins())
+				
+				ParentPin = nullptr;
+				BeforePin = GroupPins.Num() == 0 ? nullptr : GroupPins[0];
+			}
+			else
+			{
+				for (UOptimusNodePin* GroupPin : GroupPins)
 				{
-					if (Pin->IsGroupingPin())
+					if (GroupPin->GetFName() == InSelectedAction.Key)
 					{
-						BeforePin = Pin;
+						ParentPin = GroupPin;
 						break;
 					}
 				}
-			}
-			else
-			{
-				// This means the source pin does not belong to any existing group, so create a new group
-				TSet<UOptimusComponentSourceBinding*> ComponentBindings = InSourcePin->GetComponentSourceBindings();
-				
-				check(ComponentBindings.Num() != 0);
-
-				UOptimusComponentSourceBinding* ComponentBinding = *ComponentBindings.begin();
 
-				FName GroupName = Optimus::GetUniqueNameForScope(this, ComponentBinding->GetFName());
+				if (!ensure(ParentPin))
+				{
+					return {};
+				}
+				
+				// Secondary group
+				FOptimusSecondaryInputBindingsGroup* Group = SecondaryInputBindingGroups.FindByPredicate(
+				[ParentPin](const FOptimusSecondaryInputBindingsGroup& InGroup)
+				{
+					return InGroup.GroupName == ParentPin->GetFName();
+				});
 
-				SecondaryInputBindingGroups.AddDefaulted_GetRef().GroupName = GroupName;
-				BindingArray = &SecondaryInputBindingGroups.Last().BindingArray.InnerArray;
+				if (!ensure(Group))
+				{
+					return {};
+				}
 				
-				AddedPins.Add(AddGroupingPinDirect(GroupName, EOptimusNodePinDirection::Input));
-				ParentPin = AddedPins.Last();
+				BindingArray = &Group->BindingArray.InnerArray;
 			}
-		}	
-	}
-	else
-	{
-		checkNoEntry();
-	}
+		}
+	}	
+
 
 	if (ParentPin)
 	{
@@ -313,12 +438,11 @@ TArray<UOptimusNodePin*> UOptimusNode_CustomComputeKernel::TryAddPinFromPin(
 	Binding.Name = GetSanitizedNewPinName(ParentPin, InSourcePin->GetFName());
 	Binding.DataType = {InSourcePin->GetDataType()};
 	Binding.DataDomain = InSourcePin->GetDataDomain();
-
 	
 	BindingArray->Add(Binding);
 	UpdatePreamble();
 	
-	AddedPins.Add(AddPinDirect(Binding.Name, NewPinDirection, Binding.DataDomain, Binding.DataType, BeforePin, ParentPin));
+	AddedPins.Add(AddPinDirect(Binding.Name, InSelectedAction.NewPinDirection, Binding.DataDomain, Binding.DataType, BeforePin, ParentPin));
 
 	return AddedPins;
 }
@@ -1305,6 +1429,21 @@ bool UOptimusNode_CustomComputeKernel::ValidateConnection(
 	FString* OutReason
 	) const
 {
+	if (InThisNodesPin.GetDirection() == EOptimusNodePinDirection::Input)
+	{
+		// Secondary group connections need additional checks
+		if (IsGroupingInputPin(InThisNodesPin.GetRootPin()))
+		{
+			if (!DoesSourceSupportUnifiedDispatch(InOtherNodesPin))
+			{
+				if (OutReason)
+				{
+					*OutReason = TEXT("Can't link due to data interface not supporting unified dispatch, consider insert kernel inbetween");
+				}
+				return false;
+			}
+		}
+	}
 	
 	return true;
 }
@@ -1488,3 +1627,5 @@ FString UOptimusNode_CustomComputeKernel::GetDeclarationForBinding(const FOptimu
 					*Binding.Name.ToString(), *FString::Join(Indexes, TEXT(", ")), *GetShaderValueTypeFriendlyName(Binding.DataType));
 	}
 }
+
+#undef LOCTEXT_NAMESPACE
\ No newline at end of file
diff --git a/Engine/Plugins/Animation/DeformerGraph/Source/OptimusCore/Private/Nodes/OptimusNode_CustomComputeKernel.h b/Engine/Plugins/Animation/DeformerGraph/Source/OptimusCore/Private/Nodes/OptimusNode_CustomComputeKernel.h
--- a/Engine/Plugins/Animation/DeformerGraph/Source/OptimusCore/Private/Nodes/OptimusNode_CustomComputeKernel.h
+++ b/Engine/Plugins/Animation/DeformerGraph/Source/OptimusCore/Private/Nodes/OptimusNode_CustomComputeKernel.h
@@ -75,16 +75,14 @@ public:
 	FString GetBindingDeclaration(FName BindingName) const override;
 
 	// IOptimusNodeAdderPinProvider
-	bool CanAddPinFromPin(
+	TArray<FAdderPinAction> GetAvailableAdderPinActions(
 		const UOptimusNodePin* InSourcePin,
 		EOptimusNodePinDirection InNewPinDirection,
 		FString* OutReason = nullptr
 		) const override;
 
-	TArray<UOptimusNodePin*> GetTargetParentPins(const UOptimusNodePin* InSourcePin) const override;
-
 	TArray<UOptimusNodePin*> TryAddPinFromPin(
-		UOptimusNodePin* InPreferredTargetParentPin,
+		const FAdderPinAction& InSelectedAction,
 		UOptimusNodePin* InSourcePin
 		) override;
 	
@@ -173,7 +171,7 @@ protected:
 	TOptional<FText> ValidateForCompile() const override;
 	
 private:
-	
+
 #if WITH_EDITOR
 	void PropertyArrayPasted(const FPropertyChangedEvent& InPropertyChangedEvent);
 	void PropertyValueChanged(const FPropertyChangedEvent& InPropertyChangedEvent);
@@ -186,4 +184,6 @@ private:
 	void UpdatePreamble();
 
 	static FString GetDeclarationForBinding(const FOptimusParameterBinding& Binding, bool bIsInput);
+
+	static bool DoesSourceSupportUnifiedDispatch(const UOptimusNodePin& InOtherNodesPin);
 };
diff --git a/Engine/Plugins/Animation/DeformerGraph/Source/OptimusCore/Private/OptimusNodeGraph.cpp b/Engine/Plugins/Animation/DeformerGraph/Source/OptimusCore/Private/OptimusNodeGraph.cpp
--- a/Engine/Plugins/Animation/DeformerGraph/Source/OptimusCore/Private/OptimusNodeGraph.cpp
+++ b/Engine/Plugins/Animation/DeformerGraph/Source/OptimusCore/Private/OptimusNodeGraph.cpp
@@ -719,9 +719,9 @@ bool UOptimusNodeGraph::RemoveAllLinks(UOptimusNodePin* InNodePin)
 	return GetActionStack()->RunAction(Action);
 }
 
-bool UOptimusNodeGraph::AddPinAndLink(
-	UOptimusNode* InTargetNode,
-	UOptimusNodePin* InPreferredTargetParentPin,
+bool UOptimusNodeGraph::ConnectAdderPin(
+	IOptimusNodeAdderPinProvider* InTargetNode,
+	const IOptimusNodeAdderPinProvider::FAdderPinAction& InSelectedAction,
 	UOptimusNodePin* InSourcePin
 	)
 {
@@ -742,29 +742,22 @@ bool UOptimusNodeGraph::AddPinAndLink(
 		return false;
 	}
 
-	const EOptimusNodePinDirection PinDirection =
-	InSourcePin->GetDirection() == EOptimusNodePinDirection::Input ?
-		EOptimusNodePinDirection::Output : EOptimusNodePinDirection::Input;
-	if (!AdderPinProvider->CanAddPinFromPin(InSourcePin, PinDirection))
-	{
-		return false;
-	}
+	
+	IOptimusNodeAdderPinProvider::FAdderPinAction FinalAction = InSelectedAction;
 
-	// For now lets default to not disconnect exisiting link on the source pin
+	// For now lets default to not disconnect existing link on the source pin
 	// Additional logic be added in the future if we want to let the user decide
-	bool bShouldLink = true;
 	if (InSourcePin->GetDirection() == EOptimusNodePinDirection::Input)
 	{
 		// Check to see if there's an existing link on the _input_ pin.
 		TArray<int32> PinLinks = GetAllLinkIndexesToPin(InSourcePin);
 		if (PinLinks.Num() > 0)
 		{
-			
-			bShouldLink = false;
+			FinalAction.bCanAutoLink = false;
 		}
 	}
 
-	return GetActionStack()->RunAction<FOptimusNodeGraphAction_ConnectAdderPin>(AdderPinProvider, InPreferredTargetParentPin, InSourcePin, bShouldLink);
+	return GetActionStack()->RunAction<FOptimusNodeGraphAction_ConnectAdderPin>(AdderPinProvider, FinalAction, InSourcePin);
 }
 
 UOptimusNode* UOptimusNodeGraph::ConvertCustomKernelToFunction(UOptimusNode* InCustomKernel)
@@ -1405,11 +1398,15 @@ TSet<UOptimusComponentSourceBinding*> UOptimusNodeGraph::GetComponentSourceBindi
 		TArray<FOptimusRoutedNodePin> RoutedPins = GetConnectedPinsWithRouting(InNodePin, {});
 		if (RoutedPins.IsEmpty())
 		{
-			return {};
+			// this can happen if the pin is directly on a compute kernel and not connected to anything
+			StartNode = InNodePin->GetOwningNode();
+		}
+		else
+		{
+			check(RoutedPins.Num() == 1);
+			StartNode = RoutedPins[0].NodePin->GetOwningNode();
 		}
 		
-		check(RoutedPins.Num() == 1);
-		StartNode = RoutedPins[0].NodePin->GetOwningNode();
 	}
 	else
 	{
diff --git a/Engine/Plugins/Animation/DeformerGraph/Source/OptimusCore/Private/OptimusNodePin.cpp b/Engine/Plugins/Animation/DeformerGraph/Source/OptimusCore/Private/OptimusNodePin.cpp
--- a/Engine/Plugins/Animation/DeformerGraph/Source/OptimusCore/Private/OptimusNodePin.cpp
+++ b/Engine/Plugins/Animation/DeformerGraph/Source/OptimusCore/Private/OptimusNodePin.cpp
@@ -455,6 +455,12 @@ bool UOptimusNodePin::CanCannect(const UOptimusNodePin* InOtherPin, FString* Out
 		return false;
 	}
 
+	// Check connection at the node level on the other node
+	if (!InOtherPin->GetOwningNode()->CanConnectPinToPin(*InOtherPin, *this, OutReason))
+	{
+		return false;
+	}
+
 	// Check for incompatible types.
 	if (!((DataType->ShaderValueType.IsValid() && InOtherPin->DataType->ShaderValueType.IsValid() &&
 		  DataType->ShaderValueType == InOtherPin->DataType->ShaderValueType) ||
diff --git a/Engine/Plugins/Animation/DeformerGraph/Source/OptimusCore/Public/IOptimusNodeAdderPinProvider.h b/Engine/Plugins/Animation/DeformerGraph/Source/OptimusCore/Public/IOptimusNodeAdderPinProvider.h
--- a/Engine/Plugins/Animation/DeformerGraph/Source/OptimusCore/Public/IOptimusNodeAdderPinProvider.h
+++ b/Engine/Plugins/Animation/DeformerGraph/Source/OptimusCore/Public/IOptimusNodeAdderPinProvider.h
@@ -15,25 +15,30 @@ class OPTIMUSCORE_API UOptimusNodeAdderPinProvider :
 };
 
 /**
-* Interface that provides a mechanism to query information about parameter bindings
+* Interface that provides a mechanism to add pins to node from existing pins
 */
 class OPTIMUSCORE_API IOptimusNodeAdderPinProvider
 {
 	GENERATED_BODY()
 
 public:
-	virtual bool CanAddPinFromPin(
+	struct FAdderPinAction
+	{
+		FName DisplayName;
+		EOptimusNodePinDirection NewPinDirection;
+		FName Key;
+		FText ToolTip;
+		bool bCanAutoLink;
+	};
+	
+	virtual TArray<FAdderPinAction> GetAvailableAdderPinActions(
 		const UOptimusNodePin* InSourcePin,
 		EOptimusNodePinDirection InNewPinDirection,
 		FString* OutReason = nullptr
 		) const = 0;
-
-	// If there are multiple options for the target parent pin,
-	// a menu should show up for the user to choose one from these options
-	virtual TArray<UOptimusNodePin*> GetTargetParentPins(const UOptimusNodePin* InSourcePin) const = 0;
 	
 	virtual TArray<UOptimusNodePin*> TryAddPinFromPin(
-		UOptimusNodePin* InPreferredTargetParentPin,
+		const FAdderPinAction& InSelectedAction,
 		UOptimusNodePin* InSourcePin
 		) = 0;
 	
diff --git a/Engine/Plugins/Animation/DeformerGraph/Source/OptimusCore/Public/OptimusNodeGraph.h b/Engine/Plugins/Animation/DeformerGraph/Source/OptimusCore/Public/OptimusNodeGraph.h
--- a/Engine/Plugins/Animation/DeformerGraph/Source/OptimusCore/Public/OptimusNodeGraph.h
+++ b/Engine/Plugins/Animation/DeformerGraph/Source/OptimusCore/Public/OptimusNodeGraph.h
@@ -3,6 +3,7 @@
 #pragma once
 
 #include "IOptimusNodeGraphCollectionOwner.h"
+#include "IOptimusNodeAdderPinProvider.h"
 #include "OptimusCoreNotify.h"
 #include "OptimusDataType.h"
 
@@ -297,12 +298,12 @@ public:
 	/// Add a new pin to the target node with the type of source pin
 	/// and connect the source pin to the new pin
 	/// @param InTargetNode The node to add the pin to, it has to have an adder pin
-	/// @param InPreferredTargetParentPin The preferred pin to add the new pin to, or null if it can be inferred from source
+	/// @param InSelectedAction The selected action that defines where the new pin should be added
 	/// @param InSourcePin The pin to create the new pin and to connect to the new pin
-	/// @return True if new pin and the new link is created.
-	bool AddPinAndLink(
-		UOptimusNode* InTargetNode,
-		UOptimusNodePin* InPreferredTargetParentPin,
+	/// @return True if a new pin is created.
+	bool ConnectAdderPin(
+		IOptimusNodeAdderPinProvider* InTargetNode,
+		const IOptimusNodeAdderPinProvider::FAdderPinAction& InSelectedAction,
 		UOptimusNodePin* InSourcePin
 		);
 
diff --git a/Engine/Plugins/Animation/DeformerGraph/Source/OptimusEditor/Private/OptimusEditorGraphSchema.cpp b/Engine/Plugins/Animation/DeformerGraph/Source/OptimusEditor/Private/OptimusEditorGraphSchema.cpp
--- a/Engine/Plugins/Animation/DeformerGraph/Source/OptimusEditor/Private/OptimusEditorGraphSchema.cpp
+++ b/Engine/Plugins/Animation/DeformerGraph/Source/OptimusEditor/Private/OptimusEditorGraphSchema.cpp
@@ -26,7 +26,7 @@
 
 #define LOCTEXT_NAMESPACE "OptimusEditor"
 
-static bool IsValidAdderPinConnection(const UEdGraphPin* InPinA, const UEdGraphPin* InPinB, 
+static bool IsValidAdderPinConnection(const UEdGraphPin* InPinA, const UEdGraphPin* InPinB, TArray<IOptimusNodeAdderPinProvider::FAdderPinAction>* OutActions = nullptr,
 	FString* OutReason = nullptr)
 {
 	const bool bIsPinAAdderPin = OptimusEditor::IsAdderPin(InPinA);
@@ -67,9 +67,14 @@ static bool IsValidAdderPinConnection(const UEdGraphPin* InPinA, const UEdGraphP
 	const IOptimusNodeAdderPinProvider *AdderPinProvider = Cast<IOptimusNodeAdderPinProvider>(TargetNode);
 	check(AdderPinProvider);
 
-	const bool bCanConnect = AdderPinProvider->CanAddPinFromPin(SourcePin, NewPinDirection, OutReason);
+	if (OutActions)
+	{
+		*OutActions = AdderPinProvider->GetAvailableAdderPinActions(SourcePin, NewPinDirection, OutReason);
+		return OutActions->Num() > 0;
+	}
 	
-	return bCanConnect;
+	const TArray<IOptimusNodeAdderPinProvider::FAdderPinAction> Actions = AdderPinProvider->GetAvailableAdderPinActions(SourcePin, NewPinDirection, OutReason);
+	return Actions.Num() > 0;
 }
 
 UOptimusEditorGraphSchema::UOptimusEditorGraphSchema()
@@ -193,7 +198,8 @@ bool UOptimusEditorGraphSchema::TryCreateConnection(
 	// Pins might be connectable if one of the two pins is an adder pin;
 	if (OptimusEditor::IsAdderPin(InPinA) || OptimusEditor::IsAdderPin(InPinB))
 	{
-		const bool bCanConnect = IsValidAdderPinConnection(InPinA, InPinB);
+		TArray<IOptimusNodeAdderPinProvider::FAdderPinAction> AvailableActions;
+		const bool bCanConnect = IsValidAdderPinConnection(InPinA, InPinB, &AvailableActions);
 
 		if (!bCanConnect)
 		{
@@ -212,39 +218,38 @@ bool UOptimusEditorGraphSchema::TryCreateConnection(
 		}
 
 		IOptimusNodeAdderPinProvider *AdderPinProvider = Cast<IOptimusNodeAdderPinProvider>(TargetNode);
-		TArray<UOptimusNodePin*> ParentPins = AdderPinProvider->GetTargetParentPins(SourcePin);
-
-		if (ParentPins.Num() <= 1)
+		UOptimusNodeGraph *Graph = TargetNode->GetOwningGraph();
+		
+		if (AvailableActions.Num() == 1)
 		{
-			UOptimusNodeGraph *Graph = TargetNode->GetOwningGraph();
-			return Graph->AddPinAndLink(TargetNode, ParentPins.Num() == 0? nullptr : ParentPins[0], SourcePin);
+			return Graph->ConnectAdderPin(AdderPinProvider, AvailableActions[0], SourcePin);
 		}
 
 		FMenuBuilder MenuBuilder(true, NULL);
 
-		MenuBuilder.BeginSection(NAME_None, FText::FromString(TEXT("Add pin to")));
+		MenuBuilder.BeginSection(NAME_None, FText::FromString(TEXT("Add Pin To Group")));
 
-		for (UOptimusNodePin* ParentPin : ParentPins)
+		for (const IOptimusNodeAdderPinProvider::FAdderPinAction& Action : AvailableActions)
 		{
 			FText EntryName;
-			if (!ParentPin)
+
+			if (Action.bCanAutoLink)
 			{
-				EntryName = LOCTEXT("DefaultGroupName", "Default");
+				EntryName = FText::FromString(Action.DisplayName.ToString() + TEXT(" * "));
 			}
 			else
 			{
-				EntryName = FText::FromName(ParentPin->GetFName());
+				EntryName = FText::FromName(Action.DisplayName);
 			}
 			
 			MenuBuilder.AddMenuEntry(
 				EntryName,
-				FText::GetEmpty(),
+				Action.ToolTip,
 				FSlateIcon(),
 				FUIAction(
-				FExecuteAction::CreateLambda([TargetNode, ParentPin, SourcePin]()
+				FExecuteAction::CreateLambda([AdderPinProvider, Graph, Action, SourcePin]()
 					{
-						UOptimusNodeGraph *Graph = TargetNode->GetOwningGraph();
-						Graph->AddPinAndLink(TargetNode, ParentPin, SourcePin);
+						Graph->ConnectAdderPin(AdderPinProvider, Action, SourcePin);
 					}),
 					FCanExecuteAction()
 				));
@@ -295,7 +300,7 @@ const FPinConnectionResponse UOptimusEditorGraphSchema::CanCreateConnection(
 	if (OptimusEditor::IsAdderPin(InPinA) || OptimusEditor::IsAdderPin(InPinB))
 	{
 		FString FailureReason;
-		const bool bCanConnect = IsValidAdderPinConnection(InPinA, InPinB, &FailureReason);
+		const bool bCanConnect = IsValidAdderPinConnection(InPinA, InPinB, nullptr, &FailureReason);
 
 		return FPinConnectionResponse(
 			bCanConnect ? CONNECT_RESPONSE_MAKE :  CONNECT_RESPONSE_DISALLOW,
