Optimus: better validation of the compute graph to deal with ambiguity caused by non-unified vs unified dispatch

+ pins in the primary group of a kernel determines the dispatch type of the kernel

+ Added additional validation code for resource expressions. Expressions that evaluates to different results for unified/non-unified now invalidate the deformer.
   TODO: perform this validation during compile instead of runtime and throw a compile error

+ compile error if a data interface with no unified dispatch support is connected to secondary group

+ compile error if a kernel does not have execution data interface

+ support higher number dispatch group count, implying that in the long run we probably only support 1D dispatch of thread groups

diff --git a/Engine/Plugins/Animation/DeformerGraph/Source/OptimusCore/Private/DataInterfaces/OptimusDataInterfaceRawBuffer.cpp b/Engine/Plugins/Animation/DeformerGraph/Source/OptimusCore/Private/DataInterfaces/OptimusDataInterfaceRawBuffer.cpp
--- a/Engine/Plugins/Animation/DeformerGraph/Source/OptimusCore/Private/DataInterfaces/OptimusDataInterfaceRawBuffer.cpp
+++ b/Engine/Plugins/Animation/DeformerGraph/Source/OptimusCore/Private/DataInterfaces/OptimusDataInterfaceRawBuffer.cpp
@@ -234,28 +234,47 @@ bool UOptimusRawBufferDataProvider::GetLodAndInvocationElementCounts(
 					{
 						Count *= FMath::Max(1, DataDomain.Multiplier);
 					}
+					return true;
 				}
-				return true;
 			}
+
+			return false;
 			break;
 		}
 	case EOptimusDataDomainType::Expression:
 		{
 			TMap<FName, int32> EngineConstants;
 			TMap<FName, TArray<int32>> ElementCountsPerDomain;
-			int32 NumInvocations = 0;
 
+			int32 NumInvocations = -1;
 			for(FName ExecutionDomain: ComponentSource->GetExecutionDomains())
 			{
 				EngineConstants.Add(ExecutionDomain, 0);
 
 				TArray<int32>& ElementCounts = ElementCountsPerDomain.FindOrAdd(ExecutionDomain);
-				if (!ComponentSourcePtr->GetComponentElementCountsForExecutionDomain(ExecutionDomain, ComponentPtr, OutLodIndex, ElementCounts))
+				if (!ComponentSourcePtr->GetComponentElementCountsForExecutionDomain(
+						ExecutionDomain,
+						ComponentPtr,
+						OutLodIndex,
+						ElementCounts))
 				{
 					return false;
 				}
 
-				NumInvocations = FMath::Max(NumInvocations, ElementCounts.Num());
+				if (NumInvocations == -1)
+				{
+					NumInvocations = ElementCounts.Num();
+				}
+				else if (NumInvocations != ElementCounts.Num())
+				{
+					// Component source needs to provide as many values for each of its execution domains as there are invocations
+					return false;
+				}
+			}
+
+			if (NumInvocations < 1)
+			{
+				return false;
 			}
 			
 			using namespace Optimus::Expression;
@@ -270,16 +289,56 @@ bool UOptimusRawBufferDataProvider::GetLodAndInvocationElementCounts(
 			OutInvocationElementCounts.Reset(NumInvocations);
 			for (int32 Index = 0; Index < NumInvocations; Index++)
 			{
-				for (auto& [ConstantName, Value]: EngineConstants)
+				for (TPair<FName, int32>& Constant: EngineConstants)
 				{
+					const FName ConstantName = Constant.Key;
+					int32& Value = Constant.Value;
+					
 					const TArray<int32>& ElementCounts = ElementCountsPerDomain[ConstantName]; 
 					Value = ElementCounts.IsValidIndex(Index) ? ElementCounts[Index] : 1;
 				}
 				Engine.UpdateConstantValues(EngineConstants);
 				const int32 Count = Engine.Execute(ParseResult.Get<FExpressionObject>());
+
+				if (Count < 0)
+				{
+					return false;
+				}
+				
 				OutInvocationElementCounts.Add(Count);
 			}
+
+			int32 TotalElementCount = 0;
+			for (int32 Count : OutInvocationElementCounts)
+			{
+				TotalElementCount += Count;
+			}
+
+			// Extra validation for unified dispatch
+
+			// We want to make sure the results are the same for unified and non-unified
+			// Do the sum across invocations first and then evaluate the expression
+			for (auto& [ConstantName, Value]: EngineConstants)
+			{
+				Value = 0;
+				const TArray<int32>& ElementCounts = ElementCountsPerDomain[ConstantName];
+				
+				for (int32 Count : ElementCounts )
+				{
+					Value += Count;
+				}
+			}
+			
+			Engine.UpdateConstantValues(EngineConstants);
+			const int32 TotalElementCountForUnifiedDispatch = Engine.Execute(ParseResult.Get<FExpressionObject>());
+			
+			if (TotalElementCount != TotalElementCountForUnifiedDispatch)
+			{
+				return false;
+			}
+
 			return true;
+			break;
 		}
 	}
 	return false;
@@ -290,7 +349,10 @@ FComputeDataProviderRenderProxy* UOptimusTransientBufferDataProvider::GetRenderP
 {
 	int32 LodIndex;
 	TArray<int32> InvocationCounts;
-	GetLodAndInvocationElementCounts(LodIndex, InvocationCounts);
+	if (!GetLodAndInvocationElementCounts(LodIndex, InvocationCounts))
+	{
+		InvocationCounts.Reset();
+	}
 	
 	return new FOptimusTransientBufferDataProviderProxy(InvocationCounts, ElementStride, RawStride);
 }
@@ -300,7 +362,10 @@ FComputeDataProviderRenderProxy* UOptimusPersistentBufferDataProvider::GetRender
 {
 	int32 LodIndex = 0;
 	TArray<int32> InvocationCounts;
-	GetLodAndInvocationElementCounts(LodIndex, InvocationCounts);
+	if (!GetLodAndInvocationElementCounts(LodIndex, InvocationCounts))
+	{
+		InvocationCounts.Reset();
+	}
 	
 	return new FOptimusPersistentBufferDataProviderProxy(InvocationCounts, ElementStride, RawStride, BufferPool, ResourceName, LodIndex);
 }
@@ -309,11 +374,12 @@ FComputeDataProviderRenderProxy* UOptimusPersistentBufferDataProvider::GetRender
 FOptimusTransientBufferDataProviderProxy::FOptimusTransientBufferDataProviderProxy(
 	TArray<int32> InInvocationElementCounts,
 	int32 InElementStride,
-	int32 InRawStride) 
-	: InvocationElementCounts(InInvocationElementCounts)
-	, TotalElementCount(0)
-	, ElementStride(InElementStride)
-	, RawStride(InRawStride)
+	int32 InRawStride
+	) :
+	InvocationElementCounts(InInvocationElementCounts),
+	TotalElementCount(0),
+	ElementStride(InElementStride),
+	RawStride(InRawStride)
 {
 	for (int32 NumElements : InvocationElementCounts)
 	{
@@ -327,7 +393,8 @@ bool FOptimusTransientBufferDataProviderProxy::IsValid(FValidationData const& In
 	{
 		return false;
 	}
-	if (InValidationData.NumInvocations != InvocationElementCounts.Num())
+	
+	if (TotalElementCount <= 0)
 	{
 		return false;
 	}
@@ -369,14 +436,15 @@ FOptimusPersistentBufferDataProviderProxy::FOptimusPersistentBufferDataProviderP
 	int32 InRawStride,
 	TSharedPtr<FOptimusPersistentBufferPool> InBufferPool,
 	FName InResourceName,
-	int32 InLODIndex)
-	: InvocationElementCounts(InInvocationElementCounts)
-	, TotalElementCount(0)
-	, ElementStride(InElementStride)
-	, RawStride(InRawStride)
-	, BufferPool(InBufferPool)
-	, ResourceName(InResourceName)
-	, LODIndex(InLODIndex)
+	int32 InLODIndex
+	) : 
+	InvocationElementCounts(InInvocationElementCounts),
+	TotalElementCount(0),
+	ElementStride(InElementStride),
+	RawStride(InRawStride),
+	BufferPool(InBufferPool),
+	ResourceName(InResourceName),
+	LODIndex(InLODIndex)
 {
 	for (int32 NumElements : InvocationElementCounts)
 	{
@@ -390,7 +458,8 @@ bool FOptimusPersistentBufferDataProviderProxy::IsValid(FValidationData const& I
 	{
 		return false;
 	}
-	if (InValidationData.NumInvocations != InvocationElementCounts.Num())
+	
+	if (TotalElementCount <= 0)
 	{
 		return false;
 	}
diff --git a/Engine/Plugins/Animation/DeformerGraph/Source/OptimusCore/Private/Nodes/OptimusNode_ComputeKernelBase.cpp b/Engine/Plugins/Animation/DeformerGraph/Source/OptimusCore/Private/Nodes/OptimusNode_ComputeKernelBase.cpp
--- a/Engine/Plugins/Animation/DeformerGraph/Source/OptimusCore/Private/Nodes/OptimusNode_ComputeKernelBase.cpp
+++ b/Engine/Plugins/Animation/DeformerGraph/Source/OptimusCore/Private/Nodes/OptimusNode_ComputeKernelBase.cpp
@@ -269,7 +269,8 @@ TOptional<FText> UOptimusNode_ComputeKernelBase::ValidateForCompile() const
 			{
 				if (InCollectedBindings.Num() != PinBindings.Num() || !InCollectedBindings.Includes(PinBindings))
 				{
-					return FText::Format(LOCTEXT("IncompatibleBindingsPinGroup", "Component binding for pin '{0}' is different from the component bindings of the other pins in its group"), FText::FromName(InPin->GetUniqueName()));			}
+					return FText::Format(LOCTEXT("IncompatibleBindingsPinGroup", "Component binding for pin '{0}' is different from the component bindings of the other pins in its group"), FText::FromName(InPin->GetUniqueName()));
+				}
 			}
 		}
 		return {};
@@ -334,11 +335,18 @@ FString UOptimusNode_ComputeKernelBase::GetCookedKernelSource(
 	Optimus::ConvertObjectPathToShaderFilePath(ShaderPathName);
 
 	const bool bHasKernelKeyword = Source.Contains(TEXT("KERNEL"));
+
+	const FString ComputeShaderUtilsInclude = TEXT("#include \"/Engine/Private/ComputeShaderUtils.ush\"");
 	
 	const FString KernelFunc = FString::Printf(
-		TEXT("[numthreads(%d,%d,%d)]\nvoid %s(uint3 DTid : SV_DispatchThreadID)"), 
+		TEXT("[numthreads(%d,%d,%d)]\nvoid %s(uint3 GroupId : SV_GroupID, uint GroupIndex : SV_GroupIndex)"), 
 		InGroupSize.X, InGroupSize.Y, InGroupSize.Z, *InKernelName);
 	
+	const FString UnWrappedDispatchThreadId = FString::Printf(
+	TEXT("GetUnWrappedDispatchThreadId(GroupId, GroupIndex, %d)"),
+		InGroupSize.X * InGroupSize.Y * InGroupSize.Z
+	);
+
 	if (bHasKernelKeyword)
 	{
 		Source.ReplaceInline(TEXT("KERNEL"), TEXT("void __kernel_func(uint Index)"));
@@ -347,20 +355,22 @@ FString UOptimusNode_ComputeKernelBase::GetCookedKernelSource(
 			TEXT(
 				"#line 1 \"%s\"\n"
 				"%s\n\n"
-				"%s { __kernel_func(DTid.x); }\n"
-				), *ShaderPathName, *Source, *KernelFunc);
+				"%s\n\n"
+				"%s { __kernel_func(%s); }\n"
+				), *ShaderPathName, *Source, *ComputeShaderUtilsInclude,*KernelFunc, *UnWrappedDispatchThreadId);
 	}
 	else
 	{
 		return FString::Printf(
 		TEXT(
+			"%s\n"
 			"%s\n"
 			"{\n"
-			"uint Index = DTid.x;\n"
+			"uint Index = %s;\n"
 			"#line 1 \"%s\"\n"
 			"%s\n"
 			"}\n"
-			), *KernelFunc, *ShaderPathName, *Source);
+			), *ComputeShaderUtilsInclude,*KernelFunc, *UnWrappedDispatchThreadId, *ShaderPathName, *Source);
 	}
 }
 
@@ -450,6 +460,11 @@ TOptional<FText> UOptimusNode_ComputeKernelBase::ProcessInputPinForComputeKernel
 		// If we are connected from a data interface, set the input binding up now.
 		if (DataInterface && ensure(ComponentBinding))
 		{
+			if (!DataInterface->CanSupportUnifiedDispatch() && !InGroupName.IsEmpty())
+			{
+				return FText::Format(LOCTEXT("SecondaryGroupOnlySupportsDataInterfaceUnifiedDispatch", "Cannot connect Secondary group input binding pin '{0}' to a data interface that does not support unified dispatch, consider inserting a kernel inbetween"), FText::FromName(InInputPin->GetUniqueName()));
+			}
+
 			// The shader function definition that exposes the function that we use to
 			// read values to input into the kernel.
 			FShaderFunctionDefinition FuncDef;
diff --git a/Engine/Plugins/Animation/DeformerGraph/Source/OptimusCore/Private/OptimusDeformer.cpp b/Engine/Plugins/Animation/DeformerGraph/Source/OptimusCore/Private/OptimusDeformer.cpp
--- a/Engine/Plugins/Animation/DeformerGraph/Source/OptimusCore/Private/OptimusDeformer.cpp
+++ b/Engine/Plugins/Animation/DeformerGraph/Source/OptimusCore/Private/OptimusDeformer.cpp
@@ -1704,6 +1704,27 @@ UOptimusComputeGraph* UOptimusDeformer::CompileNodeGraphToComputeGraph(
 					ConnectedNode.Node);
 				return nullptr;
 			}
+
+			bool bHasExecution = false;
+			for (const TPair<int32, FOptimus_InterfaceBinding>& DataBinding: BoundKernel.InputDataBindings)
+			{
+				const int32 KernelBindingIndex = DataBinding.Key;
+				const FOptimus_InterfaceBinding& InterfaceBinding = DataBinding.Value;
+				const UComputeDataInterface* DataInterface = InterfaceBinding.DataInterface;
+				if (DataInterface->IsExecutionInterface())
+				{
+					bHasExecution = true;
+					break;
+				}
+			}
+			
+			if (!bHasExecution)
+			{
+				AddDiagnostic(EOptimusDiagnosticLevel::Error,
+				LOCTEXT("KernelHasNoExecutionDataInterface", "Kernel has no execution data interface connected. Compilation aborted."),
+					ConnectedNode.Node);
+				return nullptr;
+			}
 			
 			BoundKernel.Kernel->KernelSource = KernelSourceResult.Get<UOptimusKernelSource*>();
 
diff --git a/Engine/Plugins/Runtime/ComputeFramework/Source/ComputeFramework/Private/ComputeFramework/ComputeGraphRenderProxy.h b/Engine/Plugins/Runtime/ComputeFramework/Source/ComputeFramework/Private/ComputeFramework/ComputeGraphRenderProxy.h
--- a/Engine/Plugins/Runtime/ComputeFramework/Source/ComputeFramework/Private/ComputeFramework/ComputeGraphRenderProxy.h
+++ b/Engine/Plugins/Runtime/ComputeFramework/Source/ComputeFramework/Private/ComputeFramework/ComputeGraphRenderProxy.h
@@ -30,6 +30,8 @@ public:
 		FShaderParametersMetadata* ShaderParameterMetadata = nullptr;
 		/** Array of indices into the full graph data provider array. Contains only the indices to data providers that this kernel references. */
 		TArray<int32> BoundProviderIndices;
+		/** Same size as BoundProviderIndices, non-primary data providers will be forced to present a full view to its data (unified)*/
+		TArray<bool> BoundProviderIsPrimary;
 		/** The index of the special execution data provider in the full graph data provider array. */
 		int32 ExecutionProviderIndex = -1;
 		/** Whether the kernel can combine multiple supports combinined sub-invocations. */
diff --git a/Engine/Plugins/Runtime/ComputeFramework/Source/ComputeFramework/Private/ComputeGraph.cpp b/Engine/Plugins/Runtime/ComputeFramework/Source/ComputeFramework/Private/ComputeGraph.cpp
--- a/Engine/Plugins/Runtime/ComputeFramework/Source/ComputeFramework/Private/ComputeGraph.cpp
+++ b/Engine/Plugins/Runtime/ComputeFramework/Source/ComputeFramework/Private/ComputeGraph.cpp
@@ -315,24 +315,64 @@ FComputeGraphRenderProxy* UComputeGraph::CreateRenderProxy() const
 			Invocation.ShaderParameterMetadata = BuildKernelShaderMetadata(KernelIndex, *Proxy->ShaderParameterMetadataAllocations);
 			Invocation.bSupportsUnifiedDispatch = true;
 
+			// 1. Figure out which data interface is the execution data interface
 			for (FComputeGraphEdge const& GraphEdge : GraphEdges)
 			{
 				if (GraphEdge.KernelIndex == KernelIndex)
 				{
-					Invocation.BoundProviderIndices.AddUnique(GraphEdge.DataInterfaceIndex);
-
 					UComputeDataInterface const* DataInterface = DataInterfaces[GraphEdge.DataInterfaceIndex];
 					if (ensure(DataInterface != nullptr))
 					{
 						if (DataInterface->IsExecutionInterface())
 						{
 							Invocation.ExecutionProviderIndex = GraphEdge.DataInterfaceIndex;
+							break;
 						}
-						
-						Invocation.bSupportsUnifiedDispatch &= DataInterface->CanSupportUnifiedDispatch();
 					}
 				}
 			}
+
+			// 1. Data interfaces sharing the same binding (primary) as the kernel should present its data in a way that
+			// matches the kernel dispatch method, which can be either unified(full buffer) or non-unified (per invocation window into the full buffer)
+			// 2. Data interfaces not sharing the same binding (secondary) should always provide a full view to its data (unified)
+			// Note: In case of non-unified kernel, extra work maybe needed to read from secondary buffers.
+			// When kernel is non-unified, index = 0...section.max for each invocation/section, 
+			// so user may want to consider using a dummy buffer that maps section index to the indices of secondary buffers
+			// for example, given a non-unified kernel, primary and secondary components sharing the same vertex count, we might want to create a buffer
+			// in the primary group that is simply [0,1,2...,NumVerts-1], which we can then index into to map section vert index to the global vert index
+			if (ensure(DataInterfaceToBinding.IsValidIndex(Invocation.ExecutionProviderIndex)))
+			{
+				int32 ExecutionComponentBindingIndex = DataInterfaceToBinding[Invocation.ExecutionProviderIndex];
+
+				for (FComputeGraphEdge const& GraphEdge : GraphEdges)
+				{
+					if (GraphEdge.KernelIndex == KernelIndex)
+					{
+						UComputeDataInterface const* DataInterface = DataInterfaces[GraphEdge.DataInterfaceIndex];
+						if (ensure(DataInterface != nullptr))
+						{
+							const int32 DataInterfaceComponentBindingIndex = DataInterfaceToBinding[GraphEdge.DataInterfaceIndex];
+							const bool bIsPrimary = DataInterfaceComponentBindingIndex == ExecutionComponentBindingIndex;
+							
+							const int32 IndexOfIndex = Invocation.BoundProviderIndices.AddUnique(GraphEdge.DataInterfaceIndex);
+
+							// Added a new provider, store whether it is primary or secondary
+							if (IndexOfIndex == Invocation.BoundProviderIsPrimary.Num())
+							{
+								Invocation.BoundProviderIsPrimary.Add( bIsPrimary );
+							}	
+
+							// Only Data Interfaces in the primary group should determine the kernel dispatch type
+							if (bIsPrimary)
+							{
+								Invocation.bSupportsUnifiedDispatch &= DataInterface->CanSupportUnifiedDispatch();
+							}
+						}
+					}
+				}
+
+				check(Invocation.BoundProviderIndices.Num() == Invocation.BoundProviderIsPrimary.Num());
+			}
 		}
 	}
 
diff --git a/Engine/Plugins/Runtime/ComputeFramework/Source/ComputeFramework/Private/ComputeGraphWorker.cpp b/Engine/Plugins/Runtime/ComputeFramework/Source/ComputeFramework/Private/ComputeGraphWorker.cpp
--- a/Engine/Plugins/Runtime/ComputeFramework/Source/ComputeFramework/Private/ComputeGraphWorker.cpp
+++ b/Engine/Plugins/Runtime/ComputeFramework/Source/ComputeFramework/Private/ComputeGraphWorker.cpp
@@ -210,6 +210,15 @@ void FComputeGraphTaskWorker::SubmitWork(FRDGBuilder& GraphBuilder, FName InExec
 				FComputeDataProviderRenderProxy* DataProvider = GraphInvocation.DataProviderRenderProxies[DataProviderIndex];
 				if (ensure(DataProvider != nullptr))
 				{
+					// 1. Data interfaces sharing the same binding (primary) as the kernel should present its data in a way that
+					// matches the kernel dispatch method, which can be either unified(full buffer) or non-unified (per invocation window into the full buffer)
+					// 2. Data interfaces not sharing the same binding (secondary) should always provide a full view to its data (unified)
+					// Note: In case of non-unified kernel, extra work maybe needed to read from secondary buffers.
+					// When kernel is non-unified, index = 0...section.max for each invocation/section, 
+					// so user may want to consider using a dummy buffer that maps section index to the indices of secondary buffers
+					// for example, given a non-unified kernel, primary and secondary components sharing the same vertex count, we might want to create a buffer
+					// in the primary group that is simply [0,1,2...,NumVerts-1], which we can then index into to map section vert index to the global vert index
+					DispatchData.bUnifiedDispatch = KernelInvocation.BoundProviderIsPrimary[MemberIndex]? DispatchData.bUnifiedDispatch : true;
 					DispatchData.ParameterStructSize = Member.GetStructMetadata()->GetSize();
 					DispatchData.ParameterBufferOffset = Member.GetOffset();
 					DataProvider->GatherDispatchData(DispatchData);
@@ -221,8 +230,10 @@ void FComputeGraphTaskWorker::SubmitWork(FRDGBuilder& GraphBuilder, FName InExec
 		for (int32 SubInvocationIndex = 0; SubInvocationIndex < NumSubInvocations; ++SubInvocationIndex)
 		{
 			TShaderRef<FComputeKernelShader> Shader = Shaders[SubmitDesc.ShaderIndex + SubInvocationIndex];
-			const FIntVector GroupCount = FComputeShaderUtils::GetGroupCount(ThreadCounts[SubInvocationIndex], KernelInvocation.KernelGroupSize);
+			FIntVector GroupCount = FComputeShaderUtils::GetGroupCount(ThreadCounts[SubInvocationIndex], KernelInvocation.KernelGroupSize);
 
+			GroupCount = FComputeShaderUtils::GetGroupCountWrapped(GroupCount.X);
+			
 			FComputeShaderUtils::AddPass(
 				GraphBuilder,
 				{},
