Deformer Graph: Allow using vertex maps from skeletal meshes.

diff --git a/Engine/Plugins/Animation/DeformerGraph/Shaders/Private/DataInterfaceSkinnedMeshVertexAttribute.ush b/Engine/Plugins/Animation/DeformerGraph/Shaders/Private/DataInterfaceSkinnedMeshVertexAttribute.ush
new file mode 100644
--- /dev/null
+++ b/Engine/Plugins/Animation/DeformerGraph/Shaders/Private/DataInterfaceSkinnedMeshVertexAttribute.ush
@@ -0,0 +1,24 @@
+// Copyright Epic Games, Inc. All Rights Reserved.
+
+uint {DataInterfaceName}_NumVertices;
+uint {DataInterfaceName}_InputStreamStart;
+uint {DataInterfaceName}_bIsValid;
+Buffer<float> {DataInterfaceName}_ValueBuffer;
+
+uint ReadNumVertices_{DataInterfaceName}()
+{
+	return {DataInterfaceName}_NumVertices;
+}
+
+float ReadValue_{DataInterfaceName}(uint VertexIndex)
+{
+	if ({DataInterfaceName}_bIsValid)
+	{
+		uint BufferIndex = ({DataInterfaceName}_InputStreamStart + VertexIndex);
+		return {DataInterfaceName}_ValueBuffer[BufferIndex];
+	}
+	else
+	{
+		return 0.0f;
+	}
+}
diff --git a/Engine/Plugins/Animation/DeformerGraph/Source/OptimusCore/Private/DataInterfaces/OptimusDataInterfaceSkinnedMeshVertexAttribute.cpp b/Engine/Plugins/Animation/DeformerGraph/Source/OptimusCore/Private/DataInterfaces/OptimusDataInterfaceSkinnedMeshVertexAttribute.cpp
new file mode 100644
--- /dev/null
+++ b/Engine/Plugins/Animation/DeformerGraph/Source/OptimusCore/Private/DataInterfaces/OptimusDataInterfaceSkinnedMeshVertexAttribute.cpp
@@ -0,0 +1,173 @@
+﻿// Copyright Epic Games, Inc. All Rights Reserved.
+
+#include "OptimusDataInterfaceSkinnedMeshVertexAttribute.h"
+
+#include "ComponentSources/OptimusSkinnedMeshComponentSource.h"
+#include "ComputeFramework/ShaderParamTypeDefinition.h"
+#include "Rendering/SkeletalMeshAttributeVertexBuffer.h"
+#include "Rendering/SkeletalMeshRenderData.h"
+#include "ShaderCompilerCore.h"
+#include "ShaderParameterMetadataBuilder.h"
+#include "SkeletalRenderPublic.h"
+
+
+FString UOptimusSkinnedMeshVertexAttributeDataInterface::GetDisplayName() const
+{
+	return TEXT("Skinned Mesh Vertex Attribute");
+}
+
+
+TArray<FOptimusCDIPinDefinition> UOptimusSkinnedMeshVertexAttributeDataInterface::GetPinDefinitions() const
+{
+	TArray<FOptimusCDIPinDefinition> Defs;
+	Defs.Add({"NumVertices", "ReadNumVertices"});
+	Defs.Add({"Value", "ReadValue", UOptimusSkinnedMeshComponentSource::Domains::Vertex, "ReadNumVertices"});
+	return Defs;
+}
+
+
+TSubclassOf<UActorComponent> UOptimusSkinnedMeshVertexAttributeDataInterface::GetRequiredComponentClass() const
+{
+	return USkinnedMeshComponent::StaticClass();
+}
+
+
+void UOptimusSkinnedMeshVertexAttributeDataInterface::GetSupportedInputs(TArray<FShaderFunctionDefinition>& OutFunctions) const
+{
+	OutFunctions.AddDefaulted_GetRef()
+		.SetName(TEXT("ReadNumVertices"))
+		.AddReturnType(EShaderFundamentalType::Uint);
+
+	OutFunctions.AddDefaulted_GetRef()
+		.SetName(TEXT("ReadValue"))
+		.AddReturnType(EShaderFundamentalType::Float)
+		.AddParam(EShaderFundamentalType::Uint);
+}
+
+BEGIN_SHADER_PARAMETER_STRUCT(FSkinnedMeshVertexAttributeDataInterfaceParameters, )
+	SHADER_PARAMETER(uint32, NumVertices)
+	SHADER_PARAMETER(uint32, InputStreamStart)
+	SHADER_PARAMETER(uint32, bIsValid)
+	SHADER_PARAMETER_SRV(Buffer<float>, ValueBuffer)
+END_SHADER_PARAMETER_STRUCT()
+
+
+void UOptimusSkinnedMeshVertexAttributeDataInterface::GetShaderParameters(
+	TCHAR const* InUID,
+	FShaderParametersMetadataBuilder& InOutBuilder, 
+	FShaderParametersMetadataAllocations& InOutAllocations
+	) const
+{
+	InOutBuilder.AddNestedStruct<FSkinnedMeshVertexAttributeDataInterfaceParameters>(InUID);
+}
+
+
+void UOptimusSkinnedMeshVertexAttributeDataInterface::GetHLSL(FString& OutHLSL, FString const& InDataInterfaceName) const
+{
+	TMap<FString, FStringFormatArg> TemplateArgs =
+	{
+		{ TEXT("DataInterfaceName"), InDataInterfaceName },
+	};
+
+	FString TemplateFile;
+	LoadShaderSourceFile(TEXT("/Plugin/Optimus/Private/DataInterfaceSkinnedMeshVertexAttribute.ush"), EShaderPlatform::SP_PCD3D_SM5, &TemplateFile, nullptr);
+	OutHLSL += FString::Format(*TemplateFile, TemplateArgs);
+}
+
+
+void UOptimusSkinnedMeshVertexAttributeDataInterface::GetShaderHash(FString& InOutKey) const
+{
+	GetShaderFileHash(TEXT("/Plugin/Optimus/Private/DataInterfaceSkinnedMeshVertexAttribute.ush"), EShaderPlatform::SP_PCD3D_SM5).AppendString(InOutKey);
+}
+
+
+UComputeDataProvider* UOptimusSkinnedMeshVertexAttributeDataInterface::CreateDataProvider(
+	TObjectPtr<UObject> InBinding,
+	uint64 InInputMask,
+	uint64 InOutputMask
+	) const
+{
+	UOptimusSkinnedMeshVertexAttributeDataProvider* Provider = NewObject<UOptimusSkinnedMeshVertexAttributeDataProvider>();
+	Provider->SkinnedMeshComponent = Cast<USkinnedMeshComponent>(InBinding);
+	Provider->AttributeName = AttributeName;
+
+	return Provider;
+}
+
+
+bool UOptimusSkinnedMeshVertexAttributeDataProvider::IsValid() const
+{
+	if (SkinnedMeshComponent == nullptr || SkinnedMeshComponent->MeshObject == nullptr)
+	{
+		return false;
+	}
+	
+	const int32 LodIndex = SkinnedMeshComponent->MeshObject->GetLOD();
+	FSkeletalMeshRenderData const& SkeletalMeshRenderData = SkinnedMeshComponent->MeshObject->GetSkeletalMeshRenderData();
+	FSkeletalMeshLODRenderData const* LodRenderData = &SkeletalMeshRenderData.LODRenderData[LodIndex];
+
+	return LodRenderData->VertexAttributeBuffers.GetAttributeBuffer(AttributeName) != nullptr;
+}
+
+
+FComputeDataProviderRenderProxy* UOptimusSkinnedMeshVertexAttributeDataProvider::GetRenderProxy()
+{
+	return new FOptimusSkinnedMeshVertexAttributeDataProviderProxy(SkinnedMeshComponent, AttributeName);
+}
+
+
+FOptimusSkinnedMeshVertexAttributeDataProviderProxy::FOptimusSkinnedMeshVertexAttributeDataProviderProxy(
+	USkinnedMeshComponent* InSkinnedMeshComponent, 
+	FName InAttributeName
+	)
+{
+	SkeletalMeshObject = InSkinnedMeshComponent->MeshObject;
+	AttributeName = InAttributeName;
+}
+
+
+bool FOptimusSkinnedMeshVertexAttributeDataProviderProxy::IsValid(FValidationData const& InValidationData) const
+{
+	if (InValidationData.ParameterStructSize != sizeof(FParameters))
+	{
+		return false;
+	}
+	if (SkeletalMeshObject == nullptr)
+	{
+		return false;
+	}
+	if (SkeletalMeshObject->GetSkeletalMeshRenderData().LODRenderData[SkeletalMeshObject->GetLOD()].RenderSections.Num() != InValidationData.NumInvocations)
+	{
+		return false;
+	}
+
+	return true;
+}
+
+
+void FOptimusSkinnedMeshVertexAttributeDataProviderProxy::GatherDispatchData(
+	FDispatchData const& InDispatchData
+	)
+{
+	const int32 LodIndex = SkeletalMeshObject->GetLOD();
+	FSkeletalMeshRenderData const& SkeletalMeshRenderData = SkeletalMeshObject->GetSkeletalMeshRenderData();
+	FSkeletalMeshLODRenderData const* LodRenderData = &SkeletalMeshRenderData.LODRenderData[LodIndex];
+	
+	const TStridedView<FSkinnedMeshVertexAttributeDataInterfaceParameters> ParameterArray =
+		MakeStridedParameterView<FSkinnedMeshVertexAttributeDataInterfaceParameters>(InDispatchData);
+	
+	FRHIShaderResourceView* NullSRVBinding = GWhiteVertexBufferWithSRV->ShaderResourceViewRHI.GetReference();
+	const FSkeletalMeshAttributeVertexBuffer* AttributeBuffer = LodRenderData->VertexAttributeBuffers.GetAttributeBuffer(AttributeName);
+	FRHIShaderResourceView *AttributeSRV = AttributeBuffer ? AttributeBuffer->GetSRV() : NullSRVBinding;
+
+	for (int32 InvocationIndex = 0; InvocationIndex < ParameterArray.Num(); ++InvocationIndex)
+	{
+		FSkelMeshRenderSection const& RenderSection = LodRenderData->RenderSections[InvocationIndex];
+
+		FSkinnedMeshVertexAttributeDataInterfaceParameters& Parameters = ParameterArray[InvocationIndex];
+		Parameters.NumVertices = InDispatchData.bUnifiedDispatch ? LodRenderData->GetNumVertices() : RenderSection.NumVertices;
+		Parameters.InputStreamStart = InDispatchData.bUnifiedDispatch ? 0 : RenderSection.BaseVertexIndex;
+		Parameters.ValueBuffer = AttributeSRV;
+		Parameters.bIsValid = AttributeBuffer != nullptr;
+	}
+}
diff --git a/Engine/Plugins/Animation/DeformerGraph/Source/OptimusCore/Private/DataInterfaces/OptimusDataInterfaceSkinnedMeshVertexAttribute.h b/Engine/Plugins/Animation/DeformerGraph/Source/OptimusCore/Private/DataInterfaces/OptimusDataInterfaceSkinnedMeshVertexAttribute.h
new file mode 100644
--- /dev/null
+++ b/Engine/Plugins/Animation/DeformerGraph/Source/OptimusCore/Private/DataInterfaces/OptimusDataInterfaceSkinnedMeshVertexAttribute.h
@@ -0,0 +1,79 @@
+﻿// Copyright Epic Games, Inc. All Rights Reserved.
+
+#pragma once
+
+#include "OptimusComputeDataInterface.h"
+
+#include "ComputeFramework/ComputeDataProvider.h"
+
+#include "OptimusDataInterfaceSkinnedMeshVertexAttribute.generated.h"
+
+
+class FSkeletalMeshObject;
+class FSkinnedMeshVertexAttributeDataInterfaceParameters;
+
+
+UCLASS(Category = ComputeFramework)
+class UOptimusSkinnedMeshVertexAttributeDataInterface :
+	public UOptimusComputeDataInterface
+{
+	GENERATED_BODY()
+	
+	//~ Begin UOptimusComputeDataInterface Interface
+	FString GetDisplayName() const override;
+	TArray<FOptimusCDIPinDefinition> GetPinDefinitions() const override;
+	TSubclassOf<UActorComponent> GetRequiredComponentClass() const override;
+	//~ End UOptimusComputeDataInterface Interface
+
+	//~ Begin UComputeDataInterface Interface
+	TCHAR const* GetClassName() const override { return TEXT("VertexAttribute"); }
+	virtual void GetSupportedInputs(TArray<FShaderFunctionDefinition>& OutFunctions) const override;
+	void GetShaderParameters(TCHAR const* UID, FShaderParametersMetadataBuilder& InOutBuilder, FShaderParametersMetadataAllocations& InOutAllocations) const override;
+	void GetHLSL(FString& OutHLSL, FString const& InDataInterfaceName) const override;
+	void GetShaderHash(FString& InOutKey) const override;
+	UComputeDataProvider* CreateDataProvider(TObjectPtr<UObject> InBinding, uint64 InInputMask, uint64 InOutputMask) const override;
+	//~ End UComputeDataInterface Interface
+
+	UPROPERTY(EditAnywhere, Category="Vertex Attribute")
+	FName AttributeName;
+};
+
+
+/** Compute Framework Data Provider for reading skeletal mesh. */
+UCLASS(BlueprintType, editinlinenew, Category = ComputeFramework)
+class UOptimusSkinnedMeshVertexAttributeDataProvider : public UComputeDataProvider
+{
+	GENERATED_BODY()
+
+public:
+	UPROPERTY(EditAnywhere, BlueprintReadWrite, Category = Binding)
+	TObjectPtr<USkinnedMeshComponent> SkinnedMeshComponent = nullptr;
+
+	UPROPERTY(EditAnywhere, BlueprintReadWrite, Category = VertexAttribute)
+	FName AttributeName;
+
+	//~ Begin UComputeDataProvider Interface
+	bool IsValid() const override;
+	FComputeDataProviderRenderProxy* GetRenderProxy() override;
+	//~ End UComputeDataProvider Interface
+
+};
+
+class FOptimusSkinnedMeshVertexAttributeDataProviderProxy : public FComputeDataProviderRenderProxy
+{
+public:
+	FOptimusSkinnedMeshVertexAttributeDataProviderProxy(
+		USkinnedMeshComponent* InSkinnedMeshComponent,
+		FName InAttributeName);
+
+	//~ Begin FComputeDataProviderRenderProxy Interface
+	bool IsValid(FValidationData const& InValidationData) const override;
+	void GatherDispatchData(FDispatchData const& InDispatchData) override;
+	//~ End FComputeDataProviderRenderProxy Interface
+
+private:
+	using FParameters = FSkinnedMeshVertexAttributeDataInterfaceParameters;
+	
+	FSkeletalMeshObject* SkeletalMeshObject;
+	FName AttributeName;
+};
