Fixed mismatched copy constructor/assignment operator in FOptimusDataDomain.

diff --git a/Engine/Plugins/Animation/DeformerGraph/Source/OptimusCore/Public/OptimusDataDomain.h b/Engine/Plugins/Animation/DeformerGraph/Source/OptimusCore/Public/OptimusDataDomain.h
--- a/Engine/Plugins/Animation/DeformerGraph/Source/OptimusCore/Public/OptimusDataDomain.h
+++ b/Engine/Plugins/Animation/DeformerGraph/Source/OptimusCore/Public/OptimusDataDomain.h
@@ -48,8 +48,7 @@ struct OPTIMUSCORE_API FOptimusDataDomain
 	GENERATED_BODY()
 
 	FOptimusDataDomain() = default;
-	FOptimusDataDomain(const FOptimusDataDomain&) = default;
-	
+
 	explicit FOptimusDataDomain(TArray<FName> InDimensionNames) :
 		DimensionNames(InDimensionNames)
 	{}
@@ -59,7 +58,7 @@ struct OPTIMUSCORE_API FOptimusDataDomain
 		Multiplier(DimensionNames.Num() == 1 ? FMath::Max(InMultiplier, 1) : 1)
 	{
 	}
-	
+
 	FOptimusDataDomain(FString InExpression) :
 		Type(EOptimusDataDomainType::Expression),
 		Expression(MoveTemp(InExpression))
