Optimus: deprecating execution data interfaces and introducing kernel data interface

diff --git a/Engine/Plugins/Animation/DeformerGraph/Source/OptimusCore/Private/Actions/OptimusComponentBindingActions.cpp b/Engine/Plugins/Animation/DeformerGraph/Source/OptimusCore/Private/Actions/OptimusComponentBindingActions.cpp
--- a/Engine/Plugins/Animation/DeformerGraph/Source/OptimusCore/Private/Actions/OptimusComponentBindingActions.cpp
+++ b/Engine/Plugins/Animation/DeformerGraph/Source/OptimusCore/Private/Actions/OptimusComponentBindingActions.cpp
@@ -5,6 +5,7 @@
 #include "IOptimusPathResolver.h"
 #include "OptimusComponentSource.h"
 #include "OptimusDeformer.h"
+#include "OptimusHelpers.h"
 
 #include "UObject/Package.h"
 #include "UObject/UObjectGlobals.h" 
@@ -58,7 +59,7 @@ bool FOptimusComponentBindingAction_AddBinding::Do(IOptimusPathResolver* InRoot)
 
 	if (!Deformer->AddComponentBindingDirect(Binding, INDEX_NONE))
 	{
-		Binding->Rename(nullptr, GetTransientPackage());
+		Optimus::RemoveObject(Binding);
 		return false;
 	}
 
@@ -133,7 +134,7 @@ bool FOptimusComponentBindingAction_RemoveBinding::Undo(IOptimusPathResolver* In
 
 	if (!Deformer->AddComponentBindingDirect(Binding, ComponentBindingIndex))
 	{
-		Binding->Rename(nullptr, GetTransientPackage());
+		Optimus::RemoveObject(Binding);
 		return false;
 	}
 
diff --git a/Engine/Plugins/Animation/DeformerGraph/Source/OptimusCore/Private/Actions/OptimusNodeGraphActions.cpp b/Engine/Plugins/Animation/DeformerGraph/Source/OptimusCore/Private/Actions/OptimusNodeGraphActions.cpp
--- a/Engine/Plugins/Animation/DeformerGraph/Source/OptimusCore/Private/Actions/OptimusNodeGraphActions.cpp
+++ b/Engine/Plugins/Animation/DeformerGraph/Source/OptimusCore/Private/Actions/OptimusNodeGraphActions.cpp
@@ -72,14 +72,14 @@ bool FOptimusNodeGraphAction_AddGraph::Do(
 	
 	if (ConfigureGraphFunc && !ConfigureGraphFunc(Graph))
 	{
-		Graph->Rename(nullptr, GetTransientPackage());
+		Optimus::RemoveObject(Graph);
 		return false;
 	}
 
 	// Add the graph to the collection
 	if (!GraphOwner->AddGraph(Graph, GraphIndex))
 	{
-		Graph->Rename(nullptr, GetTransientPackage());
+		Optimus::RemoveObject(Graph);
 		return false;
 	}
 	
@@ -174,7 +174,7 @@ bool FOptimusNodeGraphAction_RemoveGraph::Undo(
 	// Now add the graph such that interested parties get notified.
 	if (!GraphOwner->AddGraph(Graph, GraphIndex))
 	{
-		Graph->Rename(nullptr, GetTransientPackage());
+		Optimus::RemoveObject(Graph);
 		return false;
 	}
 	
@@ -216,8 +216,8 @@ bool FOptimusNodeGraphAction_RenameGraph::Do(
 	{
 		return false;
 	}
-
-	return Graph->Rename(*NewGraphName.ToString(), nullptr);
+	
+	return Optimus::RenameObject(Graph, *NewGraphName.ToString(), nullptr);
 }
 
 
@@ -231,7 +231,7 @@ bool FOptimusNodeGraphAction_RenameGraph::Undo(
 		return false;
 	}
 
-	return Graph->Rename(*OldGraphName.ToString(), nullptr);
+	return Optimus::RenameObject(Graph, *OldGraphName.ToString(), nullptr);
 }
 
 
diff --git a/Engine/Plugins/Animation/DeformerGraph/Source/OptimusCore/Private/Actions/OptimusResourceActions.cpp b/Engine/Plugins/Animation/DeformerGraph/Source/OptimusCore/Private/Actions/OptimusResourceActions.cpp
--- a/Engine/Plugins/Animation/DeformerGraph/Source/OptimusCore/Private/Actions/OptimusResourceActions.cpp
+++ b/Engine/Plugins/Animation/DeformerGraph/Source/OptimusCore/Private/Actions/OptimusResourceActions.cpp
@@ -54,7 +54,7 @@ bool FOptimusResourceAction_AddResource::Do(
 
 	if (!Deformer->AddResourceDirect(Resource, INDEX_NONE))
 	{
-		Resource->Rename(nullptr, GetTransientPackage());
+		Optimus::RemoveObject(Resource);
 		return false;
 	}
 	
@@ -134,7 +134,7 @@ bool FOptimusResourceAction_RemoveResource::Undo(
 
 	if (!Deformer->AddResourceDirect(Resource, ResourceIndex))
 	{
-		Resource->Rename(nullptr, GetTransientPackage());
+		Optimus::RemoveObject(Resource);
 		return false;
 	}
 	
diff --git a/Engine/Plugins/Animation/DeformerGraph/Source/OptimusCore/Private/Actions/OptimusVariableActions.cpp b/Engine/Plugins/Animation/DeformerGraph/Source/OptimusCore/Private/Actions/OptimusVariableActions.cpp
--- a/Engine/Plugins/Animation/DeformerGraph/Source/OptimusCore/Private/Actions/OptimusVariableActions.cpp
+++ b/Engine/Plugins/Animation/DeformerGraph/Source/OptimusCore/Private/Actions/OptimusVariableActions.cpp
@@ -52,7 +52,7 @@ bool FOptimusVariableAction_AddVariable::Do(IOptimusPathResolver* InRoot)
 
 	if (!Deformer->AddVariableDirect(Variable, INDEX_NONE))
 	{
-		Variable->Rename(nullptr, GetTransientPackage());
+		Optimus::RemoveObject(Variable);
 		return false;
 	}
 	
@@ -139,7 +139,7 @@ bool FOptimusVariableAction_RemoveVariable::Undo(IOptimusPathResolver* InRoot)
 	
 	if (!Deformer->AddVariableDirect(Variable, VariableIndex))
 	{
-		Variable->Rename(nullptr, GetTransientPackage());
+		Optimus::RemoveObject(Variable);
 		return false;
 	}
 	
diff --git a/Engine/Plugins/Animation/DeformerGraph/Source/OptimusCore/Private/ComponentSources/OptimusSkinnedMeshComponentSource.cpp b/Engine/Plugins/Animation/DeformerGraph/Source/OptimusCore/Private/ComponentSources/OptimusSkinnedMeshComponentSource.cpp
--- a/Engine/Plugins/Animation/DeformerGraph/Source/OptimusCore/Private/ComponentSources/OptimusSkinnedMeshComponentSource.cpp
+++ b/Engine/Plugins/Animation/DeformerGraph/Source/OptimusCore/Private/ComponentSources/OptimusSkinnedMeshComponentSource.cpp
@@ -41,6 +41,27 @@ int32 UOptimusSkinnedMeshComponentSource::GetLodIndex(const UActorComponent* InC
 	return SkinnedMeshComponent ? SkinnedMeshComponent->GetPredictedLODLevel() : 0;
 }
 
+uint32 UOptimusSkinnedMeshComponentSource::GetDefaultNumInvocations(const UActorComponent* InComponent,
+                                                                    int32 InLod) const
+{
+	const USkinnedMeshComponent* SkinnedMeshComponent = Cast<USkinnedMeshComponent>(InComponent);
+	if (!SkinnedMeshComponent)
+	{
+		return 0;
+	}
+
+	const FSkeletalMeshObject* SkeletalMeshObject = SkinnedMeshComponent->MeshObject;
+	if (!SkeletalMeshObject)
+	{
+		return 0;
+	}
+
+	FSkeletalMeshRenderData const& SkeletalMeshRenderData = SkeletalMeshObject->GetSkeletalMeshRenderData();
+	FSkeletalMeshLODRenderData const* LodRenderData = &SkeletalMeshRenderData.LODRenderData[InLod];
+
+	return LodRenderData->RenderSections.Num();
+}
+
 bool UOptimusSkinnedMeshComponentSource::GetComponentElementCountsForExecutionDomain(
 	FName InDomainName,
 	const UActorComponent* InComponent,
diff --git a/Engine/Plugins/Animation/DeformerGraph/Source/OptimusCore/Private/ComponentSources/OptimusSkinnedMeshComponentSource.h b/Engine/Plugins/Animation/DeformerGraph/Source/OptimusCore/Private/ComponentSources/OptimusSkinnedMeshComponentSource.h
--- a/Engine/Plugins/Animation/DeformerGraph/Source/OptimusCore/Private/ComponentSources/OptimusSkinnedMeshComponentSource.h
+++ b/Engine/Plugins/Animation/DeformerGraph/Source/OptimusCore/Private/ComponentSources/OptimusSkinnedMeshComponentSource.h
@@ -26,5 +26,6 @@ public:
 	TSubclassOf<UActorComponent> GetComponentClass() const override;
 	TArray<FName> GetExecutionDomains() const override;
 	int32 GetLodIndex(const UActorComponent* InComponent) const override;
+	uint32 GetDefaultNumInvocations(const UActorComponent* InComponent, int32 InLod) const override;
 	bool GetComponentElementCountsForExecutionDomain(FName InDomainName, const UActorComponent* InComponent, int32 InLodIndex, TArray<int32>& OutInvocationElementCounts) const override;
 };
diff --git a/Engine/Plugins/Animation/DeformerGraph/Source/OptimusCore/Private/DataInterfaces/OptimusDataInterfaceCustomComputeKernel.cpp b/Engine/Plugins/Animation/DeformerGraph/Source/OptimusCore/Private/DataInterfaces/OptimusDataInterfaceCustomComputeKernel.cpp
new file mode 100644
--- /dev/null
+++ b/Engine/Plugins/Animation/DeformerGraph/Source/OptimusCore/Private/DataInterfaces/OptimusDataInterfaceCustomComputeKernel.cpp
@@ -0,0 +1,284 @@
+// Copyright Epic Games, Inc. All Rights Reserved.
+
+#include "OptimusDataInterfaceCustomComputeKernel.h"
+
+#include "OptimusComponentSource.h"
+#include "OptimusExpressionEvaluator.h"
+#include "ShaderParameterMetadataBuilder.h"
+#include "ComputeFramework/ComputeMetadataBuilder.h"
+#include "ComputeFramework/ShaderParameterMetadataAllocation.h"
+#include "ComputeFramework/ShaderParamTypeDefinition.h"
+
+#include "Nodes/OptimusNode_CustomComputeKernel.h"
+
+#include UE_INLINE_GENERATED_CPP_BY_NAME(OptimusDataInterfaceCustomComputeKernel)
+
+const FString UOptimusCustomComputeKernelDataInterface::NumThreadsReservedName = TEXT("NumThreads");
+
+void UOptimusCustomComputeKernelDataInterface::InitFromKernelNode(const UOptimusNode_CustomComputeKernel* InKernelNode)
+{
+	NumThreadsExpression = InKernelNode->ExecutionDomain.Name.ToString();
+
+	TSet<UOptimusComponentSourceBinding*> PrimaryBinding = InKernelNode->GetGroupComponentSourceBindings(InKernelNode->GetPrimaryGroupPin());
+
+	// Kernel node's ValidateForCompile() should guaranteed the existence of unique primary binding
+	if (ensure(PrimaryBinding.Num() == 1))
+	{
+		ComponentSourceBinding = PrimaryBinding.Array()[0];
+	}
+}
+
+void UOptimusCustomComputeKernelDataInterface::GetSupportedInputs(TArray<FShaderFunctionDefinition>& OutFunctions) const
+{
+	OutFunctions.AddDefaulted_GetRef()
+		.SetName(FString::Printf(TEXT("Read%s"), *NumThreadsReservedName))
+		.AddReturnType(FShaderValueType::Get(EShaderFundamentalType::Int, 3));
+}
+
+void UOptimusCustomComputeKernelDataInterface::GetShaderParameters(TCHAR const* UID,
+	FShaderParametersMetadataBuilder& InOutBuilder, FShaderParametersMetadataAllocations& InOutAllocations) const
+{
+	FShaderParametersMetadataBuilder Builder;
+	TArray<FShaderParametersMetadata*> DummyNestedStructs;
+	ComputeFramework::AddParamForType(Builder, *NumThreadsReservedName, FShaderValueType::Get(EShaderFundamentalType::Int, 3),DummyNestedStructs);
+
+	FShaderParametersMetadata* ShaderParameterMetadata = Builder.Build(FShaderParametersMetadata::EUseCase::ShaderParameterStruct, TEXT("UAnimAttributeDataInterface"));
+
+	InOutAllocations.ShaderParameterMetadatas.Add(ShaderParameterMetadata);
+	InOutAllocations.ShaderParameterMetadatas.Append(DummyNestedStructs);
+
+	// Add the generated nested struct to our builder.
+	InOutBuilder.AddNestedStruct(UID, ShaderParameterMetadata);
+
+}
+
+void UOptimusCustomComputeKernelDataInterface::GetShaderHash(FString& InOutKey) const
+{
+	
+}
+
+void UOptimusCustomComputeKernelDataInterface::GetHLSL(FString& OutHLSL, FString const& InDataInterfaceName) const
+{
+	const FString TypeName = FShaderValueType::Get(EShaderFundamentalType::Int, 3)->ToString();
+
+	if (ensure(!TypeName.IsEmpty()))
+	{
+		// Add uniforms.
+		OutHLSL += FString::Printf(TEXT("%s %s_%s;\n"), 
+			*TypeName, 
+			*InDataInterfaceName, 
+			*NumThreadsReservedName);
+			
+		// Add function getters.
+		OutHLSL += FString::Printf(TEXT("%s Read%s_%s()\n{\n\treturn %s_%s;\n}\n"), 
+			*TypeName,
+			*NumThreadsReservedName,
+			*InDataInterfaceName, 
+			*InDataInterfaceName, 
+			*NumThreadsReservedName);
+	}
+}
+
+UComputeDataProvider* UOptimusCustomComputeKernelDataInterface::CreateDataProvider(TObjectPtr<UObject> InBinding,
+	uint64 InInputMask, uint64 InOutputMask) const
+{
+	UOptimusCustomComputeKernelDataProvider* Provider = NewObject<UOptimusCustomComputeKernelDataProvider>();
+
+	Provider->InitFromDataInterface(this, InBinding);
+
+	return Provider;
+}
+
+void UOptimusCustomComputeKernelDataProvider::InitFromDataInterface(const UOptimusCustomComputeKernelDataInterface* InDataInterface, const UObject* InBinding)
+{
+	WeakComponent = Cast<UActorComponent>(InBinding);
+	WeakDataInterface = InDataInterface;
+}
+
+FComputeDataProviderRenderProxy* UOptimusCustomComputeKernelDataProvider::GetRenderProxy()
+{
+	TArray<int32> InvocationCounts;
+	int32 TotalThreadCount;
+	
+	if (!GetInvocationThreadCounts(InvocationCounts, TotalThreadCount))
+	{
+		InvocationCounts.Reset();
+	}
+	
+	FOptimusCustomComputeKernelDataProviderProxy* Proxy = new FOptimusCustomComputeKernelDataProviderProxy(MoveTemp(InvocationCounts), TotalThreadCount);
+	return Proxy;
+}
+
+bool UOptimusCustomComputeKernelDataProvider::GetInvocationThreadCounts(
+	TArray<int32>& OutInvocationThreadCount,
+	int32& OutTotalThreadCount
+	) const
+{
+	if (!WeakDataInterface.IsValid())
+	{
+		return false;
+	}
+
+	if (!WeakDataInterface->ComponentSourceBinding.IsValid())
+	{
+		return false;
+	}
+	
+	const UOptimusComponentSource* ComponentSource = WeakDataInterface->ComponentSourceBinding->GetComponentSource();
+	const FString& NumThreadsExpression = WeakDataInterface->NumThreadsExpression;
+	
+	if (!WeakComponent.IsValid() || !ComponentSource)
+	{
+		return false;
+	}
+	
+	const UActorComponent* Component = WeakComponent.Get();
+	int32 LodIndex = ComponentSource->GetLodIndex(Component);
+	int32 NumInvocations = ComponentSource->GetDefaultNumInvocations(Component, LodIndex);
+
+	// In some cases the component can be set to a unusable state intentionally
+	// such as when the OptimusEditor shuts down,
+	// in which case we don't have data to work with so simply do nothing
+	if (NumInvocations < 1)
+	{
+		return false;
+	}
+	
+	TMap<FName, int32> EngineConstants;
+	TMap<FName, TArray<int32>> ElementCountsPerDomain;
+
+	for(FName ExecutionDomain: ComponentSource->GetExecutionDomains())
+	{
+		EngineConstants.Add(ExecutionDomain, 0);
+
+		TArray<int32>& ElementCounts = ElementCountsPerDomain.FindOrAdd(ExecutionDomain);
+		if (!ComponentSource->GetComponentElementCountsForExecutionDomain(
+				ExecutionDomain,
+				Component,
+				LodIndex,
+				ElementCounts))
+		{
+			return false;
+		}
+
+		if (!ensure(NumInvocations == ElementCounts.Num()))
+		{
+			// Component source needs to provide as many values for each of its execution domains as there are invocations
+			return false;
+		}
+	}
+
+	using namespace Optimus::Expression;
+
+	FEngine Engine(EngineConstants);
+	TVariant<FExpressionObject, FParseError> ParseResult = Engine.Parse(NumThreadsExpression);
+	if (ParseResult.IsType<FParseError>())
+	{
+		return false;
+	}
+
+	OutInvocationThreadCount.Reset(NumInvocations);
+	for (int32 Index = 0; Index < NumInvocations; Index++)
+	{
+		for (TPair<FName, int32>& Constant: EngineConstants)
+		{
+			const FName ConstantName = Constant.Key;
+			int32& Value = Constant.Value;
+			
+			const TArray<int32>& ElementCounts = ElementCountsPerDomain[ConstantName]; 
+			Value = ElementCounts.IsValidIndex(Index) ? ElementCounts[Index] : 1;
+		}
+		Engine.UpdateConstantValues(EngineConstants);
+		const int32 Count = Engine.Execute(ParseResult.Get<FExpressionObject>());
+
+		if (Count < 0)
+		{
+			return false;
+		}
+		
+		OutInvocationThreadCount.Add(Count);
+	}
+
+	OutTotalThreadCount = 0;
+	for (int32 Count : OutInvocationThreadCount)
+	{
+		OutTotalThreadCount += Count;
+	}
+
+	// Extra validation for unified dispatch
+
+	// We want to make sure the results are the same for unified and non-unified
+	// Do the sum across invocations first and then evaluate the expression
+	for (auto& [ConstantName, Value]: EngineConstants)
+	{
+		Value = 0;
+		const TArray<int32>& ElementCounts = ElementCountsPerDomain[ConstantName];
+		
+		for (int32 Count : ElementCounts )
+		{
+			Value += Count;
+		}
+	}
+	
+	Engine.UpdateConstantValues(EngineConstants);
+	const int32 TotalElementCountForUnifiedDispatch = Engine.Execute(ParseResult.Get<FExpressionObject>());
+	
+	if (OutTotalThreadCount != TotalElementCountForUnifiedDispatch)
+	{
+		return false;
+	}
+
+	return true;	
+}
+
+FOptimusCustomComputeKernelDataProviderProxy::FOptimusCustomComputeKernelDataProviderProxy(
+	TArray<int32>&& InInvocationThreadCounts,
+	int32 InTotalThreadCount
+	) :
+	InvocationThreadCounts(MoveTemp(InInvocationThreadCounts)),
+	TotalThreadCount(InTotalThreadCount)
+{
+}
+
+bool FOptimusCustomComputeKernelDataProviderProxy::IsValid(FValidationData const& InValidationData) const
+{
+	if (InvocationThreadCounts.Num() == 0)
+	{
+		return false;
+	}
+	
+	return true;
+}
+
+int32 FOptimusCustomComputeKernelDataProviderProxy::GetDispatchThreadCount(TArray<FIntVector>& InOutThreadCounts) const
+{
+	InOutThreadCounts.Reset(InvocationThreadCounts.Num());
+	for (const int32 Count : InvocationThreadCounts)
+	{
+		InOutThreadCounts.Add({Count, 1, 1});
+	}
+	return InOutThreadCounts.Num();
+}
+
+void FOptimusCustomComputeKernelDataProviderProxy::GatherDispatchData(FDispatchData const& InDispatchData)
+{
+	for (int32 InvocationIndex = 0; InvocationIndex < InDispatchData.NumInvocations; ++InvocationIndex)
+	{
+		int32 NumThreads;
+		if (InDispatchData.bUnifiedDispatch)
+		{
+			NumThreads = TotalThreadCount;
+		}
+		else
+		{
+			NumThreads = InvocationThreadCounts[InvocationIndex];
+		}
+
+		uint8* ParameterBuffer = (InDispatchData.ParameterBuffer + InDispatchData.ParameterBufferOffset + InDispatchData.ParameterBufferStride * InvocationIndex);
+		
+		const int32 ThreadCount3D[] = {NumThreads, 1, 1};
+		FMemory::Memcpy(ParameterBuffer, ThreadCount3D, 3 * sizeof(int32));
+	}
+}
+
+
diff --git a/Engine/Plugins/Animation/DeformerGraph/Source/OptimusCore/Private/DataInterfaces/OptimusDataInterfaceCustomComputeKernel.h b/Engine/Plugins/Animation/DeformerGraph/Source/OptimusCore/Private/DataInterfaces/OptimusDataInterfaceCustomComputeKernel.h
new file mode 100644
--- /dev/null
+++ b/Engine/Plugins/Animation/DeformerGraph/Source/OptimusCore/Private/DataInterfaces/OptimusDataInterfaceCustomComputeKernel.h
@@ -0,0 +1,87 @@
+// Copyright Epic Games, Inc. All Rights Reserved.
+
+#pragma once
+
+#include "OptimusComputeDataInterface.h"
+#include "OptimusDataType.h"
+#include "ComputeFramework/ComputeDataProvider.h"
+
+#include "OptimusDataInterfaceCustomComputeKernel.generated.h"
+
+class UOptimusNode_CustomComputeKernel;
+class UOptimusComponentSource;
+class UOptimusComponentSourceBinding;
+
+UCLASS(Category = ComputeFramework)
+class OPTIMUSCORE_API UOptimusCustomComputeKernelDataInterface : public UComputeDataInterface
+{
+	GENERATED_BODY()
+
+public:
+	static const FString NumThreadsReservedName;
+	
+	void InitFromKernelNode(const UOptimusNode_CustomComputeKernel* InKernelNode);
+	
+	//~ Begin UComputeDataInterface Interface
+	TCHAR const* GetClassName() const override { return TEXT("CustomComputeKernelData"); }
+	bool IsExecutionInterface() const override { return true; }
+	bool CanSupportUnifiedDispatch() const override { return true; }
+	void GetSupportedInputs(TArray<FShaderFunctionDefinition>& OutFunctions) const override;
+	void GetShaderParameters(TCHAR const* UID, FShaderParametersMetadataBuilder& InOutBuilder, FShaderParametersMetadataAllocations& InOutAllocations) const override;
+	void GetShaderHash(FString& InOutKey) const override;
+	void GetHLSL(FString& OutHLSL, FString const& InDataInterfaceName) const override;
+	UComputeDataProvider* CreateDataProvider(TObjectPtr<UObject> InBinding, uint64 InInputMask, uint64 InOutputMask) const override;
+	//~ End UComputeDataInterface Interface
+	
+	UPROPERTY()
+	TWeakObjectPtr<const UOptimusComponentSourceBinding> ComponentSourceBinding;
+	
+	UPROPERTY()
+	FString NumThreadsExpression;
+
+protected:
+};
+
+/** Compute Framework Data Provider for each custom compute kernel. */
+UCLASS()
+class UOptimusCustomComputeKernelDataProvider : public UComputeDataProvider
+{
+	GENERATED_BODY()
+
+public:
+	void InitFromDataInterface(const UOptimusCustomComputeKernelDataInterface* InDataInterface, const UObject* InBinding);
+	
+	//~ Begin UComputeDataProvider Interface
+	FComputeDataProviderRenderProxy* GetRenderProxy() override;
+	//~ End UComputeDataProvider Interface
+
+	TWeakObjectPtr<const UActorComponent> WeakComponent = nullptr;
+
+	TWeakObjectPtr<const UOptimusCustomComputeKernelDataInterface> WeakDataInterface = nullptr;
+
+protected:
+	bool GetInvocationThreadCounts(
+		TArray<int32>& OutInvocationThreadCount,
+		int32& OutTotalThreadCount
+		) const;
+};
+
+class FOptimusCustomComputeKernelDataProviderProxy : public FComputeDataProviderRenderProxy
+{
+public:
+	FOptimusCustomComputeKernelDataProviderProxy(
+		TArray<int32>&& InInvocationThreadCounts,
+		int32 InTotalThreadCount
+		);
+
+	//~ Begin FComputeDataProviderRenderProxy Interface
+	bool IsValid(FValidationData const& InValidationData) const override;
+	int32 GetDispatchThreadCount(TArray<FIntVector>& InOutThreadCounts) const override;
+	void GatherDispatchData(FDispatchData const& InDispatchData) override;
+	//~ End FComputeDataProviderRenderProxy Interface
+
+	TArray<int32> InvocationThreadCounts;
+	int32 TotalThreadCount;
+};
+
+
diff --git a/Engine/Plugins/Animation/DeformerGraph/Source/OptimusCore/Private/DataInterfaces/OptimusDataInterfaceSkinnedMeshExec.cpp b/Engine/Plugins/Animation/DeformerGraph/Source/OptimusCore/Private/DataInterfaces/OptimusDataInterfaceSkinnedMeshExec.cpp
--- a/Engine/Plugins/Animation/DeformerGraph/Source/OptimusCore/Private/DataInterfaces/OptimusDataInterfaceSkinnedMeshExec.cpp
+++ b/Engine/Plugins/Animation/DeformerGraph/Source/OptimusCore/Private/DataInterfaces/OptimusDataInterfaceSkinnedMeshExec.cpp
@@ -9,6 +9,7 @@
 #include "ShaderCompilerCore.h"
 #include "ShaderParameterMetadataBuilder.h"
 #include "SkeletalRenderPublic.h"
+#include "ComponentSources/OptimusSkinnedMeshComponentSource.h"
 
 #include UE_INLINE_GENERATED_CPP_BY_NAME(OptimusDataInterfaceSkinnedMeshExec)
 
@@ -85,6 +86,21 @@ UComputeDataProvider* UOptimusSkinnedMeshExecDataInterface::CreateDataProvider(T
 	return Provider;
 }
 
+FName UOptimusSkinnedMeshExecDataInterface::GetSelectedExecutionDomainName() const
+{
+	if (Domain == EOptimusSkinnedMeshExecDomain::Vertex)
+	{
+		return UOptimusSkinnedMeshComponentSource::Domains::Vertex;
+	}
+
+	if (Domain == EOptimusSkinnedMeshExecDomain::Triangle)
+	{
+		return UOptimusSkinnedMeshComponentSource::Domains::Triangle;
+	}
+
+	return NAME_None;
+}
+
 
 FComputeDataProviderRenderProxy* UOptimusSkinnedMeshExecDataProvider::GetRenderProxy()
 {
diff --git a/Engine/Plugins/Animation/DeformerGraph/Source/OptimusCore/Private/DataInterfaces/OptimusDataInterfaceSkinnedMeshExec.h b/Engine/Plugins/Animation/DeformerGraph/Source/OptimusCore/Private/DataInterfaces/OptimusDataInterfaceSkinnedMeshExec.h
--- a/Engine/Plugins/Animation/DeformerGraph/Source/OptimusCore/Private/DataInterfaces/OptimusDataInterfaceSkinnedMeshExec.h
+++ b/Engine/Plugins/Animation/DeformerGraph/Source/OptimusCore/Private/DataInterfaces/OptimusDataInterfaceSkinnedMeshExec.h
@@ -2,6 +2,7 @@
 
 #pragma once
 
+#include "IOptimusDeprecatedExecutionDataInterface.h"
 #include "OptimusComputeDataInterface.h"
 #include "ComputeFramework/ComputeDataProvider.h"
 
@@ -25,7 +26,9 @@ enum class EOptimusSkinnedMeshExecDomain : uint8
 
 /** Compute Framework Data Interface for executing kernels over a skinned mesh domain. */
 UCLASS(Category = ComputeFramework)
-class OPTIMUSCORE_API UOptimusSkinnedMeshExecDataInterface : public UOptimusComputeDataInterface
+class OPTIMUSCORE_API UOptimusSkinnedMeshExecDataInterface :
+	public UOptimusComputeDataInterface,
+	public IOptimusDeprecatedExecutionDataInterface
 {
 	GENERATED_BODY()
 
@@ -33,6 +36,7 @@ public:
 	//~ Begin UOptimusComputeDataInterface Interface
 	FString GetDisplayName() const override;
 	FName GetCategory() const override;
+	bool IsVisible() const override {return false;};
 	TArray<FOptimusCDIPinDefinition> GetPinDefinitions() const override;
 	TSubclassOf<UActorComponent> GetRequiredComponentClass() const override;
 	//~ End UOptimusComputeDataInterface Interface
@@ -49,6 +53,10 @@ public:
 	UComputeDataProvider* CreateDataProvider(TObjectPtr<UObject> InBinding, uint64 InInputMask, uint64 InOutputMask) const override;
 	//~ End UComputeDataInterface Interface
 
+	//~ Begin IOptimusDeprecatedExecutionDataInterface Interface
+	FName GetSelectedExecutionDomainName() const override;
+	//~ End IOptimusDeprecatedExecutionDataInterface Interface
+	
 	UPROPERTY(EditAnywhere, Category = Execution)
 	EOptimusSkinnedMeshExecDomain Domain = EOptimusSkinnedMeshExecDomain::Vertex;
 
diff --git a/Engine/Plugins/Animation/DeformerGraph/Source/OptimusCore/Private/Nodes/OptimusNode_ComputeKernelBase.cpp b/Engine/Plugins/Animation/DeformerGraph/Source/OptimusCore/Private/Nodes/OptimusNode_ComputeKernelBase.cpp
--- a/Engine/Plugins/Animation/DeformerGraph/Source/OptimusCore/Private/Nodes/OptimusNode_ComputeKernelBase.cpp
+++ b/Engine/Plugins/Animation/DeformerGraph/Source/OptimusCore/Private/Nodes/OptimusNode_ComputeKernelBase.cpp
@@ -55,6 +55,7 @@ FOptimus_ComputeKernelResult UOptimusNode_ComputeKernelBase::CreateComputeKernel
 	const TArray<const UOptimusNode *>& InValueNodes,
 	const UComputeDataInterface* InGraphDataInterface,
 	const UOptimusComponentSourceBinding* InGraphDataComponentBinding,
+	const UComputeDataInterface* InKernelDataInterface,
 	FOptimus_InterfaceBindingMap& OutInputDataBindings,
 	FOptimus_InterfaceBindingMap& OutOutputDataBindings
 ) const
@@ -110,13 +111,10 @@ FOptimus_ComputeKernelResult UOptimusNode_ComputeKernelBase::CreateComputeKernel
 		return {};
 	};
 
-	TSet<UOptimusComponentSourceBinding*> PrimaryBinding;
 	for (const UOptimusNodePin* Pin: GetPins())
 	{
 		if (Pin->IsGroupingPin())
 		{
-			TSet<UOptimusComponentSourceBinding*> SecondaryBindings;
-			
 			for (const UOptimusNodePin* SubPin: Pin->GetSubPins())
 			{
 				if (TOptional<FText> Result = GetStructTypeDefFromPin(SubPin); Result.IsSet())
@@ -134,12 +132,21 @@ FOptimus_ComputeKernelResult UOptimusNode_ComputeKernelBase::CreateComputeKernel
 		}
 	}
 
+	TSet<UOptimusComponentSourceBinding*> PrimaryBindings = GetGroupComponentSourceBindings(GetPrimaryGroupPin());
+	
+	// ValidateForCompile() should guaranteed that we don't hit this error
+	if (!ensure(PrimaryBindings.Num() == 1))
+	{
+		return ReturnError(LOCTEXT("ZeroOrMultiplePrimaryBindings", "Primary Group has zero or more than one Component Bindings"));
+	}
+
 	UOptimusKernelSource* KernelSource = NewObject<UOptimusKernelSource>(InKernelSourceOuter);
 
 	// Wrap functions for unconnected resource pins (or value pins) that return default values
 	// (for reads) or do nothing (for writes).
 	TArray<FString> GeneratedFunctions;
 
+	int32 GroupIndex = 0;
 	for (const UOptimusNodePin* Pin: GetPins())
 	{
 		if (Pin->IsGroupingPin())
@@ -147,7 +154,7 @@ FOptimus_ComputeKernelResult UOptimusNode_ComputeKernelBase::CreateComputeKernel
 			for (const UOptimusNodePin* SubPin: Pin->GetSubPins())
 			{
 				TOptional<FText> Result = ProcessInputPinForComputeKernel(
-					InTraversalContext, SubPin, Pin->GetName(),
+					InTraversalContext, SubPin, GroupIndex == 0 ? FString() : Pin->GetName(),
 					InNodeDataInterfaceMap, InLinkDataInterfaceMap, InValueNodes,
 					InGraphDataInterface, InGraphDataComponentBinding,
 					KernelSource, GeneratedFunctions, OutInputDataBindings
@@ -157,20 +164,8 @@ FOptimus_ComputeKernelResult UOptimusNode_ComputeKernelBase::CreateComputeKernel
 					return ReturnError(Result.GetValue());
 				}
 			}
-		}
-		else if (Pin->GetDirection() == EOptimusNodePinDirection::Input)
-		{
-			TOptional<FText> Result = ProcessInputPinForComputeKernel(
-				InTraversalContext, Pin, FString(),
-				InNodeDataInterfaceMap, InLinkDataInterfaceMap, InValueNodes,
-				InGraphDataInterface, InGraphDataComponentBinding,
-				KernelSource, GeneratedFunctions, OutInputDataBindings
-				);
 			
-			if (Result.IsSet())
-			{
-				return ReturnError(Result.GetValue());
-			}
+			GroupIndex++;
 		}
 		else if (Pin->GetDirection() == EOptimusNodePinDirection::Output)
 		{
@@ -181,6 +176,14 @@ FOptimus_ComputeKernelResult UOptimusNode_ComputeKernelBase::CreateComputeKernel
 		}
 	}
 
+	if (ensure(InKernelDataInterface))
+	{
+		BindKernelDataInterfaceForComputeKernel(
+			PrimaryBindings.Array()[0],
+			InKernelDataInterface,
+			KernelSource, OutInputDataBindings);
+	}
+
 	FString CookedSource;
 	CookedSource += FString::Join(GeneratedFunctions, TEXT("\n"));
 	CookedSource += "\n\n";
@@ -253,39 +256,43 @@ TOptional<FText> UOptimusNode_ComputeKernelBase::ValidateForCompile() const
 	// Collect struct types and ensure all groups have the same component source (or none) within each group
 	auto VerifyPinBindings = [](TSet<UOptimusComponentSourceBinding*> &InCollectedBindings, const UOptimusNodePin* InPin)-> TOptional<FText>
 	{
-		// Component bindings only matter for resource pins.
-		if (!InPin->GetDataDomain().IsSingleton())
+		// Check for all pins, even if the data domain is singleton, it can have different values per invocation
+		// UOptimusNode_DataInterface::ValidateForCompile also double checks
+		const TSet<UOptimusComponentSourceBinding*> PinBindings = InPin->GetComponentSourceBindings();
+		if (PinBindings.Num() > 1)
 		{
-			const TSet<UOptimusComponentSourceBinding*> PinBindings = InPin->GetComponentSourceBindings();
-			if (PinBindings.Num() > 1)
-			{
-				return FText::Format(LOCTEXT("MultipleBindingsOnPin", "Multiple component bindings arriving into pin '{0}'"), FText::FromName(InPin->GetUniqueName()));
-			}
-			if (InCollectedBindings.IsEmpty())
-			{
-				InCollectedBindings = PinBindings;
-			}
-			else if (!PinBindings.IsEmpty())
+			return FText::Format(LOCTEXT("MultipleBindingsOnPin", "Multiple component bindings arriving into pin '{0}'"), FText::FromName(InPin->GetUniqueName()));
+		}
+		if (InCollectedBindings.IsEmpty())
+		{
+			InCollectedBindings = PinBindings;
+		}
+		else if (!PinBindings.IsEmpty())
+		{
+			if (InCollectedBindings.Num() != PinBindings.Num() || !InCollectedBindings.Includes(PinBindings))
 			{
-				if (InCollectedBindings.Num() != PinBindings.Num() || !InCollectedBindings.Includes(PinBindings))
-				{
-					return FText::Format(LOCTEXT("IncompatibleBindingsPinGroup", "Component binding for pin '{0}' is different from the component bindings of the other pins in its group"), FText::FromName(InPin->GetUniqueName()));
-				}
+				return FText::Format(LOCTEXT("IncompatibleBindingsPinGroup", "Component binding for pin '{0}' is different from the component bindings of the other pins in its group"), FText::FromName(InPin->GetUniqueName()));
 			}
 		}
 		return {};
 	};
 	
-	TSet<UOptimusComponentSourceBinding*> PrimaryBinding;
+	TArray<TSet<UOptimusComponentSourceBinding*>> BindingsPerGroup;
+	
 	for (const UOptimusNodePin* Pin: GetPins())
 	{
 		if (Pin->IsGroupingPin())
 		{
-			TSet<UOptimusComponentSourceBinding*> SecondaryBindings;
+			TSet<UOptimusComponentSourceBinding*>& Bindings = BindingsPerGroup.AddDefaulted_GetRef();
+
+			if(TOptional<FText> Result = VerifyPinBindings(Bindings, Pin); Result.IsSet())
+			{
+				return Result;
+			}
 			
 			for (const UOptimusNodePin* SubPin: Pin->GetSubPins())
 			{
-				if(TOptional<FText> Result = VerifyPinBindings(SecondaryBindings, Pin); Result.IsSet())
+				if(TOptional<FText> Result = VerifyPinBindings(Bindings, SubPin); Result.IsSet())
 				{
 					return Result;
 				}
@@ -296,21 +303,14 @@ TOptional<FText> UOptimusNode_ComputeKernelBase::ValidateForCompile() const
 				}
 			}
 		}
-		else 
-		{
-			if (Pin->GetDirection() == EOptimusNodePinDirection::Input)
-			{
-				if(TOptional<FText> Result = VerifyPinBindings(PrimaryBinding, Pin); Result.IsSet())
-				{
-					return Result;
-				}
-			}
-			if (TOptional<FText> Result = GetStructTypeDefFromPin(Pin); Result.IsSet())
-			{
-				return Result;
-			}
-		}
 	}
+
+	// We should have at least the primary group, which needs to have a unique component binding
+	if (ensure(BindingsPerGroup.Num() >= 1) && BindingsPerGroup[0].Num() == 0)
+	{
+		return LOCTEXT("NoPrimaryBinding", "Primary Group has no Component Binding");
+	}
+	
 	return {};
 }
 
@@ -334,7 +334,7 @@ FString UOptimusNode_ComputeKernelBase::GetCookedKernelSource(
 	FString ShaderPathName = InObjectPathName;
 	Optimus::ConvertObjectPathToShaderFilePath(ShaderPathName);
 
-	const bool bHasKernelKeyword = Source.Contains(TEXT("KERNEL"));
+	const bool bHasKernelKeyword = Source.Contains(TEXT("KERNEL"), ESearchCase::CaseSensitive);
 
 	const FString ComputeShaderUtilsInclude = TEXT("#include \"/Engine/Private/ComputeShaderUtils.ush\"");
 	
@@ -349,7 +349,7 @@ FString UOptimusNode_ComputeKernelBase::GetCookedKernelSource(
 
 	if (bHasKernelKeyword)
 	{
-		Source.ReplaceInline(TEXT("KERNEL"), TEXT("void __kernel_func(uint Index)"));
+		Source.ReplaceInline(TEXT("KERNEL"), TEXT("void __kernel_func(uint Index)"), ESearchCase::CaseSensitive);
 
 		return FString::Printf(
 			TEXT(
@@ -374,6 +374,21 @@ FString UOptimusNode_ComputeKernelBase::GetCookedKernelSource(
 	}
 }
 
+TSet<UOptimusComponentSourceBinding*> UOptimusNode_ComputeKernelBase::GetGroupComponentSourceBindings(const UOptimusNodePin* InGroupPin)
+{
+	TSet<UOptimusComponentSourceBinding*> Bindings;
+
+	Bindings.Append(InGroupPin->GetComponentSourceBindings());
+	
+	for (const UOptimusNodePin* Pin: InGroupPin->GetSubPins())
+	{
+		Bindings.Append(Pin->GetComponentSourceBindings());
+	}
+	
+	return Bindings;
+}
+
+
 
 TOptional<FText> UOptimusNode_ComputeKernelBase::ProcessInputPinForComputeKernel(
 	const FOptimusPinTraversalContext& InTraversalContext,
@@ -699,4 +714,41 @@ void UOptimusNode_ComputeKernelBase::ProcessOutputPinForComputeKernel(
 	}
 }
 
+void UOptimusNode_ComputeKernelBase::BindKernelDataInterfaceForComputeKernel(
+	const UOptimusComponentSourceBinding* InKernelPrimaryComponentSourceBinding,
+	const UComputeDataInterface* InKernelDataInterface,
+	UOptimusKernelSource* InKernelSource,
+	FOptimus_InterfaceBindingMap& OutInputDataBindings
+	) const
+{
+	check(InKernelDataInterface);
+	check(InKernelPrimaryComponentSourceBinding);
+	
+	TArray<FShaderFunctionDefinition> ReadFunctions;
+	InKernelDataInterface->GetSupportedInputs(ReadFunctions);	
+
+	// Simply grab everything the kernel data interface has to offer
+	for (int32 FuncIndex = 0; FuncIndex < ReadFunctions.Num(); FuncIndex++)
+	{
+		FShaderFunctionDefinition FuncDef = ReadFunctions[FuncIndex];
+
+		for (FShaderParamTypeDefinition& ParamType : FuncDef.ParamTypes)
+		{
+			// Making sure parameter has type declaration generated
+			CopyValueType(ParamType.ValueType, ParamType);
+		}
+
+		FOptimus_InterfaceBinding InterfaceBinding;
+		InterfaceBinding.DataInterface = InKernelDataInterface;
+		InterfaceBinding.ComponentBinding = InKernelPrimaryComponentSourceBinding;
+		InterfaceBinding.DataInterfaceBindingIndex = FuncIndex;
+		InterfaceBinding.BindingFunctionName = FuncDef.Name;
+		InterfaceBinding.BindingFunctionNamespace = FString();
+				
+		OutInputDataBindings.Add(InKernelSource->ExternalInputs.Num(), InterfaceBinding);  
+				
+		InKernelSource->ExternalInputs.Emplace(FuncDef);	
+	}
+}
+
 #undef LOCTEXT_NAMESPACE
diff --git a/Engine/Plugins/Animation/DeformerGraph/Source/OptimusCore/Private/Nodes/OptimusNode_ComputeKernelBase.h b/Engine/Plugins/Animation/DeformerGraph/Source/OptimusCore/Private/Nodes/OptimusNode_ComputeKernelBase.h
--- a/Engine/Plugins/Animation/DeformerGraph/Source/OptimusCore/Private/Nodes/OptimusNode_ComputeKernelBase.h
+++ b/Engine/Plugins/Animation/DeformerGraph/Source/OptimusCore/Private/Nodes/OptimusNode_ComputeKernelBase.h
@@ -41,13 +41,15 @@ public:
 		const TArray<const UOptimusNode*>& InValueNodes,
 		const UComputeDataInterface* InGraphDataInterface,
 		const UOptimusComponentSourceBinding* InGraphDataComponentBinding,
+		const UComputeDataInterface* InKernelDataInterface,
 		FOptimus_InterfaceBindingMap& OutInputDataBindings,
 		FOptimus_InterfaceBindingMap& OutOutputDataBindings
 	) const override;
 
 	FName GetExecutionDomain() const override PURE_VIRTUAL(UOptimusNode_ComputeKernelBase::GetExecutionDomain, return NAME_None; );
-	TArray<const UOptimusNodePin*> GetPrimaryGroupInputPins() const override PURE_VIRTUAL(UOptimusNode_ComputeKernelBase::GetPrimaryGroupInputPins, return {}; );  
-
+	const UOptimusNodePin* GetPrimaryGroupPin() const override PURE_VIRTUAL(UOptimusNode_ComputeKernelBase::GetPrimaryGroupPin, return {}; );  
+	UComputeDataInterface* GetKernelDataInterface(UObject* InOuter) const override PURE_VIRTUAL(UOptimusNode_ComputeKernelBase::GetKernelDataInterface, return {}; );
+	
 	// -- UOptimusNode overrides
 	TOptional<FText> ValidateForCompile() const override;
 	
@@ -71,7 +73,8 @@ protected:
 		const FString& InKernelName,
 		FIntVector InGroupSize
 		);
-	
+
+	static TSet<UOptimusComponentSourceBinding*> GetGroupComponentSourceBindings(const UOptimusNodePin* InGroupPin);
 	
 private:
 	TOptional<FText> ProcessInputPinForComputeKernel(
@@ -98,5 +101,11 @@ private:
 		FOptimus_InterfaceBindingMap& OutOutputDataBindings
 		) const;
 
+	void BindKernelDataInterfaceForComputeKernel(
+		const UOptimusComponentSourceBinding* InKernelPrimaryComponentSourceBinding,
+		const UComputeDataInterface* InKernelDataInterface,
+		UOptimusKernelSource* InKernelSource,
+		FOptimus_InterfaceBindingMap& OutInputDataBindings
+		) const;
 };
 
diff --git a/Engine/Plugins/Animation/DeformerGraph/Source/OptimusCore/Private/Nodes/OptimusNode_ComputeKernelFunction.h b/Engine/Plugins/Animation/DeformerGraph/Source/OptimusCore/Private/Nodes/OptimusNode_ComputeKernelFunction.h
--- a/Engine/Plugins/Animation/DeformerGraph/Source/OptimusCore/Private/Nodes/OptimusNode_ComputeKernelFunction.h
+++ b/Engine/Plugins/Animation/DeformerGraph/Source/OptimusCore/Private/Nodes/OptimusNode_ComputeKernelFunction.h
@@ -77,7 +77,9 @@ public:
 	
 	// IOptimusComputeKernelProvider
 	FName GetExecutionDomain() const override; 
-	TArray<const UOptimusNodePin*> GetPrimaryGroupInputPins() const override { return {}; }
+	const UOptimusNodePin* GetPrimaryGroupPin() const override { return {}; }
+	UComputeDataInterface* GetKernelDataInterface(UObject* InOuter) const override { return nullptr; };
+	
 private:
 	UOptimusNode_ComputeKernelFunctionGeneratorClass *GetGeneratorClass() const;
 };
diff --git a/Engine/Plugins/Animation/DeformerGraph/Source/OptimusCore/Private/Nodes/OptimusNode_CustomComputeKernel.cpp b/Engine/Plugins/Animation/DeformerGraph/Source/OptimusCore/Private/Nodes/OptimusNode_CustomComputeKernel.cpp
--- a/Engine/Plugins/Animation/DeformerGraph/Source/OptimusCore/Private/Nodes/OptimusNode_CustomComputeKernel.cpp
+++ b/Engine/Plugins/Animation/DeformerGraph/Source/OptimusCore/Private/Nodes/OptimusNode_CustomComputeKernel.cpp
@@ -14,7 +14,8 @@
 #include "Actions/OptimusNodeActions.h"
 #include "ComponentSources/OptimusSkeletalMeshComponentSource.h"
 #include "ComponentSources/OptimusSkinnedMeshComponentSource.h"
-#include "DataInterfaces/OptimusDataInterfaceSkinnedMeshExec.h"
+#include "DataInterfaces/OptimusDataInterfaceCustomComputeKernel.h"
+#include "IOptimusDeprecatedExecutionDataInterface.h"
 #include "Engine/UserDefinedStruct.h"
 
 #include UE_INLINE_GENERATED_CPP_BY_NAME(OptimusNode_CustomComputeKernel)
@@ -27,9 +28,11 @@ static const FName InputBindingsName= GET_MEMBER_NAME_CHECKED(UOptimusNode_Custo
 static const FName OutputBindingsName = GET_MEMBER_NAME_CHECKED(UOptimusNode_CustomComputeKernel, OutputBindingArray);
 static const FName ExtraInputBindingGroupsName = GET_MEMBER_NAME_CHECKED(UOptimusNode_CustomComputeKernel, SecondaryInputBindingGroups);
 
-static const FName AdderPinActionKeyAddToPrimaryGroup("@Primary");
 static const FName AdderPinActionKeyAddNewGroup("@NewGroup");
 
+// Secondary group name cannot have spaces, so putting a space here should guard against name collision
+static const FName PrimaryGroupPinName = TEXT("Primary Group");
+
 static FString GetShaderValueTypeFriendlyName(const FOptimusDataTypeRef& InDataType)
 {
 	FName FriendlyName = NAME_None;
@@ -41,7 +44,7 @@ static FString GetShaderValueTypeFriendlyName(const FOptimusDataTypeRef& InDataT
 }
 
 
-static bool IsConnectableInputPin(const UOptimusNodePin *InPin)
+static bool IsConnectableInputDataPin(const UOptimusNodePin *InPin)
 {
 	return InPin->GetDirection() == EOptimusNodePinDirection::Input && !InPin->IsGroupingPin();
 }
@@ -56,6 +59,11 @@ static bool IsGroupingInputPin(const UOptimusNodePin *InPin)
 	return InPin->GetDirection() == EOptimusNodePinDirection::Input && InPin->IsGroupingPin();
 }
 
+static bool IsSecondaryGroupInputPin(const UOptimusNodePin *InPin)
+{
+	return InPin->GetName() != PrimaryGroupPinName && InPin->GetDirection() == EOptimusNodePinDirection::Input && InPin->IsGroupingPin();
+}
+
 bool UOptimusNode_CustomComputeKernel::DoesSourceSupportUnifiedDispatch(const UOptimusNodePin& InOtherNodesPin)
 {
 	if (InOtherNodesPin.GetDirection() == EOptimusNodePinDirection::Output )
@@ -90,17 +98,18 @@ FString UOptimusNode_CustomComputeKernel::GetKernelSourceText() const
 	return GetCookedKernelSource(GetPathName(), ShaderSource.ShaderText, KernelName.ToString(), GroupSize);
 }
 
-TArray<const UOptimusNodePin*> UOptimusNode_CustomComputeKernel::GetPrimaryGroupInputPins() const
+const UOptimusNodePin* UOptimusNode_CustomComputeKernel::GetPrimaryGroupPin() const
 {
-	TArray<const UOptimusNodePin*> PrimaryGroupInputPins;
-	for (const UOptimusNodePin* Pin: GetPins())
-	{
-		if (IsConnectableInputPin(Pin))
-		{
-			PrimaryGroupInputPins.Add(Pin);
-		}
-	}
-	return PrimaryGroupInputPins;
+	return GetPrimaryGroupPin_Internal();
+}
+
+UComputeDataInterface* UOptimusNode_CustomComputeKernel::GetKernelDataInterface(UObject* InOuter) const
+{
+	UOptimusCustomComputeKernelDataInterface* KernelDataInterface = NewObject<UOptimusCustomComputeKernelDataInterface>(InOuter);
+	
+	KernelDataInterface->InitFromKernelNode(this);
+	
+	return KernelDataInterface;
 }
 
 #if WITH_EDITOR
@@ -161,11 +170,19 @@ TArray<IOptimusNodeAdderPinProvider::FAdderPinAction> UOptimusNode_CustomCompute
 	FString* OutReason
 	) const
 {
+	if (InSourcePin->GetDataType().IsValid() && InSourcePin->GetDataType()->UsageFlags == EOptimusDataTypeUsageFlags::None)
+	{
+		if (OutReason)
+		{
+			*OutReason = TEXT("Can't add pin with this type");
+		}
+
+		return {};
+	}
+	
 	TSet<UOptimusComponentSourceBinding*> SourceComponentBindings = InSourcePin->GetComponentSourceBindings();
 	
 	TArray<UOptimusNodePin*> GroupPins;
-	// Nullptr represents the primary group
-	GroupPins.Add(nullptr);
 
 	for (UOptimusNodePin* GroupPin: GetPins())
 	{
@@ -188,21 +205,17 @@ TArray<IOptimusNodeAdderPinProvider::FAdderPinAction> UOptimusNode_CustomCompute
 		}
 
 		FAdderPinAction	Action;
-		Action.DisplayName = TEXT("Primary Group");
+		// Action is unique, no need to set name/key
+		Action.DisplayName = NAME_None;
 		Action.NewPinDirection = InNewPinDirection;
-		Action.Key = AdderPinActionKeyAddToPrimaryGroup;
+		Action.Key = NAME_None;
 		Action.ToolTip = FText::GetEmpty();
 		Action.bCanAutoLink = false;
 
 		if (InSourcePin->GetDirection() != InNewPinDirection)
 		{
-			TArray<const UOptimusNodePin*> PrimaryGroupPins = GetPrimaryGroupInputPins();
-			TSet<UOptimusComponentSourceBinding*> PrimaryGroupComponentBindings;
-			for (const UOptimusNodePin* Pin : PrimaryGroupPins)
-			{
-				PrimaryGroupComponentBindings.Append(Pin->GetComponentSourceBindings());
-			}
-		
+			const UOptimusNodePin* PrimaryGroupPin = GetPrimaryGroupPin();
+			TSet<UOptimusComponentSourceBinding*> PrimaryGroupComponentBindings = GetGroupComponentSourceBindings(PrimaryGroupPin);
 			TSet<UOptimusComponentSourceBinding*> MatchingBindings = PrimaryGroupComponentBindings.Intersect(SourceComponentBindings);
 		
 			if (MatchingBindings.Num() > 0 ||
@@ -227,19 +240,12 @@ TArray<IOptimusNodeAdderPinProvider::FAdderPinAction> UOptimusNode_CustomCompute
 			{
 				FAdderPinAction Action;
 				Action.NewPinDirection = InNewPinDirection;
-				if (GroupPin == nullptr)
-				{
-					Action.DisplayName = TEXT("Primary Group");
-					Action.Key = AdderPinActionKeyAddToPrimaryGroup;
-				}
-				else
-				{
-					Action.DisplayName = GroupPin->GetFName();
-					Action.Key = GroupPin->GetFName();
-				}
-
+				Action.DisplayName = GroupPin->GetFName();
+				Action.Key = GroupPin->GetFName();
+				
 				Action.bCanAutoLink = false;
 				Action.ToolTip = LOCTEXT("AdderPinToolTipDirectionMismatch", "Add pin only, can't link input to input");
+				
 				Actions.Add(Action);
 			}
 
@@ -251,33 +257,15 @@ TArray<IOptimusNodeAdderPinProvider::FAdderPinAction> UOptimusNode_CustomCompute
 			
 			const bool bSourceSupportUnifiedDispatch = DoesSourceSupportUnifiedDispatch(*InSourcePin);
 			
-			for (UOptimusNodePin* GroupPin: GroupPins)
+			for (int32 GroupIndex = 0; GroupIndex < GroupPins.Num(); GroupIndex++)
 			{
+				UOptimusNodePin* GroupPin = GroupPins[GroupIndex];
 				FAdderPinAction Action;
 				Action.NewPinDirection = InNewPinDirection;
-				TArray<const UOptimusNodePin*> PinsPerGroup;
-
-				if (GroupPin == nullptr)
-				{
-					Action.DisplayName = TEXT("Primary Group");
-					Action.Key = AdderPinActionKeyAddToPrimaryGroup;
-					
-					PinsPerGroup = GetPrimaryGroupInputPins();
-				}
-				else 
-				{
-					Action.DisplayName = GroupPin->GetFName();
-					Action.Key = GroupPin->GetFName();
+				Action.DisplayName = GroupPin->GetFName();
+				Action.Key = GroupPin->GetFName();
 					
-					PinsPerGroup = GroupPin->GetSubPins();
-				}
-
-				TSet<UOptimusComponentSourceBinding*> ComponentBindingsForGroup;
-				for (const UOptimusNodePin* Pin : PinsPerGroup)
-				{
-					ComponentBindingsForGroup.Append(Pin->GetComponentSourceBindings());
-				}
-
+				TSet<UOptimusComponentSourceBinding*> ComponentBindingsForGroup = GetGroupComponentSourceBindings(GroupPin);
 				TSet<UOptimusComponentSourceBinding*> MatchingBindings = ComponentBindingsForGroup.Intersect(SourceComponentBindings);
 	
 				if (MatchingBindings.Num() > 0 ||
@@ -294,7 +282,7 @@ TArray<IOptimusNodeAdderPinProvider::FAdderPinAction> UOptimusNode_CustomCompute
 
 				// Extra restrictions
 				// Secondary group does not accept data interface that does not support non unified dispatch
-				if (GroupPin != nullptr && !bSourceSupportUnifiedDispatch)
+				if (GroupIndex >= 1 && !bSourceSupportUnifiedDispatch)
 				{
 					Action.bCanAutoLink = false;
 					Action.ToolTip = LOCTEXT("AdderPinToolTipCannotLinkNonUnified", "Add pin only, can't link due to data interface not supporting unified dispatch, consider insert kernel inbetween");
@@ -367,7 +355,7 @@ TArray<UOptimusNodePin*> UOptimusNode_CustomComputeKernel::TryAddPinFromPin(
 
 			UOptimusComponentSourceBinding* ComponentBinding = *ComponentBindings.begin();
 
-			FName GroupName = Optimus::GetUniqueNameForScope(this, ComponentBinding->GetFName());
+			FName GroupName = GetSanitizedNewPinName(nullptr, ComponentBinding->GetFName());
 
 			SecondaryInputBindingGroups.AddDefaulted_GetRef().GroupName = GroupName;
 			BindingArray = &SecondaryInputBindingGroups.Last().BindingArray.InnerArray;
@@ -388,46 +376,44 @@ TArray<UOptimusNodePin*> UOptimusNode_CustomComputeKernel::TryAddPinFromPin(
 				}
 			}	
 			
-			if (InSelectedAction.Key == AdderPinActionKeyAddToPrimaryGroup)
+			for (int32 GroupIndex = 0; GroupIndex < GroupPins.Num(); GroupIndex++)
 			{
-				BindingArray = &InputBindingArray.InnerArray;
+				UOptimusNodePin* GroupPin = GroupPins[GroupIndex];
 				
-				ParentPin = nullptr;
-				BeforePin = GroupPins.Num() == 0 ? nullptr : GroupPins[0];
-			}
-			else
-			{
-				for (UOptimusNodePin* GroupPin : GroupPins)
+				if (GroupPin->GetFName() == InSelectedAction.Key)
 				{
-					if (GroupPin->GetFName() == InSelectedAction.Key)
+					ParentPin = GroupPin;
+
+					if (GroupIndex == 0)
 					{
-						ParentPin = GroupPin;
-						break;
+						BindingArray = &InputBindingArray.InnerArray;
 					}
-				}
+					else
+					{
+						// Secondary group
+						FOptimusSecondaryInputBindingsGroup* Group = SecondaryInputBindingGroups.FindByPredicate(
+						[ParentPin](const FOptimusSecondaryInputBindingsGroup& InGroup)
+						{
+							return InGroup.GroupName == ParentPin->GetFName();
+						});
 
-				if (!ensure(ParentPin))
-				{
-					return {};
-				}
+						if (!ensure(Group))
+						{
+							return {};
+						}
 				
-				// Secondary group
-				FOptimusSecondaryInputBindingsGroup* Group = SecondaryInputBindingGroups.FindByPredicate(
-				[ParentPin](const FOptimusSecondaryInputBindingsGroup& InGroup)
-				{
-					return InGroup.GroupName == ParentPin->GetFName();
-				});
-
-				if (!ensure(Group))
-				{
-					return {};
+						BindingArray = &Group->BindingArray.InnerArray;
+					}
+					break;
 				}
-				
-				BindingArray = &Group->BindingArray.InnerArray;
 			}
 		}
 	}	
 
+	if (!ensure(BindingArray))
+	{
+		return {};
+	}
 
 	if (ParentPin)
 	{
@@ -473,18 +459,21 @@ bool UOptimusNode_CustomComputeKernel::RemoveAddedPins(
 	
 		if (UOptimusNodePin* GroupPin = PinToRemove->GetParentPin())
 		{
-			FOptimusSecondaryInputBindingsGroup* Group = SecondaryInputBindingGroups.FindByPredicate(
-			   [GroupPin](const FOptimusSecondaryInputBindingsGroup& InGroup)
-			   {
-				   return InGroup.GroupName == GroupPin->GetFName();
-			   });
-
-			if (!ensure(Group))
+			if (GroupPin->GetFName() != PrimaryGroupPinName)
 			{
-				return false;
+				FOptimusSecondaryInputBindingsGroup* Group = SecondaryInputBindingGroups.FindByPredicate(
+				   [GroupPin](const FOptimusSecondaryInputBindingsGroup& InGroup)
+				   {
+					   return InGroup.GroupName == GroupPin->GetFName();
+				   });
+				
+				if (!ensure(Group))
+				{
+					return false;
+				}
+
+				BindingArray = &Group->BindingArray.InnerArray;
 			}
-		
-			BindingArray = &Group->BindingArray.InnerArray;
 		}
 	
 
@@ -517,19 +506,31 @@ bool UOptimusNode_CustomComputeKernel::RemoveAddedPins(
 
 
 FName UOptimusNode_CustomComputeKernel::GetSanitizedNewPinName(
-	UOptimusNodePin* InTargetParentPin,
+	UOptimusNodePin* InParentPin,
 	FName InPinName
 	)
 {
 	FName NewName = Optimus::GetSanitizedNameForHlsl(InPinName);
-
-	UObject* ScopeObject = this;
-	if (InTargetParentPin)
+	
+	// "NumThreads" is reserved for the kernel data interface
+	if (NewName == *UOptimusCustomComputeKernelDataInterface::NumThreadsReservedName)
 	{
-		ScopeObject = InTargetParentPin;
+		NewName.SetNumber(1);
 	}
+
+	UOptimusNodePin* PrimaryGroupPin = GetPrimaryGroupPin_Internal();
 	
-	NewName = Optimus::GetUniqueNameForScope(ScopeObject, NewName);
+	if (InParentPin && InParentPin != PrimaryGroupPin)
+	{
+		NewName = Optimus::GetUniqueNameForScope(InParentPin, NewName);
+	}
+	else
+	{
+		// Primary Input and Output need to name check against node scope as well as primary group pin scope
+		NewName = Optimus::GetUniqueNameForScope(this, NewName);
+		NewName = Optimus::GetUniqueNameForScope(PrimaryGroupPin, NewName);
+	}
+
 
 	return NewName;
 }
@@ -538,16 +539,7 @@ TArray<FName> UOptimusNode_CustomComputeKernel::GetExecutionDomains() const
 {
 	// Find all component sources for the primary pins. If we end up with any other number
 	// than one, then something's gone wrong and we can't determine the execution domains.
-	TSet<UOptimusComponentSourceBinding*> PrimaryBindings;
-
-	const UOptimusNodeGraph* Graph = GetOwningGraph();
-	for (const UOptimusNodePin* Pin: GetPins())
-	{
-		if (IsConnectableInputPin(Pin))
-		{
-			PrimaryBindings.Append(Graph->GetComponentSourceBindingsForPin(Pin));
-		}
-	}
+	TSet<UOptimusComponentSourceBinding*> PrimaryBindings = GetGroupComponentSourceBindings(GetPrimaryGroupPin());
 
 	if (PrimaryBindings.Num() == 1)
 	{
@@ -636,7 +628,7 @@ void UOptimusNode_CustomComputeKernel::PropertyArrayPasted(const FPropertyChange
 		// 2. Pasting into the binding array of one of the groups
 		//		- Unfortunately we cannot know which group is the subject, so recreate the entire group array as well
 
-		RemoveAllPinsByPredicate(GetPins(), IsGroupingInputPin);
+		RemoveAllPinsByPredicate(GetPins(), IsSecondaryGroupInputPin);
 
 		for (const FOptimusSecondaryInputBindingsGroup& InputGroup: SecondaryInputBindingGroups)
 		{
@@ -658,13 +650,19 @@ void UOptimusNode_CustomComputeKernel::PropertyArrayPasted(const FPropertyChange
 	}
 	else if (InPropertyChangedEvent.GetMemberPropertyName() == InputBindingsName)
 	{
-		RemoveAllPinsByPredicate(GetPins(), IsConnectableInputPin);
+		UOptimusNodePin* PrimaryGroupPin = GetPrimaryGroupPin_Internal();
+		RemoveAllPinsByPredicate(PrimaryGroupPin->GetSubPins() , IsConnectableInputDataPin);
 
 		for (const FOptimusParameterBinding& Binding: InputBindingArray)
 		{
-			AddPin(Binding.Name, EOptimusNodePinDirection::Input, Binding.DataDomain, Binding.DataType);
+			AddPin(Binding.Name, EOptimusNodePinDirection::Input, Binding.DataDomain, Binding.DataType, nullptr, PrimaryGroupPin);
 		}
 
+		if (InputBindingArray.Num() > 0)
+		{
+			PrimaryGroupPin->SetIsExpanded(true);
+		}
+		
 		UpdatePreamble();
 	}
 	else if (InPropertyChangedEvent.GetMemberPropertyName() == OutputBindingsName)
@@ -685,11 +683,11 @@ void UOptimusNode_CustomComputeKernel::PropertyValueChanged(
 	)
 {
 	auto UpdatePinNamesFromBindings = [this](
-		UObject* InNameScope,
+		UOptimusNodePin* InParentPin,
 		TArrayView<UOptimusNodePin* const> InPins,
 		TFunction<bool(const UOptimusNodePin*)> InPinFilter,
 		FOptimusParameterBindingArray &InBindings,
-		TFunction<void(UObject*, FOptimusParameterBinding&, UOptimusNodePin*)> InApplyFunc
+		TFunction<void(UOptimusNodePin*, FOptimusParameterBinding&, UOptimusNodePin*)> InApplyFunc
 		)
 	{
 		int32 Index = 0;
@@ -699,7 +697,7 @@ void UOptimusNode_CustomComputeKernel::PropertyValueChanged(
 			{
 				FOptimusParameterBinding& Binding = InBindings[Index];
 
-				InApplyFunc(InNameScope, Binding, Pin);
+				InApplyFunc(InParentPin, Binding, Pin);
 
 				Index++;
 			}
@@ -707,17 +705,18 @@ void UOptimusNode_CustomComputeKernel::PropertyValueChanged(
 	};
 
 	auto UpdateAllBindings = [&](
-		TFunction<void(UObject*, FOptimusParameterBinding&, UOptimusNodePin*)> InApplyFunc
+		TFunction<void(UOptimusNodePin*, FOptimusParameterBinding&, UOptimusNodePin*)> InApplyFunc
 		) -> bool
 	{
 		if (InPropertyChangedEvent.GetMemberPropertyName() == InputBindingsName)
 		{
-			UpdatePinNamesFromBindings(this, GetPins(), IsConnectableInputPin, InputBindingArray, InApplyFunc);
+			UOptimusNodePin* PrimaryGroupPin = GetPrimaryGroupPin_Internal();	
+			UpdatePinNamesFromBindings(PrimaryGroupPin,  PrimaryGroupPin->GetSubPins(), IsConnectableInputDataPin, InputBindingArray, InApplyFunc);
 			return true;
 		}
 		if (InPropertyChangedEvent.GetMemberPropertyName() == OutputBindingsName)
 		{
-			UpdatePinNamesFromBindings(this, GetPins(), IsConnectableOutputPin, OutputBindingArray, InApplyFunc);
+			UpdatePinNamesFromBindings(nullptr, GetPins(), IsConnectableOutputPin, OutputBindingArray, InApplyFunc);
 			return true;
 		}
 		if (InPropertyChangedEvent.GetMemberPropertyName() == ExtraInputBindingGroupsName)
@@ -725,10 +724,10 @@ void UOptimusNode_CustomComputeKernel::PropertyValueChanged(
 			int32 GroupIndex = 0;
 			for (UOptimusNodePin* GroupPin: GetPins())
 			{
-				if (IsGroupingInputPin(GroupPin))
+				if (IsSecondaryGroupInputPin(GroupPin))
 				{
 					FOptimusSecondaryInputBindingsGroup& BindingsGroup = SecondaryInputBindingGroups[GroupIndex];
-					UpdatePinNamesFromBindings(GroupPin, GroupPin->GetSubPins(), IsConnectableInputPin, BindingsGroup.BindingArray, InApplyFunc);
+					UpdatePinNamesFromBindings(GroupPin, GroupPin->GetSubPins(), IsConnectableInputDataPin, BindingsGroup.BindingArray, InApplyFunc);
 
 					GroupIndex++;
 				}
@@ -739,16 +738,16 @@ void UOptimusNode_CustomComputeKernel::PropertyValueChanged(
 		return false;
 	};
 	
-	auto UpdateName = [this](UObject* InNameScope, FOptimusParameterBinding& InBinding, UOptimusNodePin* InPin)
+	auto UpdateName = [this](UOptimusNodePin* InParentPin, FOptimusParameterBinding& InBinding, UOptimusNodePin* InPin)
 	{
 		if (InPin->GetFName() != InBinding.Name)
 		{
-			InBinding.Name = Optimus::GetUniqueNameForScope(InNameScope, InBinding.Name);
+			InBinding.Name = GetSanitizedNewPinName(InParentPin, InBinding.Name);
 			SetPinName(InPin, InBinding.Name);
 		}
 	};
 
-	auto UpdatePinType = [this](UObject* , FOptimusParameterBinding& InBinding, UOptimusNodePin* InPin)
+	auto UpdatePinType = [this](UOptimusNodePin* , FOptimusParameterBinding& InBinding, UOptimusNodePin* InPin)
 	{
 		const FOptimusDataTypeHandle DataType = InBinding.DataType.Resolve();
 		if (InPin->GetDataType() != DataType)
@@ -757,7 +756,7 @@ void UOptimusNode_CustomComputeKernel::PropertyValueChanged(
 		}
 	};
 
-	auto UpdatePinDataDomain = [this](UObject* , FOptimusParameterBinding& InBinding, UOptimusNodePin* InPin)
+	auto UpdatePinDataDomain = [this](UOptimusNodePin* , FOptimusParameterBinding& InBinding, UOptimusNodePin* InPin)
 	{
 		if (InPin->GetDataDomain() != InBinding.DataDomain)
 		{
@@ -765,11 +764,11 @@ void UOptimusNode_CustomComputeKernel::PropertyValueChanged(
 		}
 	};
 	
-	auto UpdatePin = [&](UObject* InNameScope, FOptimusParameterBinding& InBinding, UOptimusNodePin* InPin)
+	auto UpdatePin = [&](UOptimusNodePin* InParentPin, FOptimusParameterBinding& InBinding, UOptimusNodePin* InPin)
 	{
-		UpdateName(InNameScope, InBinding, InPin);
-		UpdatePinType(InNameScope, InBinding, InPin);
-		UpdatePinDataDomain(InNameScope, InBinding, InPin);
+		UpdateName(InParentPin, InBinding, InPin);
+		UpdatePinType(InParentPin, InBinding, InPin);
+		UpdatePinDataDomain(InParentPin, InBinding, InPin);
 	};
 	
 	if (InPropertyChangedEvent.GetMemberPropertyName() == GET_MEMBER_NAME_STRING_CHECKED(UOptimusNode_CustomComputeKernel, KernelName) && 
@@ -788,12 +787,12 @@ void UOptimusNode_CustomComputeKernel::PropertyValueChanged(
 		bool bUpdatePreamble = false;
 		for (UOptimusNodePin* GroupPin: GetPins())
 		{
-			if (IsGroupingInputPin(GroupPin))
+			if (IsSecondaryGroupInputPin(GroupPin))
 			{
 				FOptimusSecondaryInputBindingsGroup& SecondaryGroup = SecondaryInputBindingGroups[Index++]; 
 				if (GroupPin->GetFName() != SecondaryGroup.GroupName)
 				{
-					SecondaryGroup.GroupName = Optimus::GetUniqueNameForScope(this, SecondaryGroup.GroupName);
+					SecondaryGroup.GroupName = GetSanitizedNewPinName(nullptr, SecondaryGroup.GroupName);
 					SetPinName(GroupPin, SecondaryGroup.GroupName);
 					bUpdatePreamble = true;
 				}
@@ -865,7 +864,7 @@ void UOptimusNode_CustomComputeKernel::PropertyArrayItemAdded(
 		UOptimusNodePin* InGroupPin = nullptr
 	) -> void
 	{
-		InBinding.Name = Optimus::GetUniqueNameForScope(InGroupPin ? Cast<UObject>(InGroupPin) : this, InName);
+		InBinding.Name = GetSanitizedNewPinName(InGroupPin, InName);
 		InBinding.DataType = FOptimusDataTypeRegistry::Get().FindType(*FFloatProperty::StaticClass());
 		if (InDirection == EOptimusNodePinDirection::Input)
 		{
@@ -896,29 +895,6 @@ void UOptimusNode_CustomComputeKernel::PropertyArrayItemAdded(
 		UpdatePreamble();
 	};
 
-	// Find the top-level pin before the index of the given direction. If we reach a group pin, return that so that
-	// we can insert before the group pin.
-	auto FindTopLevelBeforePin = [this](
-		EOptimusNodePinDirection InDirection,
-		int32 InsertIndex
-		) -> UOptimusNodePin * 
-	{
-		for (UOptimusNodePin *Pin: GetPins())
-		{
-			if (Pin->IsGroupingPin())
-			{
-				return Pin;
-			}
-			if (Pin->GetDirection() == InDirection)
-			{
-				if (InsertIndex-- == 0)
-				{
-					return Pin;
-				}
-			}
-		}
-		return nullptr;
-	};
 
 	if (InPropertyChangedEvent.GetMemberPropertyName() == ExtraInputBindingGroupsName)
 	{
@@ -935,10 +911,11 @@ void UOptimusNode_CustomComputeKernel::PropertyArrayItemAdded(
 			UOptimusNodePin* BeforePin = nullptr;
 			if (GroupIndex < (SecondaryInputBindingGroups.Num() - 1))
 			{
-				BeforePin = GetPins()[InputBindingArray.Num() + OutputBindingArray.Num() + GroupIndex];
+				// The top level hierarchy is always PrimaryGroupPin -> OutputPins -> SecondaryGroupPins
+				BeforePin = GetPins()[1 + OutputBindingArray.Num() + GroupIndex];
 			}
 			
-			const FName GroupName = Optimus::GetUniqueNameForScope(this, DefaultGroupName);
+			const FName GroupName = GetSanitizedNewPinName(nullptr, DefaultGroupName);
 
 			AddGroupingPin(GroupName, EOptimusNodePinDirection::Input, BeforePin);
 
@@ -959,7 +936,7 @@ void UOptimusNode_CustomComputeKernel::PropertyArrayItemAdded(
 			UOptimusNodePin* GroupPin = nullptr;
 			for (UOptimusNodePin *Pin: GetPins())
 			{
-				if (Pin->IsGroupingPin())
+				if (IsSecondaryGroupInputPin(Pin))
 				{
 					if (SecondaryInputBindingGroups[GroupIndex].BindingArray.Num() == (Pin->GetSubPins().Num() + 1))
 					{
@@ -994,17 +971,44 @@ void UOptimusNode_CustomComputeKernel::PropertyArrayItemAdded(
 	if (InPropertyChangedEvent.GetMemberPropertyName() == InputBindingsName)
 	{
 		const int32 BindingIndex = GetArrayIndex(GET_MEMBER_NAME_STRING_CHECKED(FOptimusParameterBindingArray, InnerArray), InputBindingArray);
-		
-		UOptimusNodePin* BeforePin = FindTopLevelBeforePin(EOptimusNodePinDirection::Input, BindingIndex);
 
-		AddPinForBinding(InputBindingArray[BindingIndex], "Input", EOptimusNodePinDirection::Input, BeforePin);
+		UOptimusNodePin* PrimaryGroupPin = GetPrimaryGroupPin_Internal();
+		UOptimusNodePin* BeforePin = PrimaryGroupPin->GetSubPins().IsValidIndex(BindingIndex) ? PrimaryGroupPin->GetSubPins()[BindingIndex] : nullptr;
+
+		AddPinForBinding(InputBindingArray[BindingIndex], "Input", EOptimusNodePinDirection::Input, BeforePin, PrimaryGroupPin);
+
+		if (InputBindingArray.Num() > 0)
+		{
+			PrimaryGroupPin->SetIsExpanded(true);
+		}
 	}
 	else if (InPropertyChangedEvent.GetMemberPropertyName() == OutputBindingsName)
 	{
 		const int32 BindingIndex = GetArrayIndex(GET_MEMBER_NAME_STRING_CHECKED(FOptimusParameterBindingArray, InnerArray), OutputBindingArray);
 
-		UOptimusNodePin* BeforePin = FindTopLevelBeforePin(EOptimusNodePinDirection::Output, BindingIndex);
-		
+		UOptimusNodePin* BeforePin = nullptr;
+
+		int32 PinCount = BindingIndex;
+		for (UOptimusNodePin *Pin: GetPins())
+		{
+			// Inserting after the last output pin
+			if (IsSecondaryGroupInputPin(Pin))
+			{
+				BeforePin = Pin;
+				break;
+			}
+
+			// Inserting before the last output pin
+			if (Pin->GetDirection() == EOptimusNodePinDirection::Output)
+			{
+				if (PinCount-- == 0)
+				{
+					BeforePin = Pin;
+					break;
+				}
+			}
+		}
+
 		AddPinForBinding(OutputBindingArray[BindingIndex], "Output", EOptimusNodePinDirection::Output, BeforePin);
 	}
 }
@@ -1046,7 +1050,7 @@ void UOptimusNode_CustomComputeKernel::PropertyArrayItemRemoved(const FPropertyC
 
 			for (UOptimusNodePin *Pin: GetPins())
 			{
-				if (Pin->IsGroupingPin())
+				if (IsSecondaryGroupInputPin(Pin))
 				{
 					if (GroupIndex-- == 0)
 					{
@@ -1070,7 +1074,7 @@ void UOptimusNode_CustomComputeKernel::PropertyArrayItemRemoved(const FPropertyC
 			UOptimusNodePin* GroupPin = nullptr;
 			for (UOptimusNodePin *Pin: GetPins())
 			{
-				if (Pin->IsGroupingPin())
+				if (IsSecondaryGroupInputPin(Pin))
 				{
 					if (SecondaryInputBindingGroups[GroupIndex].BindingArray.Num() == (Pin->GetSubPins().Num() - 1))
 					{
@@ -1102,11 +1106,42 @@ void UOptimusNode_CustomComputeKernel::PropertyArrayItemRemoved(const FPropertyC
 	}
 	else if (InPropertyChangedEvent.GetMemberPropertyName() == InputBindingsName)
 	{
-		RemoveTopLevelPinByIndex(EOptimusNodePinDirection::Input);
+		UOptimusNodePin* PrimaryGroupPin = GetPrimaryGroupPin_Internal();
+		
+		int32 BindingIndex = InPropertyChangedEvent.GetArrayIndex(GET_MEMBER_NAME_STRING_CHECKED(FOptimusParameterBindingArray, InnerArray));
+		
+		if (ensure(PrimaryGroupPin->GetSubPins().IsValidIndex(BindingIndex)))
+		{
+			RemovePin(PrimaryGroupPin->GetSubPins()[BindingIndex]);
+			UpdatePreamble();
+		}
+
+		if (InputBindingArray.IsEmpty())
+		{
+			PrimaryGroupPin->SetIsExpanded(false);
+		}
+		
 	}
 	else if (InPropertyChangedEvent.GetMemberPropertyName() == OutputBindingsName)
 	{
-		RemoveTopLevelPinByIndex(EOptimusNodePinDirection::Output);
+		int32 PinIndex = InPropertyChangedEvent.GetArrayIndex(GET_MEMBER_NAME_STRING_CHECKED(FOptimusParameterBindingArray, InnerArray));
+		if (!ensure(PinIndex != INDEX_NONE))
+		{
+			return;
+		}
+		
+		for (UOptimusNodePin *Pin: GetPins())
+		{
+			if (IsConnectableOutputPin(Pin))
+			{
+				if (PinIndex-- == 0)
+				{
+					RemovePin(Pin);
+					UpdatePreamble();
+					break;
+				}
+			}
+		}	
 	}
 }
 
@@ -1133,7 +1168,7 @@ void UOptimusNode_CustomComputeKernel::PropertyArrayCleared(
 	{
 		if (InPropertyChangedEvent.GetPropertyName() == ExtraInputBindingGroupsName)
 		{
-			RemoveAllPinsByPredicate(GetPins(), IsGroupingInputPin);
+			RemoveAllPinsByPredicate(GetPins(), IsSecondaryGroupInputPin);
 			UpdatePreamble();
 		}
 		else if (InPropertyChangedEvent.GetPropertyName() == GET_MEMBER_NAME_CHECKED(FOptimusParameterBindingArray, InnerArray))
@@ -1149,7 +1184,7 @@ void UOptimusNode_CustomComputeKernel::PropertyArrayCleared(
 			UOptimusNodePin* GroupPin = nullptr;
 			for (UOptimusNodePin *Pin: GetPins())
 			{
-				if (Pin->IsGroupingPin())
+				if (IsSecondaryGroupInputPin(Pin))
 				{
 					// If the binding array is empty but the pins are still there, then that's our group.
 					if (SecondaryInputBindingGroups[GroupIndex].BindingArray.IsEmpty() && !Pin->GetSubPins().IsEmpty())
@@ -1169,13 +1204,17 @@ void UOptimusNode_CustomComputeKernel::PropertyArrayCleared(
 			// Since  the group is empty, collapse the pin for consistent look.
 			GroupPin->SetIsExpanded(false);
 			
-			RemoveAllPinsByPredicate(GroupPin->GetSubPins(), IsConnectableInputPin);
+			RemoveAllPinsByPredicate(GroupPin->GetSubPins(), IsConnectableInputDataPin);
 			UpdatePreamble();
 		}
 	}
 	else if (InPropertyChangedEvent.GetMemberPropertyName() == InputBindingsName)
 	{
-		RemoveAllPinsByPredicate(GetPins(), IsConnectableInputPin);
+		UOptimusNodePin* PrimaryGroupPin = GetPrimaryGroupPin_Internal();
+		RemoveAllPinsByPredicate(PrimaryGroupPin->GetSubPins(), IsConnectableInputDataPin);
+
+		PrimaryGroupPin->SetIsExpanded(false);
+		
 		UpdatePreamble();
 	}
 	else if (InPropertyChangedEvent.GetMemberPropertyName() == OutputBindingsName)
@@ -1256,7 +1295,7 @@ void UOptimusNode_CustomComputeKernel::PropertyArrayItemMoved(
 				GroupNames.Add(BindingsGroup.GroupName);
 			}
 			
-			FindMovedBindingAndMoveMatchingPin(GetPins(), GroupNames, IsGroupingInputPin);
+			FindMovedBindingAndMoveMatchingPin(GetPins(), GroupNames, IsSecondaryGroupInputPin);
 			UpdatePreamble();
 			
 		}
@@ -1288,7 +1327,7 @@ void UOptimusNode_CustomComputeKernel::PropertyArrayItemMoved(
 			UOptimusNodePin* GroupPin = nullptr;
 			for (UOptimusNodePin *Pin: GetPins())
 			{
-				if (Pin->IsGroupingPin())
+				if (IsSecondaryGroupInputPin(Pin))
 				{
 					// If the binding array is empty but the pins are still there, then that's our group.
 					if (PinsAndBindingsDiffer(SecondaryInputBindingGroups[GroupIndex].BindingArray, Pin->GetSubPins()))
@@ -1307,13 +1346,14 @@ void UOptimusNode_CustomComputeKernel::PropertyArrayItemMoved(
 
 			FindMovedBindingAndMoveMatchingPin(
 				GroupPin->GetSubPins(), MakeBindingNameArray(SecondaryInputBindingGroups[GroupIndex].BindingArray),
-				IsConnectableInputPin);
+				IsConnectableInputDataPin);
 			UpdatePreamble();
 		}
 	}
 	else if (InPropertyChangedEvent.GetMemberPropertyName() == InputBindingsName)
 	{
-		FindMovedBindingAndMoveMatchingPin(GetPins(), MakeBindingNameArray(InputBindingArray), IsConnectableInputPin);
+		UOptimusNodePin* PrimaryGroupPin = GetPrimaryGroupPin_Internal();
+		FindMovedBindingAndMoveMatchingPin(PrimaryGroupPin->GetSubPins(), MakeBindingNameArray(InputBindingArray), IsConnectableInputDataPin);
 		UpdatePreamble();
 	}
 	else if (InPropertyChangedEvent.GetMemberPropertyName() == OutputBindingsName)
@@ -1356,39 +1396,10 @@ PRAGMA_ENABLE_DEPRECATION_WARNINGS
 	
 	Super::PostLoad();
 	
-	if (GetLinkerCustomVersion(FOptimusObjectVersion::GUID) < FOptimusObjectVersion::ComponentProviderSupport)
+	if (GetLinkerCustomVersion(FOptimusObjectVersion::GUID) < FOptimusObjectVersion::KernelDataInterface)
 	{
-		ExecutionDomain.Name = UOptimusSkeletalMeshComponentSource::Domains::Vertex;
-		
-		// Check if there's an execution node connected and grab the domain from it.
-		FOptimusDataTypeHandle IntVector3Type = FOptimusDataTypeRegistry::Get().FindType(Optimus::GetTypeName(TBaseStructure<FIntVector3>::Get()));
-
-		for (const UOptimusNodePin* Pin: GetPins())
-		{
-			if (Pin->GetDirection() == EOptimusNodePinDirection::Input && Pin->GetDataType() == IntVector3Type)
-			{
-				TArray<FOptimusRoutedNodePin> ConnectedPins = Pin->GetConnectedPinsWithRouting();
-				if (ConnectedPins.Num() == 1)
-				{
-					if (const UOptimusNode_DataInterface* DataInterfaceNode = Cast<UOptimusNode_DataInterface>(ConnectedPins[0].NodePin->GetOwningNode()))
-					{
-						if (const UOptimusSkinnedMeshExecDataInterface* ExecDataInterface = Cast<UOptimusSkinnedMeshExecDataInterface>(DataInterfaceNode->GetDataInterface(GetTransientPackage())))
-						{
-							switch(ExecDataInterface->Domain)
-							{
-							default:
-							case EOptimusSkinnedMeshExecDomain::Vertex:
-								ExecutionDomain.Name = UOptimusSkinnedMeshComponentSource::Domains::Vertex;
-								break;
-							case EOptimusSkinnedMeshExecDomain::Triangle:
-								ExecutionDomain.Name = UOptimusSkinnedMeshComponentSource::Domains::Triangle;
-								break;
-							}
-						}
-					}
-				}
-			}
-		}
+		PostLoadExtractExecutionDomain();
+		PostLoadAddMissingPrimaryGroupPin();
 	}
 
 	SetDisplayName(FText::FromName(KernelName));
@@ -1401,9 +1412,12 @@ void UOptimusNode_CustomComputeKernel::ConstructNode()
 	
 	// After a duplicate, the kernel node has no pins, so we need to reconstruct them from
 	// the bindings. We can assume that all naming clashes have already been dealt with.
+
+	UOptimusNodePin* PrimaryGroupPin = AddGroupingPinDirect(PrimaryGroupPinName, EOptimusNodePinDirection::Input);
+	
 	for (const FOptimusParameterBinding& Binding: InputBindingArray)
 	{
-		AddPinDirect(Binding.Name, EOptimusNodePinDirection::Input, Binding.DataDomain, Binding.DataType);
+		AddPinDirect(Binding.Name, EOptimusNodePinDirection::Input, Binding.DataDomain, Binding.DataType, nullptr, PrimaryGroupPin);
 	}
 	for (const FOptimusParameterBinding& Binding: OutputBindingArray)
 	{
@@ -1414,7 +1428,7 @@ void UOptimusNode_CustomComputeKernel::ConstructNode()
 	for (const FOptimusSecondaryInputBindingsGroup& InputGroup: SecondaryInputBindingGroups)
 	{
 		UOptimusNodePin* GroupPin = AddGroupingPinDirect(InputGroup.GroupName, EOptimusNodePinDirection::Input);
-
+		
 		for (const FOptimusParameterBinding& Binding: InputGroup.BindingArray)
 		{
 			AddPinDirect(Binding.Name, EOptimusNodePinDirection::Input, Binding.DataDomain, Binding.DataType, nullptr, GroupPin);
@@ -1432,7 +1446,7 @@ bool UOptimusNode_CustomComputeKernel::ValidateConnection(
 	if (InThisNodesPin.GetDirection() == EOptimusNodePinDirection::Input)
 	{
 		// Secondary group connections need additional checks
-		if (IsGroupingInputPin(InThisNodesPin.GetRootPin()))
+		if (InThisNodesPin.GetRootPin()->GetFName() != PrimaryGroupPinName)
 		{
 			if (!DoesSourceSupportUnifiedDispatch(InOtherNodesPin))
 			{
@@ -1562,11 +1576,11 @@ void UOptimusNode_CustomComputeKernel::UpdatePreamble()
 	};
 
 	TArray<FOptimusParameterBinding> Bindings = InputBindingArray.InnerArray;
-	Bindings.Sort(ContextsPredicate);
+	Bindings.StableSort(ContextsPredicate);
 	AddInputBindings(Bindings);
 
 	Bindings = OutputBindingArray.InnerArray;
-	Bindings.Sort(ContextsPredicate);
+	Bindings.StableSort(ContextsPredicate);
 	for (const FOptimusParameterBinding& Binding: Bindings)
 	{
 		AddCountFunctionIfNeeded(Binding.DataDomain.DimensionNames);
@@ -1605,9 +1619,120 @@ void UOptimusNode_CustomComputeKernel::UpdatePreamble()
 	}
 	ShaderSource.Declarations += "\n// Resource Indexing\n";
 	ShaderSource.Declarations += "uint Index;	// From SV_DispatchThreadID.x\n";
+	ShaderSource.Declarations += "int3 ReadNumThreads();\n";
+}
+
+UOptimusNodePin* UOptimusNode_CustomComputeKernel::GetPrimaryGroupPin_Internal() const
+{
+	// Always the first pin, as per PostLoadAddMissingPrimaryGroupPin and ConstructNode
+	return GetPins()[0];
+}
+
+
+void UOptimusNode_CustomComputeKernel::PostLoadExtractExecutionDomain()
+{
+	ExecutionDomain.Name = NAME_None;
+		
+	// Check if there's an execution node connected and grab the domain from it.
+	FOptimusDataTypeHandle IntVector3Type = FOptimusDataTypeRegistry::Get().FindType(Optimus::GetTypeName(TBaseStructure<FIntVector3>::Get()));
+
+	for (UOptimusNodePin* Pin: GetPins())
+	{
+		if (Pin->GetDirection() == EOptimusNodePinDirection::Input && Pin->GetDataType() == IntVector3Type)
+		{
+			TArray<FOptimusRoutedNodePin> ConnectedPins = Pin->GetConnectedPinsWithRouting();
+			if (ConnectedPins.Num() == 1)
+			{
+				if (const UOptimusNode_DataInterface* DataInterfaceNode = Cast<UOptimusNode_DataInterface>(ConnectedPins[0].NodePin->GetOwningNode()))
+				{
+					if (const IOptimusDeprecatedExecutionDataInterface* ExecDataInterface = Cast<IOptimusDeprecatedExecutionDataInterface>(DataInterfaceNode->GetDataInterface(GetTransientPackage())))
+					{
+						FName Domain = ExecDataInterface->GetSelectedExecutionDomainName();
+						ExecutionDomain.Name = Domain;
+					}
+				}
+			}
+		}
+	}
+}
+
+void UOptimusNode_CustomComputeKernel::PostLoadAddMissingPrimaryGroupPin()
+{
+	// Add a group pin for the primary group since old nodes did not have one
+	TArray<UOptimusNodePin*> PrimaryGroupInputPins;
+	UOptimusNodePin* BeforePin = nullptr;
+		
+	for (UOptimusNodePin* Pin: GetPins())
+	{
+		if (IsConnectableInputDataPin(Pin))
+		{
+			PrimaryGroupInputPins.Add(Pin);
+		}
+		else if (BeforePin == nullptr)
+		{
+			BeforePin = Pin;
+		}
+	}
+
+	UOptimusNodePin* PrimaryGroupPin = AddGroupingPinDirect(PrimaryGroupPinName, EOptimusNodePinDirection::Input, BeforePin);
+		
+	for (UOptimusNodePin* InputPin : PrimaryGroupInputPins)
+	{
+		MovePinToGroupPinDirect(InputPin, PrimaryGroupPin);
+	}
+
+	// Making sure the primary group pin is the first pin
+	if (!ensure(GetPins()[0] == PrimaryGroupPin))
+	{
+		MovePinDirect(PrimaryGroupPin, GetPins()[0]);
+	}
+	
+	PrimaryGroupPin->SetIsExpanded(PrimaryGroupPin->GetSubPins().Num() > 0);
+
 }
 
 
+bool UOptimusNode_CustomComputeKernel::PostLoadRemoveDeprecatedNumThreadsPin()
+{
+	// NumThreads is now a reserved name that pins cannot use
+	const FName NumThreadsPinName = *UOptimusCustomComputeKernelDataInterface::NumThreadsReservedName;
+	
+	UOptimusNodePin* PrimaryGroupPin = GetPrimaryGroupPin_Internal();
+	
+	if (!ensure(InputBindingArray.InnerArray.Num() == PrimaryGroupPin->GetSubPins().Num()))
+	{
+		return false;
+	}
+
+	bool bNodeChanged = false;
+	for (int32 BindingIndex = 0; BindingIndex < PrimaryGroupPin->GetSubPins().Num(); BindingIndex++)
+	{
+		UOptimusNodePin* Pin = PrimaryGroupPin->GetSubPins()[BindingIndex];
+		
+		if (Pin->GetFName() != NumThreadsPinName)
+		{
+			continue;
+		}
+
+		// Pins and Binding array should match 1:1
+		if (!ensure(InputBindingArray.InnerArray[BindingIndex].Name == NumThreadsPinName))
+		{
+			continue;
+		}
+
+		RemovePin(Pin);
+		
+		InputBindingArray.InnerArray.RemoveAt(BindingIndex);
+		bNodeChanged = true;
+		
+		break;
+	}
+
+	UpdatePreamble();
+	
+	return bNodeChanged;
+}
+
 FString UOptimusNode_CustomComputeKernel::GetDeclarationForBinding(const FOptimusParameterBinding& Binding, bool bIsInput)
 {
 	TArray<FString> Indexes;
diff --git a/Engine/Plugins/Animation/DeformerGraph/Source/OptimusCore/Private/Nodes/OptimusNode_CustomComputeKernel.h b/Engine/Plugins/Animation/DeformerGraph/Source/OptimusCore/Private/Nodes/OptimusNode_CustomComputeKernel.h
--- a/Engine/Plugins/Animation/DeformerGraph/Source/OptimusCore/Private/Nodes/OptimusNode_CustomComputeKernel.h
+++ b/Engine/Plugins/Animation/DeformerGraph/Source/OptimusCore/Private/Nodes/OptimusNode_CustomComputeKernel.h
@@ -60,8 +60,9 @@ public:
 
 	// IOptimusComputeKernelProvider
 	FName GetExecutionDomain() const override { return ExecutionDomain.Name; } 
-	TArray<const UOptimusNodePin*> GetPrimaryGroupInputPins() const override;
-	
+	const UOptimusNodePin* GetPrimaryGroupPin() const override;
+	UComputeDataInterface* GetKernelDataInterface(UObject* InOuter) const override;
+
 #if WITH_EDITOR
 	// IOptimusShaderTextProvider overrides
 	FString GetNameForShaderTextEditor() const override;
@@ -90,10 +91,7 @@ public:
 		TConstArrayView<UOptimusNodePin*> InAddedPinsToRemove
 	) override;
 	
-	FName GetSanitizedNewPinName(
-		UOptimusNodePin* InTargetParentPin,
-		FName InPinName
-		) override;
+
 
 	// IOptimusExecutionDomainProvider
 	TArray<FName> GetExecutionDomains() const override;
@@ -164,6 +162,8 @@ PRAGMA_ENABLE_DEPRECATION_WARNINGS
 
 	void Serialize(FArchive& Ar) override;
 	void PostLoad() override;
+
+	
 	
 protected:
 	void ConstructNode() override;
@@ -183,7 +183,21 @@ private:
 	
 	void UpdatePreamble();
 
+	UOptimusNodePin* GetPrimaryGroupPin_Internal() const;
+	void PostLoadExtractExecutionDomain();
+	void PostLoadAddMissingPrimaryGroupPin();
+	
+	// Called during UOptimusDeformer::PostLoad()
+	bool PostLoadRemoveDeprecatedNumThreadsPin();
+	
+	FName GetSanitizedNewPinName(
+		UOptimusNodePin* InParentPin,
+		FName InPinName
+		);
+	
 	static FString GetDeclarationForBinding(const FOptimusParameterBinding& Binding, bool bIsInput);
 
 	static bool DoesSourceSupportUnifiedDispatch(const UOptimusNodePin& InOtherNodesPin);
+	friend class UOptimusDeformer; // For PostLoad fixup
+	friend class UOptimusCustomComputeKernelDataInterface;
 };
diff --git a/Engine/Plugins/Animation/DeformerGraph/Source/OptimusCore/Private/OptimusDataTypeRegistry.cpp b/Engine/Plugins/Animation/DeformerGraph/Source/OptimusCore/Private/OptimusDataTypeRegistry.cpp
--- a/Engine/Plugins/Animation/DeformerGraph/Source/OptimusCore/Private/OptimusDataTypeRegistry.cpp
+++ b/Engine/Plugins/Animation/DeformerGraph/Source/OptimusCore/Private/OptimusDataTypeRegistry.cpp
@@ -448,11 +448,12 @@ void FOptimusDataTypeRegistry::RegisterBuiltinTypes()
 	// FIXME: Add type aliases (e.g. "3x4 Float" above should really be "float3x4")
 
 	// Actor Component
+	// ComponentSourceBinding pin does not carry concrete data, and is special
 	Registry.RegisterType(
 		UOptimusComponentSourceBinding::StaticClass(),
 		FText::FromString("Component"),
 		FLinearColor(0.3f, 0.3f, 0.4f, 1.0f),
-		EOptimusDataTypeUsageFlags::Variable
+		EOptimusDataTypeUsageFlags::None
 		);
 }
 
diff --git a/Engine/Plugins/Animation/DeformerGraph/Source/OptimusCore/Private/OptimusDeformer.cpp b/Engine/Plugins/Animation/DeformerGraph/Source/OptimusCore/Private/OptimusDeformer.cpp
--- a/Engine/Plugins/Animation/DeformerGraph/Source/OptimusCore/Private/OptimusDeformer.cpp
+++ b/Engine/Plugins/Animation/DeformerGraph/Source/OptimusCore/Private/OptimusDeformer.cpp
@@ -45,6 +45,11 @@
 #include "Nodes/OptimusNode_GetVariable.h"
 #include "Nodes/OptimusNode_ResourceAccessorBase.h"
 
+#include "IOptimusDeprecatedExecutionDataInterface.h"
+#include "OptimusNodeLink.h"
+#include "Nodes/OptimusNode_CustomComputeKernel.h"
+
+
 #include UE_INLINE_GENERATED_CPP_BY_NAME(OptimusDeformer)
 
 #include <limits>
@@ -459,8 +464,7 @@ bool UOptimusDeformer::RemoveVariableDirect(
 
 	Notify(EOptimusGlobalNotifyType::VariableRemoved, InVariableDesc);
 
-	InVariableDesc->Rename(nullptr, GetTransientPackage());
-	InVariableDesc->MarkAsGarbage();
+	Optimus::RemoveObject(InVariableDesc);
 
 	(void)MarkPackageDirty();
 
@@ -479,7 +483,8 @@ bool UOptimusDeformer::RenameVariableDirect(
 		return false;
 	}
 
-	if (InVariableDesc->Rename(*InNewName.ToString(), nullptr, REN_NonTransactional))
+	
+	if (Optimus::RenameObject(InVariableDesc, *InNewName.ToString(), nullptr))
 	{
 		InVariableDesc->VariableName = InNewName;
 		Notify(EOptimusGlobalNotifyType::VariableRenamed, InVariableDesc);
@@ -878,8 +883,7 @@ bool UOptimusDeformer::RemoveResourceDirect(
 
 	Notify(EOptimusGlobalNotifyType::ResourceRemoved, InResourceDesc);
 
-	InResourceDesc->Rename(nullptr, GetTransientPackage());
-	InResourceDesc->MarkAsGarbage();
+	Optimus::RemoveObject(InResourceDesc);
 
 	(void)MarkPackageDirty();
 
@@ -900,7 +904,7 @@ bool UOptimusDeformer::RenameResourceDirect(
 	}
 
 	// Rename in a non-transactional manner, since we're handling undo/redo.
-	if (InResourceDesc->Rename(*InNewName.ToString(), nullptr, REN_NonTransactional))
+	if (Optimus::RenameObject(InResourceDesc, *InNewName.ToString(), nullptr))
 	{
 		InResourceDesc->ResourceName = InNewName;
 		Notify(EOptimusGlobalNotifyType::ResourceRenamed, InResourceDesc);
@@ -1100,9 +1104,8 @@ bool UOptimusDeformer::RemoveComponentBindingDirect(UOptimusComponentSourceBindi
 
 	Notify(EOptimusGlobalNotifyType::ComponentBindingRemoved, InBinding);
 
-	InBinding->Rename(nullptr, GetTransientPackage());
-	InBinding->MarkAsGarbage();
-
+	Optimus::RemoveObject(InBinding);
+	
 	(void)MarkPackageDirty();
 	
 	return true;
@@ -1174,7 +1177,7 @@ bool UOptimusDeformer::RenameComponentBindingDirect(
 		return false;
 	}
 	
-	if (InBinding->Rename(*InNewName.ToString(), nullptr, REN_NonTransactional))
+	if (Optimus::RenameObject(InBinding, *InNewName.ToString(), nullptr))
 	{
 		InBinding->BindingName = InNewName;
 		Notify(EOptimusGlobalNotifyType::ComponentBindingRenamed, InBinding);
@@ -1580,6 +1583,9 @@ UOptimusComputeGraph* UOptimusDeformer::CompileNodeGraphToComputeGraph(
 	// because a resource output may be used multiple times, but only written into once.
 	FOptimus_PinToDataInterfaceMap LinkDataInterfaceMap;
 
+	// Each kernel spawns a data interface to pass along kernel specific data to the kernel itself
+	FOptimus_KernelNodeToKernelDataInterfaceMap KernelDataInterfaceMap;
+
 	// Find all value nodes (constant and variable) 
 	TArray<const UOptimusNode *> ValueNodes; 
 
@@ -1596,8 +1602,12 @@ UOptimusComputeGraph* UOptimusDeformer::CompileNodeGraphToComputeGraph(
 
 			NodeDataInterfaceMap.Add(ConnectedNode.Node, DataInterface);
 		}
-		else if (Cast<const IOptimusComputeKernelProvider>(ConnectedNode.Node) != nullptr)
+		else if (const IOptimusComputeKernelProvider* KernelProvider = Cast<const IOptimusComputeKernelProvider>(ConnectedNode.Node))
 		{
+			UComputeDataInterface* KernelDataInterface = KernelProvider->GetKernelDataInterface(this);
+
+			KernelDataInterfaceMap.Add(ConnectedNode.Node, KernelDataInterface);
+			
 			for (const UOptimusNodePin* Pin: ConnectedNode.Node->GetPins())
 			{
 				if (Pin->GetDirection() == EOptimusNodePinDirection::Output &&
@@ -1682,11 +1692,14 @@ UOptimusComputeGraph* UOptimusDeformer::CompileNodeGraphToComputeGraph(
 			BoundKernel.KernelNodeIndex = InNodeGraph->Nodes.IndexOfByKey(ConnectedNode.Node);
 			BoundKernel.Kernel = NewObject<UComputeKernel>(this);
 
+			UComputeDataInterface* KernelDataInterface = KernelDataInterfaceMap[ConnectedNode.Node];
+			
 			FOptimus_ComputeKernelResult KernelSourceResult = KernelProvider->CreateComputeKernel(	
 				BoundKernel.Kernel, ConnectedNode.TraversalContext,
 				NodeDataInterfaceMap, LinkDataInterfaceMap,
 				ValueNodes,
 				GraphDataInterface, GraphDataComponentBinding,
+				KernelDataInterface,
 				BoundKernel.InputDataBindings, BoundKernel.OutputDataBindings
 			);
 			if (FText* ErrorMessage = KernelSourceResult.TryGet<FText>())
@@ -1804,6 +1817,11 @@ UOptimusComputeGraph* UOptimusDeformer::CompileNodeGraphToComputeGraph(
 		ComputeGraph->DataInterfaces.Add(Item.Value);
 		ComputeGraph->DataInterfaceToBinding.Add(DataInterfaceToBindingIndexMap[Item.Value]);
 	}
+	for (TPair<const UOptimusNode *, UComputeDataInterface *>&Item: KernelDataInterfaceMap)
+	{
+		ComputeGraph->DataInterfaces.Add(Item.Value);
+		ComputeGraph->DataInterfaceToBinding.Add(DataInterfaceToBindingIndexMap[Item.Value]);
+	}
 
 	// Create the graph edges.
 	for (int32 KernelIndex = 0; KernelIndex < ComputeGraph->KernelInvocations.Num(); KernelIndex++)
@@ -2126,6 +2144,13 @@ void UOptimusDeformer::PostLoad()
 {
 	Super::PostLoad();
 
+	// PostLoad everything first before changing anything for back compat
+	// Each graph postloads everything it owns
+	for (UOptimusNodeGraph* Graph: GetGraphs())
+	{
+		Graph->ConditionalPostLoad();
+	}
+	
 	// Fixup any empty array entries.
 	Resources->Descriptions.RemoveAllSwap([](const TObjectPtr<UOptimusResourceDescription>& Value) { return Value == nullptr; });
 	Variables->Descriptions.RemoveAllSwap([](const TObjectPtr<UOptimusVariableDescription>& Value) { return Value == nullptr; });
@@ -2152,14 +2177,14 @@ void UOptimusDeformer::PostLoad()
 		{
 			if (ResourceDescription->GetOuter() != Resources)
 			{
-				ResourceDescription->Rename(nullptr, Resources);
+				Optimus::RenameObject(ResourceDescription, nullptr, Resources);
 			}
 		}
 		for (UObject* VariableDescription: Variables->Descriptions)
 		{
 			if (VariableDescription->GetOuter() != Variables)
 			{
-				VariableDescription->Rename(nullptr, Variables);
+				Optimus::RenameObject(VariableDescription, nullptr, Variables);
 			}
 		}
 	}
@@ -2186,7 +2211,7 @@ void UOptimusDeformer::PostLoad()
 		{
 			if (Binding->bIsPrimaryBinding)
 			{
-				Binding->Rename(*PrimaryBindingName.ToString(), nullptr, REN_ForceNoResetLoaders | REN_NonTransactional);
+				Optimus::RenameObject(Binding, *PrimaryBindingName.ToString(), nullptr);
 				Binding->BindingName = PrimaryBindingName;
 			}
 		}
@@ -2208,6 +2233,11 @@ void UOptimusDeformer::PostLoad()
 	{
 		PostLoadFixupMismatchedResourceDataDomains();
 	}
+
+	if (GetLinkerCustomVersion(FOptimusObjectVersion::GUID) < FOptimusObjectVersion::KernelDataInterface)
+	{
+		PostLoadRemoveDeprecatedExecutionNodes();
+	}
 }
 
 
@@ -2312,6 +2342,103 @@ void UOptimusDeformer::PostLoadFixupMismatchedResourceDataDomains()
 #endif
 }
 
+void UOptimusDeformer::PostLoadRemoveDeprecatedExecutionNodes()
+{
+	for (UOptimusNodeGraph* Graph: GetGraphs())
+	{
+		// At the time of deprecation, subgraph is not supported
+		if (!ensure(Optimus::IsExecutionGraphType(Graph->GetGraphType())))
+		{
+			continue;
+		}
+	
+		TArray<UOptimusNode*> DeprecatedExecutionDataInterfaceNodes;
+
+		for (UOptimusNode* Node: Graph->GetAllNodes())
+		{
+			// PostLoad fixup for Kernel Nodes
+			if (UOptimusNode_CustomComputeKernel* KernelNode = Cast<UOptimusNode_CustomComputeKernel>(Node))
+			{
+				// Find the primary ComponentSourceNode for each kernel node
+				const UOptimusNodePin* PrimaryGroupPin = KernelNode->GetPrimaryGroupPin();
+
+				UOptimusNode_ComponentSource* ComponentSourceNode = nullptr;
+				FOptimusDataTypeHandle IntVector3Type = FOptimusDataTypeRegistry::Get().FindType(Optimus::GetTypeName(TBaseStructure<FIntVector3>::Get()));
+				
+				for (const UOptimusNodePin* Pin : PrimaryGroupPin->GetSubPins())
+				{
+					if (Pin->GetDirection() == EOptimusNodePinDirection::Input && Pin->GetDataType() == IntVector3Type)
+					{
+						TArray<FOptimusRoutedNodePin> ConnectedPins = Pin->GetConnectedPinsWithRouting();
+						if (ConnectedPins.Num() != 1)
+						{
+							// Skip if invalid/no connection
+							continue;
+						}
+
+						const UOptimusNode_DataInterface* DataInterfaceNode =
+							Cast<UOptimusNode_DataInterface>(ConnectedPins[0].NodePin->GetOwningNode());
+						const IOptimusDeprecatedExecutionDataInterface* ExecDataInterface =
+							Cast<IOptimusDeprecatedExecutionDataInterface>(DataInterfaceNode->GetDataInterface(GetTransientPackage()));
+
+						if (!ExecDataInterface)
+						{
+							// Skip if not connected to exec data interface
+							continue;
+						}
+
+						UOptimusNodePin* ComponentPin = DataInterfaceNode->GetComponentPin();
+
+						TArray<FOptimusRoutedNodePin> ConnectedComponentPins = ComponentPin->GetConnectedPinsWithRouting();
+
+						if (ConnectedComponentPins.Num() != 1)
+						{
+							// Skip if the exec data interface does not have a component source
+							continue;
+						}
+
+						ComponentSourceNode = Cast<UOptimusNode_ComponentSource>(ConnectedComponentPins[0].NodePin->GetOwningNode());
+						
+						if (ComponentSourceNode)
+						{
+							// Found a valid component source node, ready to link
+							break;
+						}
+					}
+				}
+
+				// Now that we have extract information from every pin, the deprecated ones have no more use and can be removed
+				KernelNode->PostLoadRemoveDeprecatedNumThreadsPin();
+
+				// After pin removal, we may have no input data pin to infer component source from,
+				// in which case we have to force a direct link between component source node and the kernel primary group pin
+				if (PrimaryGroupPin->GetSubPins().Num() > 0)
+				{
+					continue;
+				}
+
+				if (ComponentSourceNode)
+				{
+					Graph->AddLink(ComponentSourceNode->GetComponentPin(), KernelNode->GetPrimaryGroupPin_Internal());	
+				}
+			}
+
+			// PostLoad remove execution data interface nodes
+			if (const UOptimusNode_DataInterface* DataInterfaceNode = Cast<UOptimusNode_DataInterface>(Node))
+			{
+				if (Cast<IOptimusDeprecatedExecutionDataInterface>(DataInterfaceNode->GetDataInterface(GetTransientPackage())))
+				{
+					DeprecatedExecutionDataInterfaceNodes.Add(Node);
+				}
+			}
+		}
+		
+		Graph->RemoveNodes(DeprecatedExecutionDataInterfaceNodes);
+	}
+	
+	(void)MarkPackageDirty();
+}
+
 
 UOptimusComponentSourceBinding* UOptimusDeformer::FindCompatibleBindingWithInterface(
 	const UOptimusComputeDataInterface* InDataInterface) const
@@ -2556,7 +2683,7 @@ UOptimusNodeGraph* UOptimusDeformer::CreateGraph(
 	{
 		if (!AddGraph(Graph, InInsertBefore.GetValue()))
 		{
-			Graph->Rename(nullptr, GetTransientPackage());
+			Optimus::RemoveObject(Graph);
 			return nullptr;
 		}
 	}
@@ -2650,7 +2777,7 @@ bool UOptimusDeformer::RemoveGraph(
 	if (bInDeleteGraph)
 	{
 		// Un-parent this graph to a temporary storage and mark it for kill.
-		InGraph->Rename(nullptr, GetTransientPackage());
+		Optimus::RemoveObject(InGraph);
 	}
 
 	return true;
diff --git a/Engine/Plugins/Animation/DeformerGraph/Source/OptimusCore/Private/OptimusHelpers.cpp b/Engine/Plugins/Animation/DeformerGraph/Source/OptimusCore/Private/OptimusHelpers.cpp
--- a/Engine/Plugins/Animation/DeformerGraph/Source/OptimusCore/Private/OptimusHelpers.cpp
+++ b/Engine/Plugins/Animation/DeformerGraph/Source/OptimusCore/Private/OptimusHelpers.cpp
@@ -47,6 +47,12 @@ bool Optimus::RenameObject(UObject* InObjectToRename, const TCHAR* InNewName, UO
 	return InObjectToRename->Rename(InNewName, InNewOuter, REN_ForceNoResetLoaders | REN_DoNotDirty | REN_DontCreateRedirectors | REN_NonTransactional);
 }
 
+void Optimus::RemoveObject(UObject* InObjectToRemove)
+{
+	InObjectToRemove->Rename(nullptr, GetTransientPackage(), REN_ForceNoResetLoaders | REN_DoNotDirty | REN_DontCreateRedirectors | REN_NonTransactional);
+	InObjectToRemove->MarkAsGarbage();
+}
+
 TArray<UClass*> Optimus::GetClassObjectsInPackage(UPackage* InPackage)
 {
 	TArray<UObject*> Objects;
diff --git a/Engine/Plugins/Animation/DeformerGraph/Source/OptimusCore/Private/OptimusNode.cpp b/Engine/Plugins/Animation/DeformerGraph/Source/OptimusCore/Private/OptimusNode.cpp
--- a/Engine/Plugins/Animation/DeformerGraph/Source/OptimusCore/Private/OptimusNode.cpp
+++ b/Engine/Plugins/Animation/DeformerGraph/Source/OptimusCore/Private/OptimusNode.cpp
@@ -10,6 +10,7 @@
 #include "OptimusBindingTypes.h"
 #include "OptimusDeformer.h"
 #include "OptimusDiagnostic.h"
+#include "OptimusHelpers.h"
 #include "OptimusNodeGraph.h"
 #include "OptimusNodePin.h"
 #include "OptimusObjectVersion.h"
@@ -377,6 +378,11 @@ void UOptimusNode::PostLoad()
 
 	// Earlier iterations didn't set this flag. 
 	SetFlags(RF_Transactional);
+
+	for (UOptimusNodePin* Pin : Pins)
+	{
+		Pin->ConditionalPostLoad();
+	}
 }
 
 
@@ -660,9 +666,8 @@ bool UOptimusNode::RemovePinDirect(UOptimusNodePin* InPin)
 	for (UOptimusNodePin* Pin: PinsToRemove)
 	{
 		ExpandedPins.Remove(Pin->GetUniqueName());
-		
-		Pin->Rename(nullptr, GetTransientPackage());
-		Pin->MarkAsGarbage();
+
+		Optimus::RemoveObject(Pin);
 	}
 
 	CachedPinLookup.Reset();
@@ -741,6 +746,20 @@ bool UOptimusNode::MovePinDirect(
 	return true;
 }
 
+bool UOptimusNode::MovePinToGroupPinDirect(UOptimusNodePin* InPinToMove, UOptimusNodePin* InGroupPin)
+{
+	UOptimusNodePin* CurrentGroupPin = InPinToMove->GetParentPin();
+	TArray<UOptimusNodePin*>& PinGroup = CurrentGroupPin ? CurrentGroupPin->SubPins : Pins;
+	
+	PinGroup.Remove(InPinToMove);
+
+	Optimus::RenameObject(InPinToMove, nullptr, InGroupPin);
+	
+	InsertPinIntoHierarchy(InPinToMove, InGroupPin, nullptr);
+
+	return true;
+}
+
 
 bool UOptimusNode::SetPinDataType
 (
@@ -794,10 +813,10 @@ bool UOptimusNode::SetPinDataTypeDirect(
 		{
 			// Remove all sub-pins, if there were any.		
 			TGuardValue<bool> SuppressNotifications(bSendNotifications, false);
-				
+			
 			// If the type was already a sub-element type, remove the existing pins.
 			InPin->ClearSubPins();
-				
+			
 			// Add sub-pins, if the registered type is set to show them but only for value types.
 			if (EnumHasAllFlags(InDataType->TypeFlags, EOptimusDataTypeFlags::ShowElements))
 			{
@@ -807,7 +826,7 @@ bool UOptimusNode::SetPinDataTypeDirect(
 				}
 			}
 		}
-
+		
 		if (CanNotify())
 		{
 			InPin->Notify(EOptimusGraphNotifyType::PinTypeChanged);
diff --git a/Engine/Plugins/Animation/DeformerGraph/Source/OptimusCore/Private/OptimusNodeGraph.cpp b/Engine/Plugins/Animation/DeformerGraph/Source/OptimusCore/Private/OptimusNodeGraph.cpp
--- a/Engine/Plugins/Animation/DeformerGraph/Source/OptimusCore/Private/OptimusNodeGraph.cpp
+++ b/Engine/Plugins/Animation/DeformerGraph/Source/OptimusCore/Private/OptimusNodeGraph.cpp
@@ -153,6 +153,26 @@ void UOptimusNodeGraph::PostDuplicate(EDuplicateMode::Type DuplicateMode)
 	}
 }
 
+void UOptimusNodeGraph::PostLoad()
+{
+	Super::PostLoad();
+
+	for (UOptimusNodeGraph* SubGraph: SubGraphs)
+	{
+		SubGraph->ConditionalPostLoad();
+	}
+	
+	for (UOptimusNode* Node : Nodes)
+	{
+		Node->ConditionalPostLoad();
+	}
+
+	for (UOptimusNodeLink* Link : Links)
+	{
+		Link->ConditionalPostLoad();
+	}
+}
+
 UOptimusNodeGraph* UOptimusNodeGraph::GetParentGraph() const
 {
 	return Cast<UOptimusNodeGraph>(GetOuter());
@@ -1116,7 +1136,7 @@ UOptimusNode* UOptimusNodeGraph::CreateNodeDirect(
 		
 		if (!InConfigureNodeFunc(NewNode))
 		{
-			NewNode->Rename(nullptr, GetTransientPackage());
+			Optimus::RemoveObject(NewNode);
 			return nullptr;
 		}
 	}
@@ -1147,7 +1167,7 @@ bool UOptimusNodeGraph::AddNodeDirect(UOptimusNode* InNode)
 			return false;
 		}
 
-		InNode->Rename(nullptr, this);
+		Optimus::RenameObject(InNode, nullptr, this);
 	}
 
 	Nodes.Add(InNode);
@@ -1193,7 +1213,7 @@ bool UOptimusNodeGraph::RemoveNodeDirect(
 	Notify(EOptimusGraphNotifyType::NodeRemoved, InNode);
 
 	// Unparent this node to a temporary storage and mark it for kill.
-	InNode->Rename(nullptr, GetTransientPackage());
+	Optimus::RemoveObject(InNode);
 
 	return true;
 }
@@ -1398,15 +1418,14 @@ TSet<UOptimusComponentSourceBinding*> UOptimusNodeGraph::GetComponentSourceBindi
 		TArray<FOptimusRoutedNodePin> RoutedPins = GetConnectedPinsWithRouting(InNodePin, {});
 		if (RoutedPins.IsEmpty())
 		{
-			// this can happen if the pin is directly on a compute kernel and not connected to anything
-			StartNode = InNodePin->GetOwningNode();
-		}
-		else
-		{
-			check(RoutedPins.Num() == 1);
-			StartNode = RoutedPins[0].NodePin->GetOwningNode();
+			// If the input pin is directly on a compute kernel and not connected to anything
+			// we consider it to not have a component source, even if other pins in the same group
+			// has one. Pin group level gathering is caller's responsibility
+			return {};
 		}
 		
+		check(RoutedPins.Num() == 1);
+		StartNode = RoutedPins[0].NodePin->GetOwningNode();
 	}
 	else
 	{
@@ -1432,7 +1451,11 @@ TSet<UOptimusComponentSourceBinding*> UOptimusNodeGraph::GetComponentSourceBindi
 		
 		if (const IOptimusComputeKernelProvider* KernelProvider = Cast<const IOptimusComputeKernelProvider>(Node))
 		{
-			InputPins = KernelProvider->GetPrimaryGroupInputPins();
+			const UOptimusNodePin* PrimaryGroupPin = KernelProvider->GetPrimaryGroupPin();
+
+			// Group pin traversal is useful when the kernel does not have any input data pins
+			InputPins.Add(PrimaryGroupPin);
+			InputPins.Append(PrimaryGroupPin->GetSubPins());
 		}
 		else
 		{
@@ -1500,7 +1523,7 @@ void UOptimusNodeGraph::RemoveLinkByIndex(int32 LinkIndex)
 	Notify(EOptimusGraphNotifyType::LinkRemoved, Link);
 
 	// Unparent the link to a temporary storage and mark it for kill.
-	Link->Rename(nullptr, GetTransientPackage());
+	Optimus::RemoveObject(Link);
 }
 
 bool UOptimusNodeGraph::DoesLinkFormCycle(const UOptimusNode* InOutputNode, const UOptimusNode* InInputNode) const
@@ -1691,7 +1714,7 @@ UOptimusNodeGraph* UOptimusNodeGraph::CreateGraph(
 	{
 		if (!AddGraph(Graph, InInsertBefore.GetValue()))
 		{
-			Graph->Rename(nullptr, GetTransientPackage());
+			Optimus::RemoveObject(Graph);
 			return nullptr;
 		}
 	}
@@ -1752,7 +1775,7 @@ bool UOptimusNodeGraph::RemoveGraph(
 	if (bInDeleteGraph)
 	{
 		// Un-parent this graph to a temporary storage and mark it for kill.
-		InGraph->Rename(nullptr, GetTransientPackage());
+		Optimus::RemoveObject(InGraph);
 	}
 
 	return true;
diff --git a/Engine/Plugins/Animation/DeformerGraph/Source/OptimusCore/Private/OptimusNodePin.cpp b/Engine/Plugins/Animation/DeformerGraph/Source/OptimusCore/Private/OptimusNodePin.cpp
--- a/Engine/Plugins/Animation/DeformerGraph/Source/OptimusCore/Private/OptimusNodePin.cpp
+++ b/Engine/Plugins/Animation/DeformerGraph/Source/OptimusCore/Private/OptimusNodePin.cpp
@@ -4,6 +4,7 @@
 
 #include "Actions/OptimusNodeActions.h"
 #include "OptimusActionStack.h"
+#include "OptimusComponentSource.h"
 #include "OptimusDataTypeRegistry.h"
 #include "OptimusHelpers.h"
 #include "OptimusNode.h"
@@ -461,12 +462,26 @@ bool UOptimusNodePin::CanCannect(const UOptimusNodePin* InOtherPin, FString* Out
 		return false;
 	}
 
+	if (IsGroupingPin() || InOtherPin->IsGroupingPin())
+	{
+		// Allow direct connection from component source to group pin
+		FOptimusDataTypeHandle ComponentSourceType = FOptimusDataTypeRegistry::Get().FindType(*UOptimusComponentSourceBinding::StaticClass());
+
+		if (DataType == ComponentSourceType || InOtherPin->DataType == ComponentSourceType)
+		{
+			return true;
+		}
+
+		return false;
+	}
+	
 	// Check for incompatible types.
 	if (!((DataType->ShaderValueType.IsValid() && InOtherPin->DataType->ShaderValueType.IsValid() &&
 		  DataType->ShaderValueType == InOtherPin->DataType->ShaderValueType) ||
 		 DataType == InOtherPin->DataType))
 	{
 		// TBD: Automatic conversion.
+
 		if (OutReason)
 		{
 			*OutReason = TEXT("Incompatible pin types.");
@@ -525,6 +540,11 @@ void UOptimusNodePin::PostLoad()
 {
 	UObject::PostLoad();
 
+	for (UOptimusNodePin* SubPin : SubPins)
+	{
+		SubPin->ConditionalPostLoad();
+	}
+	
 	// If the storage was marked as a value, the domain should now be a singleton.
 	if (GetLinkerCustomVersion(FOptimusObjectVersion::GUID) < FOptimusObjectVersion::DataDomainExpansion)
 	{
@@ -579,7 +599,7 @@ void UOptimusNodePin::ClearSubPins()
 	for (UOptimusNodePin* Pin: SubPins)
 	{
 		// Consign them to oblivion.
-		Pin->Rename(nullptr, GetTransientPackage());
+		Optimus::RemoveObject(Pin);
 	}
 	SubPins.Reset();
 }
diff --git a/Engine/Plugins/Animation/DeformerGraph/Source/OptimusCore/Public/IOptimusComputeKernelProvider.h b/Engine/Plugins/Animation/DeformerGraph/Source/OptimusCore/Public/IOptimusComputeKernelProvider.h
--- a/Engine/Plugins/Animation/DeformerGraph/Source/OptimusCore/Public/IOptimusComputeKernelProvider.h
+++ b/Engine/Plugins/Animation/DeformerGraph/Source/OptimusCore/Public/IOptimusComputeKernelProvider.h
@@ -50,6 +50,8 @@ using FOptimus_NodeToDataInterfaceMap =  TMap<const UOptimusNode*, UOptimusCompu
 // that it represents.
 using FOptimus_PinToDataInterfaceMap = TMap<const UOptimusNodePin*, UOptimusComputeDataInterface*>;
 
+// Maps from a kernel node to the kernel data interface 
+using FOptimus_KernelNodeToKernelDataInterfaceMap = TMap<const UOptimusNode*, UComputeDataInterface*>;
 
 using FOptimus_ComputeKernelResult = TVariant<UOptimusKernelSource* /* Kernel */, FText /* Error */>;
 
@@ -78,6 +80,7 @@ public:
 	 * @param InValueNodes
 	 * @param InGraphDataInterface
 	 * @param InGraphDataComponentBinding
+	 * @param InKernelDataInterface
 	 * @param OutInputDataBindings
 	 * @param OutOutputDataBindings
 	 * @param OutExecutionDataInterface
@@ -90,6 +93,7 @@ public:
 		const TArray<const UOptimusNode*>& InValueNodes,
 		const UComputeDataInterface* InGraphDataInterface,
 		const UOptimusComponentSourceBinding* InGraphDataComponentBinding,
+		const UComputeDataInterface* InKernelDataInterface,
 		FOptimus_InterfaceBindingMap& OutInputDataBindings,
 		FOptimus_InterfaceBindingMap& OutOutputDataBindings
 		) const = 0;
@@ -97,8 +101,9 @@ public:
 	/** Returns the execution domain that this kernel should iterate over */
 	virtual FName GetExecutionDomain() const = 0;
 
-	/** Returns the input pins of the primary group only. Used for traversing up the graph to identify the component
-	 *  source.
-	 */
-	virtual TArray<const UOptimusNodePin*> GetPrimaryGroupInputPins() const = 0; 
+	/** Used for traversing up the graph to identify the component source of the compute kernel */
+	virtual const UOptimusNodePin* GetPrimaryGroupPin() const = 0;
+	
+	/** Each kernel may have its own data interface responsible for passing kernel related data to GPU */
+	virtual UComputeDataInterface* GetKernelDataInterface(UObject* InOuter) const = 0;
 };
diff --git a/Engine/Plugins/Animation/DeformerGraph/Source/OptimusCore/Public/IOptimusDeprecatedExecutionDataInterface.h b/Engine/Plugins/Animation/DeformerGraph/Source/OptimusCore/Public/IOptimusDeprecatedExecutionDataInterface.h
new file mode 100644
--- /dev/null
+++ b/Engine/Plugins/Animation/DeformerGraph/Source/OptimusCore/Public/IOptimusDeprecatedExecutionDataInterface.h
@@ -0,0 +1,28 @@
+﻿// Copyright Epic Games, Inc. All Rights Reserved.
+
+#pragma once
+
+#include "UObject/Interface.h"
+
+#include "IOptimusDeprecatedExecutionDataInterface.generated.h"
+
+
+
+UINTERFACE()
+class OPTIMUSCORE_API UOptimusDeprecatedExecutionDataInterface:
+	public UInterface
+{
+	GENERATED_BODY()
+};
+
+
+/** An interface that should be implemented by any legacy Data Interface Node whose data interface is an execution data interface.
+ *  such that the serialized execution domain can be extracted and assigned to the kernel data interface
+ */
+class IOptimusDeprecatedExecutionDataInterface
+{
+	GENERATED_BODY()
+
+public:
+	virtual FName GetSelectedExecutionDomainName() const = 0;
+};
diff --git a/Engine/Plugins/Animation/DeformerGraph/Source/OptimusCore/Public/IOptimusNodeAdderPinProvider.h b/Engine/Plugins/Animation/DeformerGraph/Source/OptimusCore/Public/IOptimusNodeAdderPinProvider.h
--- a/Engine/Plugins/Animation/DeformerGraph/Source/OptimusCore/Public/IOptimusNodeAdderPinProvider.h
+++ b/Engine/Plugins/Animation/DeformerGraph/Source/OptimusCore/Public/IOptimusNodeAdderPinProvider.h
@@ -43,6 +43,4 @@ public:
 		) = 0;
 	
 	virtual bool RemoveAddedPins(TConstArrayView<UOptimusNodePin*> InAddedPinsToRemove) = 0;
-
-	virtual FName GetSanitizedNewPinName(UOptimusNodePin* InTargetParentPin, FName InPinName)= 0;
 };
diff --git a/Engine/Plugins/Animation/DeformerGraph/Source/OptimusCore/Public/OptimusComponentSource.h b/Engine/Plugins/Animation/DeformerGraph/Source/OptimusCore/Public/OptimusComponentSource.h
--- a/Engine/Plugins/Animation/DeformerGraph/Source/OptimusCore/Public/OptimusComponentSource.h
+++ b/Engine/Plugins/Animation/DeformerGraph/Source/OptimusCore/Public/OptimusComponentSource.h
@@ -43,6 +43,13 @@ public:
 		return 0;
 	}
 
+	/** Returns the total number of thread invocations a component
+	  * requires in case of non-unified dispatch. */
+	virtual uint32 GetDefaultNumInvocations(const UActorComponent* InComponent, int32 InLod) const
+	{
+		return 0;
+	}
+	
 	/** For a given execution domain, and component LOD index, return the range of the domain as given by 
 	 *  the component that is associated with this source.
 	 */
diff --git a/Engine/Plugins/Animation/DeformerGraph/Source/OptimusCore/Public/OptimusDeformer.h b/Engine/Plugins/Animation/DeformerGraph/Source/OptimusCore/Public/OptimusDeformer.h
--- a/Engine/Plugins/Animation/DeformerGraph/Source/OptimusCore/Public/OptimusDeformer.h
+++ b/Engine/Plugins/Animation/DeformerGraph/Source/OptimusCore/Public/OptimusDeformer.h
@@ -444,6 +444,7 @@ protected:
 private:
 	void PostLoadFixupMissingComponentBindingsCompat();
 	void PostLoadFixupMismatchedResourceDataDomains();
+	void PostLoadRemoveDeprecatedExecutionNodes();
 
 	/** Find a compatible binding with the given data interface. Returns nullptr if no such binding exists */
 	UOptimusComponentSourceBinding* FindCompatibleBindingWithInterface(
diff --git a/Engine/Plugins/Animation/DeformerGraph/Source/OptimusCore/Public/OptimusHelpers.h b/Engine/Plugins/Animation/DeformerGraph/Source/OptimusCore/Public/OptimusHelpers.h
--- a/Engine/Plugins/Animation/DeformerGraph/Source/OptimusCore/Public/OptimusHelpers.h
+++ b/Engine/Plugins/Animation/DeformerGraph/Source/OptimusCore/Public/OptimusHelpers.h
@@ -79,6 +79,9 @@ namespace Optimus
 
 	bool RenameObject(UObject* InObjectToRename, const TCHAR* InNewName, UObject* InNewOuter);
 
+	/** Use this function to remove objects during postload safely */
+	void RemoveObject(UObject* InObjectToRemove);
+
 	/** Our generated classes are parented to the package, this is a utility function
 		to collect them */
 	TArray<UClass*> GetClassObjectsInPackage(UPackage* InPackage);
diff --git a/Engine/Plugins/Animation/DeformerGraph/Source/OptimusCore/Public/OptimusNode.h b/Engine/Plugins/Animation/DeformerGraph/Source/OptimusCore/Public/OptimusNode.h
--- a/Engine/Plugins/Animation/DeformerGraph/Source/OptimusCore/Public/OptimusNode.h
+++ b/Engine/Plugins/Animation/DeformerGraph/Source/OptimusCore/Public/OptimusNode.h
@@ -290,6 +290,11 @@ protected:
 		UOptimusNodePin* InPinToMove,
 		const UOptimusNodePin* InPinBefore
 		);
+
+	bool MovePinToGroupPinDirect(
+		UOptimusNodePin* InPinToMove,
+		UOptimusNodePin* InGroupPin
+		);
 	
 	/** Set the pin data type. */
 	bool SetPinDataType(
diff --git a/Engine/Plugins/Animation/DeformerGraph/Source/OptimusCore/Public/OptimusNodeGraph.h b/Engine/Plugins/Animation/DeformerGraph/Source/OptimusCore/Public/OptimusNodeGraph.h
--- a/Engine/Plugins/Animation/DeformerGraph/Source/OptimusCore/Public/OptimusNodeGraph.h
+++ b/Engine/Plugins/Animation/DeformerGraph/Source/OptimusCore/Public/OptimusNodeGraph.h
@@ -71,6 +71,8 @@ public:
 	// in the source package/asset are not duplicated automatically to the new package/asset
 	void PostDuplicate(EDuplicateMode::Type DuplicateMode) override;
 
+	void PostLoad() override;
+
 	UOptimusNodeGraph *GetParentGraph() const;
 	
 	FString GetGraphPath() const;
diff --git a/Engine/Plugins/Animation/DeformerGraph/Source/OptimusCore/Public/OptimusNodePin.h b/Engine/Plugins/Animation/DeformerGraph/Source/OptimusCore/Public/OptimusNodePin.h
--- a/Engine/Plugins/Animation/DeformerGraph/Source/OptimusCore/Public/OptimusNodePin.h
+++ b/Engine/Plugins/Animation/DeformerGraph/Source/OptimusCore/Public/OptimusNodePin.h
@@ -44,8 +44,7 @@ public:
 	/// @return The direction of this pin.
 	EOptimusNodePinDirection GetDirection() const { return Direction; }
 
-	/// Returns \c true if this pin is a grouping pin. A grouping pin accepts no connections
-	/// but just acts as a grouping mechanism for other sub-pins.
+	/// Returns \c true if this pin is a grouping pin. 
 	bool IsGroupingPin() const { return bIsGroupingPin; }
 
 	/// Returns the parent pin of this pin, or nullptr if it is the top-most pin.
diff --git a/Engine/Plugins/Animation/DeformerGraph/Source/OptimusCore/Public/OptimusObjectVersion.h b/Engine/Plugins/Animation/DeformerGraph/Source/OptimusCore/Public/OptimusObjectVersion.h
--- a/Engine/Plugins/Animation/DeformerGraph/Source/OptimusCore/Public/OptimusObjectVersion.h
+++ b/Engine/Plugins/Animation/DeformerGraph/Source/OptimusCore/Public/OptimusObjectVersion.h
@@ -19,6 +19,7 @@ struct OPTIMUSCORE_API FOptimusObjectVersion
 		ComponentProviderSupport,
 		SetPrimaryBindingName,
 		DataDomainExpansion,
+		KernelDataInterface,
 
 		// -----<new versions can be added above this line>-------------------------------------------------
 		VersionPlusOne,
diff --git a/Engine/Plugins/Animation/DeformerGraph/Source/OptimusEditor/Private/OptimusDetailsCustomization.cpp b/Engine/Plugins/Animation/DeformerGraph/Source/OptimusEditor/Private/OptimusDetailsCustomization.cpp
--- a/Engine/Plugins/Animation/DeformerGraph/Source/OptimusEditor/Private/OptimusDetailsCustomization.cpp
+++ b/Engine/Plugins/Animation/DeformerGraph/Source/OptimusEditor/Private/OptimusDetailsCustomization.cpp
@@ -32,7 +32,7 @@
 #include "Styling/AppStyle.h"
 #include "Styling/SlateIconFinder.h"
 #include "Widgets/Input/SButton.h"
-#include "Widgets/Input/SComboBox.h"
+
 #include "Widgets/Input/SEditableTextBox.h"
 #include "Widgets/Input/SMultiLineEditableTextBox.h"
 #include "Widgets/Input/STextComboBox.h"
@@ -248,32 +248,8 @@ void FOptimusExecutionDomainCustomization::CustomizeHeader(
 
 	TArray<UObject*> OwningObjects;
 	InPropertyHandle->GetOuterObjects(OwningObjects);
-
-	TSet<FName> CollectedContextNames;
-	for (int32 Index = 0; Index < OwningObjects.Num(); Index++)
-	{
-		TSet<FName> ProviderContextNames;
-		if (const IOptimusExecutionDomainProvider* ExecutionDomainProvider = Cast<IOptimusExecutionDomainProvider>(OwningObjects[Index]))
-		{
-			if (Index == 0)
-			{
-				CollectedContextNames.Append(ExecutionDomainProvider->GetExecutionDomains());
-			}
-			else
-			{
-				CollectedContextNames = CollectedContextNames.Intersect(TSet<FName>{ExecutionDomainProvider->GetExecutionDomains()});
-			}
-		}
-	}
-	ContextNames = CollectedContextNames.Array();
-	if (ContextNames.IsEmpty())
-	{
-		ContextNames.Add(NAME_None);
-	}
-	else
-	{
-		ContextNames.Sort(FNameLexicalLess{});
-	}
+	
+	WeakOwningObjects = TArray<TWeakObjectPtr<UObject>>(OwningObjects);
 	
 	InHeaderRow.NameContent()
 	[
@@ -281,8 +257,7 @@ void FOptimusExecutionDomainCustomization::CustomizeHeader(
 	]
 	.ValueContent()
 	[
-		SNew(SComboBox<FName>)
-			.ToolTipText(LOCTEXT("ExecContextListerToolTip", "Select an execution context from the list of available contexts."))
+		SAssignNew(ComboBox, SComboBox<FName>)
 			.OptionsSource(&ContextNames)
 			.IsEnabled_Lambda([InPropertyHandle]() -> bool
 			{
@@ -299,6 +274,7 @@ void FOptimusExecutionDomainCustomization::CustomizeHeader(
 			{
 				ContextNameProperty->SetValue(InName);
 			})
+			.OnComboBoxOpening(this, &FOptimusExecutionDomainCustomization::UpdateContextNames)
 			[
 				SNew(STextBlock)
 				.Font(IPropertyTypeCustomizationUtils::GetRegularFont())
@@ -312,6 +288,38 @@ void FOptimusExecutionDomainCustomization::CustomizeHeader(
 	];
 }
 
+void FOptimusExecutionDomainCustomization::UpdateContextNames()
+{
+	TSet<FName> CollectedContextNames;
+	for (int32 Index = 0; Index < WeakOwningObjects.Num(); Index++)
+	{
+		TSet<FName> ProviderContextNames;
+		if (const IOptimusExecutionDomainProvider* ExecutionDomainProvider = Cast<IOptimusExecutionDomainProvider>(WeakOwningObjects[Index]))
+		{
+			if (Index == 0)
+			{
+				CollectedContextNames.Append(ExecutionDomainProvider->GetExecutionDomains());
+			}
+			else
+			{
+				CollectedContextNames = CollectedContextNames.Intersect(TSet<FName>{ExecutionDomainProvider->GetExecutionDomains()});
+			}
+		}
+	}
+	
+	ContextNames = CollectedContextNames.Array();
+	if (ContextNames.IsEmpty())
+	{
+		ContextNames.Add(NAME_None);
+	}
+	else
+	{
+		ContextNames.Sort(FNameLexicalLess{});
+	}
+	
+	ComboBox->RefreshOptions();
+}
+
 
 TSharedRef<IPropertyTypeCustomization> FOptimusDataDomainCustomization::MakeInstance()
 {
diff --git a/Engine/Plugins/Animation/DeformerGraph/Source/OptimusEditor/Private/OptimusDetailsCustomization.h b/Engine/Plugins/Animation/DeformerGraph/Source/OptimusEditor/Private/OptimusDetailsCustomization.h
--- a/Engine/Plugins/Animation/DeformerGraph/Source/OptimusEditor/Private/OptimusDetailsCustomization.h
+++ b/Engine/Plugins/Animation/DeformerGraph/Source/OptimusEditor/Private/OptimusDetailsCustomization.h
@@ -7,7 +7,9 @@
 #include "OptimusDataType.h"
 #include "PropertyCustomizationHelpers.h"
 #include "PropertyHandle.h"
+#include "Widgets/Input/SComboBox.h"
 
+class IOptimusExecutionDomainProvider;
 class FOptimusHLSLSyntaxHighlighter;
 class SEditableTextBox;
 class SExpandableArea;
@@ -72,7 +74,14 @@ public:
 	{ }
 
 private:
+
+	void UpdateContextNames();
+
+	TSharedPtr<SComboBox<FName>> ComboBox;
+	
 	TArray<FName> ContextNames;
+
+	TArray<TWeakObjectPtr<UObject>> WeakOwningObjects;
 };
 
 
diff --git a/Engine/Plugins/Animation/DeformerGraph/Source/OptimusEditor/Private/OptimusEditorGraphNode.cpp b/Engine/Plugins/Animation/DeformerGraph/Source/OptimusEditor/Private/OptimusEditorGraphNode.cpp
--- a/Engine/Plugins/Animation/DeformerGraph/Source/OptimusEditor/Private/OptimusEditorGraphNode.cpp
+++ b/Engine/Plugins/Animation/DeformerGraph/Source/OptimusEditor/Private/OptimusEditorGraphNode.cpp
@@ -207,7 +207,6 @@ void UOptimusEditorGraphNode::SynchronizeGraphPinTypeWithModelPin(
 		if (InModelPin->IsGroupingPin())
 		{
 			GraphPin->PinType.PinCategory = GroupTypeName;
-			GraphPin->bNotConnectable = true;
 		}
 		else
 		{
@@ -339,7 +338,6 @@ bool UOptimusEditorGraphNode::CreateGraphPinFromModelPin(
 	UEdGraphPin *GraphPin = CreatePin(InDirection, PinType, PinPath);
 
 	GraphPin->PinFriendlyName = InModelPin->GetDisplayName();
-	GraphPin->bNotConnectable = InModelPin->IsGroupingPin();
 
 	if (InParentPin)
 	{
diff --git a/Engine/Plugins/Animation/DeformerGraph/Source/OptimusEditor/Private/OptimusEditorGraphSchema.cpp b/Engine/Plugins/Animation/DeformerGraph/Source/OptimusEditor/Private/OptimusEditorGraphSchema.cpp
--- a/Engine/Plugins/Animation/DeformerGraph/Source/OptimusEditor/Private/OptimusEditorGraphSchema.cpp
+++ b/Engine/Plugins/Animation/DeformerGraph/Source/OptimusEditor/Private/OptimusEditorGraphSchema.cpp
@@ -16,6 +16,7 @@
 #include "EdGraphSchema_K2.h"
 #include "Editor.h"
 #include "OptimusActionStack.h"
+#include "OptimusComponentSource.h"
 #include "OptimusComputeDataInterface.h"
 #include "OptimusEditorGraphConnectionDrawingPolicy.h"
 #include "Styling/SlateIconFinder.h"
@@ -432,6 +433,17 @@ FLinearColor UOptimusEditorGraphSchema::GetColorFromPinType(const FEdGraphPinTyp
 		return Settings->WildcardPinTypeColor;
 	}
 	
+	if (InPinType.PinCategory == UOptimusEditorGraphNode::GroupTypeName)
+	{
+		// We allow direct connection from component source to group pin, thus use the same color for them
+		FOptimusDataTypeHandle ComponentSourceType = FOptimusDataTypeRegistry::Get().FindType(*UOptimusComponentSourceBinding::StaticClass());
+
+		if (ensure(ComponentSourceType->bHasCustomPinColor))
+		{
+			return ComponentSourceType->CustomPinColor;
+		}
+	}
+	
 	// Use the PinSubCategory value to resolve the type. It's set in
 	// UOptimusEditorGraphSchema::GetPinTypeFromDataType.
 	FOptimusDataTypeHandle DataType = FOptimusDataTypeRegistry::Get().FindType(InPinType.PinSubCategory);
diff --git a/Engine/Plugins/Animation/DeformerGraph/Source/OptimusEditor/Private/Widgets/SOptimusEditorGraphNode.cpp b/Engine/Plugins/Animation/DeformerGraph/Source/OptimusEditor/Private/Widgets/SOptimusEditorGraphNode.cpp
--- a/Engine/Plugins/Animation/DeformerGraph/Source/OptimusEditor/Private/Widgets/SOptimusEditorGraphNode.cpp
+++ b/Engine/Plugins/Animation/DeformerGraph/Source/OptimusEditor/Private/Widgets/SOptimusEditorGraphNode.cpp
@@ -380,7 +380,7 @@ public:
 			ExpanderWidget =
 				SNew(SOptimusEditorExpanderArrow, InOptionalOwnerRow)
 					.LeftAligned(bLeftAligned)
-					.AlwaysVisible(bIsGroupPin);
+					.AlwaysVisible(false);
 		}
 		else
 		{
@@ -557,6 +557,18 @@ void SOptimusEditorGraphNode::Construct(const FArguments& InArgs)
 	        .ItemHeight(20.0f)
 		];
 
+	// Padding for group pin
+	if (EditorGraphNode->GetTopLevelInputPins().IsValidIndex(0) &&
+		EditorGraphNode->GetTopLevelInputPins()[0]->IsGroupingPin())
+	{
+		RightNodeBox->AddSlot()
+			.MaxHeight(26.0f)
+			[
+				SNew(SSpacer)
+				.Size(FVector2d(1.0f,26.0f))
+			];
+	}
+	
 	RightNodeBox->AddSlot()
 	    .AutoHeight()
 		[
@@ -590,7 +602,7 @@ void SOptimusEditorGraphNode::Construct(const FArguments& InArgs)
 			EHorizontalAlignment Alignment = Direction == EGPD_Input ? HAlign_Left : HAlign_Right;
 			
 			NodeBox->AddSlot()
-				.MaxHeight(22.0f)
+				.MaxHeight(26.0f)
 				.HAlign(Alignment)
 				[
 					SNew(SHorizontalBox)
@@ -754,7 +766,7 @@ void SOptimusEditorGraphNode::UpdatePinIcon(
 	{
 		if (ModelPin->IsGroupingPin())
 		{
-			InPinToUpdate->SetCustomPinIcon(CachedImg_Pin_Grouping, CachedImg_Pin_Grouping);
+			InPinToUpdate->SetCustomPinIcon(CachedImg_Pin_Value_Connected, CachedImg_Pin_Value_Disconnected);
 		}
 		else if (ModelPin->GetDataDomain().IsSingleton())
 		{
@@ -979,7 +991,9 @@ TSharedRef<ITableRow> SOptimusEditorGraphNode::MakeTableRowWidget(
 	}
 	else
 	{
-		RowWidget->SetToolTipText(LOCTEXT("GroupPinToolTip", "Click to open or close the pin group"));
+		RowWidget->SetToolTipText(
+			LOCTEXT("GroupPinToolTip",
+				"Group Pin accepts direct connection from Component Source Node"));
 	}
 
 	RowWidget->SetContent(
diff --git a/Engine/Plugins/Animation/MLDeformer/MLDeformerFramework/Source/MLDeformerFramework/Private/MLDeformerComponentSource.cpp b/Engine/Plugins/Animation/MLDeformer/MLDeformerFramework/Source/MLDeformerFramework/Private/MLDeformerComponentSource.cpp
--- a/Engine/Plugins/Animation/MLDeformer/MLDeformerFramework/Source/MLDeformerFramework/Private/MLDeformerComponentSource.cpp
+++ b/Engine/Plugins/Animation/MLDeformer/MLDeformerFramework/Source/MLDeformerFramework/Private/MLDeformerComponentSource.cpp
@@ -46,6 +46,30 @@ int32 UMLDeformerComponentSource::GetLodIndex(const UActorComponent* InComponent
 }
 
 
+uint32 UMLDeformerComponentSource::GetDefaultNumInvocations(
+	const UActorComponent* InComponent,
+    int32 InLod
+    ) const
+{
+	const USkinnedMeshComponent* SkinnedMeshComponent = Cast<USkinnedMeshComponent>(InComponent);
+	if (!SkinnedMeshComponent)
+	{
+		return 0;
+	}
+
+	const FSkeletalMeshObject* SkeletalMeshObject = SkinnedMeshComponent->MeshObject;
+	if (!SkeletalMeshObject)
+	{
+		return 0;
+	}
+
+	FSkeletalMeshRenderData const& SkeletalMeshRenderData = SkeletalMeshObject->GetSkeletalMeshRenderData();
+	FSkeletalMeshLODRenderData const* LodRenderData = &SkeletalMeshRenderData.LODRenderData[InLod];
+
+	return LodRenderData->RenderSections.Num();
+}
+
+
 bool UMLDeformerComponentSource::GetComponentElementCountsForExecutionDomain(
 	FName InDomainName,
 	const UActorComponent* InComponent,
diff --git a/Engine/Plugins/Animation/MLDeformer/MLDeformerFramework/Source/MLDeformerFramework/Private/MLDeformerComponentSource.h b/Engine/Plugins/Animation/MLDeformer/MLDeformerFramework/Source/MLDeformerFramework/Private/MLDeformerComponentSource.h
--- a/Engine/Plugins/Animation/MLDeformer/MLDeformerFramework/Source/MLDeformerFramework/Private/MLDeformerComponentSource.h
+++ b/Engine/Plugins/Animation/MLDeformer/MLDeformerFramework/Source/MLDeformerFramework/Private/MLDeformerComponentSource.h
@@ -24,5 +24,6 @@ public:
 	TSubclassOf<UActorComponent> GetComponentClass() const override;
 	TArray<FName> GetExecutionDomains() const override;
 	int32 GetLodIndex(const UActorComponent* InComponent) const override;
+	uint32 GetDefaultNumInvocations(const UActorComponent* InComponent, int32 InLod) const override;
 	bool GetComponentElementCountsForExecutionDomain(FName InDomainName, const UActorComponent* InComponent, int32 InLodIndex, TArray<int32>& OutElementCounts) const override;
 };
diff --git a/Engine/Plugins/Runtime/HairStrands/Source/HairStrandsCore/Private/Deformer/DeformerDataInterfaceGroomExec.cpp b/Engine/Plugins/Runtime/HairStrands/Source/HairStrandsCore/Private/Deformer/DeformerDataInterfaceGroomExec.cpp
--- a/Engine/Plugins/Runtime/HairStrands/Source/HairStrandsCore/Private/Deformer/DeformerDataInterfaceGroomExec.cpp
+++ b/Engine/Plugins/Runtime/HairStrands/Source/HairStrandsCore/Private/Deformer/DeformerDataInterfaceGroomExec.cpp
@@ -2,6 +2,7 @@
 
 #include "DeformerDataInterfaceGroomExec.h"
 
+#include "DeformerGroomComponentSource.h"
 #include "GroomComponent.h"
 #include "ComputeFramework/ShaderParamTypeDefinition.h"
 #include "ShaderCompilerCore.h"
@@ -81,6 +82,21 @@ UComputeDataProvider* UOptimusGroomExecDataInterface::CreateDataProvider(TObject
 	return Provider;
 }
 
+FName UOptimusGroomExecDataInterface::GetSelectedExecutionDomainName() const
+{
+	if (Domain == EOptimusGroomExecDomain::ControlPoint)
+	{
+		return UOptimusGroomComponentSource::Domains::ControlPoint;
+	}
+
+	if (Domain == EOptimusGroomExecDomain::Curve)
+	{
+		return UOptimusGroomComponentSource::Domains::Curve;
+	}
+
+	return NAME_None;
+}
+
 
 FComputeDataProviderRenderProxy* UOptimusGroomExecDataProvider::GetRenderProxy()
 {
diff --git a/Engine/Plugins/Runtime/HairStrands/Source/HairStrandsCore/Private/Deformer/DeformerDataInterfaceGroomExec.h b/Engine/Plugins/Runtime/HairStrands/Source/HairStrandsCore/Private/Deformer/DeformerDataInterfaceGroomExec.h
--- a/Engine/Plugins/Runtime/HairStrands/Source/HairStrandsCore/Private/Deformer/DeformerDataInterfaceGroomExec.h
+++ b/Engine/Plugins/Runtime/HairStrands/Source/HairStrandsCore/Private/Deformer/DeformerDataInterfaceGroomExec.h
@@ -2,6 +2,7 @@
 
 #pragma once
 
+#include "IOptimusDeprecatedExecutionDataInterface.h"
 #include "OptimusComputeDataInterface.h"
 #include "ComputeFramework/ComputeDataProvider.h"
 #include "DeformerDataInterfaceGroomExec.generated.h"
@@ -21,7 +22,9 @@ enum class EOptimusGroomExecDomain : uint8
 
 /** Compute Framework Data Interface for executing kernels over a skinned mesh domain. */
 UCLASS(Category = ComputeFramework)
-class HAIRSTRANDSCORE_API UOptimusGroomExecDataInterface : public UOptimusComputeDataInterface
+class HAIRSTRANDSCORE_API UOptimusGroomExecDataInterface :
+	public UOptimusComputeDataInterface,
+	public IOptimusDeprecatedExecutionDataInterface
 {
 	GENERATED_BODY()
 
@@ -29,6 +32,7 @@ public:
 	//~ Begin UOptimusComputeDataInterface Interface
 	FString GetDisplayName() const override;
 	FName GetCategory() const override;
+	bool IsVisible() const override {return false;};
 	TArray<FOptimusCDIPinDefinition> GetPinDefinitions() const override;
 	TSubclassOf<UActorComponent> GetRequiredComponentClass() const override;
 	//~ End UOptimusComputeDataInterface Interface
@@ -44,6 +48,10 @@ public:
 	UComputeDataProvider* CreateDataProvider(TObjectPtr<UObject> InBinding, uint64 InInputMask, uint64 InOutputMask) const override;
 	//~ End UComputeDataInterface Interface
 
+	//~ Begin IOptimusDeprecatedExecutionDataInterface Interface
+	FName GetSelectedExecutionDomainName() const override;
+	//~ End IOptimusDeprecatedExecutionDataInterface Interface
+	
 	UPROPERTY(EditAnywhere, Category = Execution)
 	EOptimusGroomExecDomain Domain = EOptimusGroomExecDomain::ControlPoint;
 
diff --git a/Engine/Plugins/Runtime/HairStrands/Source/HairStrandsCore/Private/Deformer/DeformerGroomComponentSource.cpp b/Engine/Plugins/Runtime/HairStrands/Source/HairStrandsCore/Private/Deformer/DeformerGroomComponentSource.cpp
--- a/Engine/Plugins/Runtime/HairStrands/Source/HairStrandsCore/Private/Deformer/DeformerGroomComponentSource.cpp
+++ b/Engine/Plugins/Runtime/HairStrands/Source/HairStrandsCore/Private/Deformer/DeformerGroomComponentSource.cpp
@@ -34,6 +34,22 @@ int32 UOptimusGroomComponentSource::GetLodIndex(const UActorComponent* InCompone
 	return 0;
 }
 
+uint32 UOptimusGroomComponentSource::GetDefaultNumInvocations(
+	const UActorComponent* InComponent,
+	int32 InLod
+	) const
+{
+	const UGroomComponent* GroomComponent = Cast<UGroomComponent>(InComponent);
+	if (!GroomComponent)
+	{
+		return 0;
+	}
+
+	const uint32 GroupCount = GroomComponent->GetGroupCount();
+
+	return GroupCount;
+}
+
 bool UOptimusGroomComponentSource::GetComponentElementCountsForExecutionDomain(
 	FName InDomainName,
 	const UActorComponent* InComponent,
diff --git a/Engine/Plugins/Runtime/HairStrands/Source/HairStrandsCore/Private/Deformer/DeformerGroomComponentSource.h b/Engine/Plugins/Runtime/HairStrands/Source/HairStrandsCore/Private/Deformer/DeformerGroomComponentSource.h
--- a/Engine/Plugins/Runtime/HairStrands/Source/HairStrandsCore/Private/Deformer/DeformerGroomComponentSource.h
+++ b/Engine/Plugins/Runtime/HairStrands/Source/HairStrandsCore/Private/Deformer/DeformerGroomComponentSource.h
@@ -23,6 +23,7 @@ public:
 	TSubclassOf<UActorComponent> GetComponentClass() const override;
 	TArray<FName> GetExecutionDomains() const override;
 	int32 GetLodIndex(const UActorComponent* InComponent) const override;
+	uint32 GetDefaultNumInvocations(const UActorComponent* InComponent, int32 InLod) const;
 	bool GetComponentElementCountsForExecutionDomain(FName InDomainName, const UActorComponent* InComponent, int32 InLodIndex, TArray<int32>& OutInvocationElementCounts) const override;
 	bool IsUsableAsPrimarySource() const override;
 };
