Update more callsites to prepare for upcoming TArray/TArrayView<TObjectPtr<...>> restrictions

diff --git a/Engine/Plugins/Animation/DeformerGraph/Source/OptimusCore/Private/OptimusNode.cpp b/Engine/Plugins/Animation/DeformerGraph/Source/OptimusCore/Private/OptimusNode.cpp
--- a/Engine/Plugins/Animation/DeformerGraph/Source/OptimusCore/Private/OptimusNode.cpp
+++ b/Engine/Plugins/Animation/DeformerGraph/Source/OptimusCore/Private/OptimusNode.cpp
@@ -749,7 +749,7 @@ bool UOptimusNode::MovePinDirect(
 bool UOptimusNode::MovePinToGroupPinDirect(UOptimusNodePin* InPinToMove, UOptimusNodePin* InGroupPin)
 {
 	UOptimusNodePin* CurrentGroupPin = InPinToMove->GetParentPin();
-	TArray<UOptimusNodePin*>& PinGroup = CurrentGroupPin ? CurrentGroupPin->SubPins : Pins;
+	TArray<TObjectPtr<UOptimusNodePin>>& PinGroup = CurrentGroupPin ? CurrentGroupPin->SubPins : Pins;
 	
 	PinGroup.Remove(InPinToMove);
 
