Optimus: Avoid loading all referencer packages when a user defined struct is changed. Only load deformer assets instead

diff --git a/Engine/Plugins/Animation/DeformerGraph/Source/OptimusEditor/Private/OptimusEditorModule.cpp b/Engine/Plugins/Animation/DeformerGraph/Source/OptimusEditor/Private/OptimusEditorModule.cpp
--- a/Engine/Plugins/Animation/DeformerGraph/Source/OptimusEditor/Private/OptimusEditorModule.cpp
+++ b/Engine/Plugins/Animation/DeformerGraph/Source/OptimusEditor/Private/OptimusEditorModule.cpp
@@ -125,21 +125,26 @@ void FOptimusEditorModule::PreChange(const UUserDefinedStruct* Changed,
 		[](FStructProperty* InStructProperty){ return Cast<UOptimusValueContainerGeneratorClass>(InStructProperty->GetOwnerClass()) != nullptr; },
 		[](UStruct*){});
 	
-	FAssetRegistryModule& AssetRegistryModule = FModuleManager::LoadModuleChecked<FAssetRegistryModule>("AssetRegistry");
-	TArray<FName> Referencers;
-	AssetRegistryModule.Get().GetReferencers(Changed->GetPackage()->GetFName(), Referencers);
+	FAssetRegistryModule& AssetRegistryModule = FModuleManager::LoadModuleChecked<FAssetRegistryModule>(FName("AssetRegistry"));
+	IAssetRegistry& AssetRegistry = AssetRegistryModule.Get();
 
-	TArray<FString> PackageNames;
-	Algo::Transform(Referencers, PackageNames, [](FName Name) { return Name.ToString(); });
-	TArray<UPackage*> Packages = AssetViewUtils::LoadPackages(PackageNames);
-	
-	for (UPackage* Package : Packages)
+	// Search for referencing packages to the currently changing user defined struct 
+	TArray<FAssetIdentifier> Referencers;
+	AssetRegistry.GetReferencers(Changed->GetPackage()->GetFName(), Referencers);
+	for (const FAssetIdentifier& Identifier : Referencers)
 	{
-		UObject* AssetObject = Package->FindAssetInPackage();
+		TArray<FAssetData> Assets;
+		AssetRegistry.GetAssetsByPackageName(Identifier.PackageName, Assets);
 
-		if (UOptimusDeformer* DeformerAsset = Cast<UOptimusDeformer>(AssetObject))
+		for (const FAssetData& Asset : Assets)
 		{
-			DeformerAsset->SetAllInstancesCanbeActive(false);
+			// Only load assets whose class is of UOptimusDeformer
+			if (Asset.IsInstanceOf(UOptimusDeformer::StaticClass()))
+			{
+				UOptimusDeformer* DeformerAsset = CastChecked<UOptimusDeformer>(Asset.GetAsset());
+
+				DeformerAsset->SetAllInstancesCanbeActive(false);
+			}
 		}
 	}
 
@@ -156,22 +161,26 @@ void FOptimusEditorModule::PostChange(const UUserDefinedStruct* Changed,
 	// Only recompile/reactivate all instances once all struct changes have been processed
 	if (UserDefinedStructsPendingPostChange == 0)
 	{
-		FAssetRegistryModule& AssetRegistryModule = FModuleManager::LoadModuleChecked<FAssetRegistryModule>("AssetRegistry");
-		TArray<FName> Referencers;
-		AssetRegistryModule.Get().GetReferencers(Changed->GetPackage()->GetFName(), Referencers);
+		FAssetRegistryModule& AssetRegistryModule = FModuleManager::LoadModuleChecked<FAssetRegistryModule>(FName("AssetRegistry"));
+		IAssetRegistry& AssetRegistry = AssetRegistryModule.Get();
 
-		TArray<FString> PackageNames;
-		Algo::Transform(Referencers, PackageNames, [](FName Name) { return Name.ToString(); });
-		TArray<UPackage*> Packages = AssetViewUtils::LoadPackages(PackageNames);
-	
-		for (UPackage* Package : Packages)
+		// Search for referencing packages to the currently changing user defined struct 
+		TArray<FAssetIdentifier> Referencers;
+		AssetRegistry.GetReferencers(Changed->GetPackage()->GetFName(), Referencers);
+		for (const FAssetIdentifier& Identifier : Referencers)
 		{
-			UObject* AssetObject = Package->FindAssetInPackage();
+			TArray<FAssetData> Assets;
+			AssetRegistry.GetAssetsByPackageName(Identifier.PackageName, Assets);
 
-			if (UOptimusDeformer* DeformerAsset = Cast<UOptimusDeformer>(AssetObject))
+			for (const FAssetData& Asset : Assets)
 			{
-				DeformerAsset->Compile();
-				DeformerAsset->SetAllInstancesCanbeActive(true);
+				// Only load assets whose class is of UOptimusDeformer
+				if (Asset.IsInstanceOf(UOptimusDeformer::StaticClass()))
+				{
+					UOptimusDeformer* DeformerAsset = CastChecked<UOptimusDeformer>(Asset.GetAsset());
+					DeformerAsset->Compile();
+					DeformerAsset->SetAllInstancesCanbeActive(true);
+				}
 			}
 		}	
 	}
