Optimus: follow-up 25610487, avoid loading unloaded deformer assets, which would not have any active instances that can cause problems

diff --git a/Engine/Plugins/Animation/DeformerGraph/Source/OptimusEditor/Private/OptimusEditorModule.cpp b/Engine/Plugins/Animation/DeformerGraph/Source/OptimusEditor/Private/OptimusEditorModule.cpp
--- a/Engine/Plugins/Animation/DeformerGraph/Source/OptimusEditor/Private/OptimusEditorModule.cpp
+++ b/Engine/Plugins/Animation/DeformerGraph/Source/OptimusEditor/Private/OptimusEditorModule.cpp
@@ -138,12 +138,14 @@ void FOptimusEditorModule::PreChange(const UUserDefinedStruct* Changed,
 
 		for (const FAssetData& Asset : Assets)
 		{
-			// Only load assets whose class is of UOptimusDeformer
+			
 			if (Asset.IsInstanceOf(UOptimusDeformer::StaticClass()))
 			{
-				UOptimusDeformer* DeformerAsset = CastChecked<UOptimusDeformer>(Asset.GetAsset());
-
-				DeformerAsset->SetAllInstancesCanbeActive(false);
+				// Only care about loaded assets, which may have active instances
+				if (UOptimusDeformer* DeformerAsset = CastChecked<UOptimusDeformer>(Asset.FastGetAsset(false)))
+				{
+					DeformerAsset->SetAllInstancesCanbeActive(false);
+				}
 			}
 		}
 	}
@@ -174,12 +176,14 @@ void FOptimusEditorModule::PostChange(const UUserDefinedStruct* Changed,
 
 			for (const FAssetData& Asset : Assets)
 			{
-				// Only load assets whose class is of UOptimusDeformer
 				if (Asset.IsInstanceOf(UOptimusDeformer::StaticClass()))
 				{
-					UOptimusDeformer* DeformerAsset = CastChecked<UOptimusDeformer>(Asset.GetAsset());
-					DeformerAsset->Compile();
-					DeformerAsset->SetAllInstancesCanbeActive(true);
+					// Only care about loaded assets, which may have active instances
+					if (UOptimusDeformer* DeformerAsset = CastChecked<UOptimusDeformer>(Asset.FastGetAsset(false)))
+					{
+						DeformerAsset->Compile();
+						DeformerAsset->SetAllInstancesCanbeActive(true);
+					}
 				}
 			}
 		}	
