DeformerGraph: Persistent buffers use a UAV per kernel in the graph. This allows us to use SkipBarrier on the UAV to allow concurrent subinvocations, and still use correct barriers between kernels.
This fixes some flickering seen in graphs that recompute normals using the scatter technique.

diff --git a/Engine/Plugins/Animation/DeformerGraph/Source/OptimusCore/Private/DataInterfaces/OptimusDataInterfaceAnimAttribute.cpp b/Engine/Plugins/Animation/DeformerGraph/Source/OptimusCore/Private/DataInterfaces/OptimusDataInterfaceAnimAttribute.cpp
--- a/Engine/Plugins/Animation/DeformerGraph/Source/OptimusCore/Private/DataInterfaces/OptimusDataInterfaceAnimAttribute.cpp
+++ b/Engine/Plugins/Animation/DeformerGraph/Source/OptimusCore/Private/DataInterfaces/OptimusDataInterfaceAnimAttribute.cpp
@@ -719,8 +719,6 @@ bool FOptimusAnimAttributeDataProviderProxy::IsValid(FValidationData const& InVa
 
 void FOptimusAnimAttributeDataProviderProxy::AllocateResources(FRDGBuilder& GraphBuilder)
 {
-	FComputeDataProviderRenderProxy::AllocateResources(GraphBuilder);
-
 	ArrayBuffers.Reset();
 	ArrayBufferSRVs.Reset();
 
diff --git a/Engine/Plugins/Animation/DeformerGraph/Source/OptimusCore/Private/DataInterfaces/OptimusDataInterfaceRawBuffer.cpp b/Engine/Plugins/Animation/DeformerGraph/Source/OptimusCore/Private/DataInterfaces/OptimusDataInterfaceRawBuffer.cpp
--- a/Engine/Plugins/Animation/DeformerGraph/Source/OptimusCore/Private/DataInterfaces/OptimusDataInterfaceRawBuffer.cpp
+++ b/Engine/Plugins/Animation/DeformerGraph/Source/OptimusCore/Private/DataInterfaces/OptimusDataInterfaceRawBuffer.cpp
@@ -467,7 +467,7 @@ bool FOptimusPersistentBufferDataProviderProxy::IsValid(FValidationData const& I
 	return true;
 }
 
-void FOptimusPersistentBufferDataProviderProxy::AllocateResources(FRDGBuilder& GraphBuilder)
+void FOptimusPersistentBufferDataProviderProxy::AllocateResources(FRDGBuilder& GraphBuilder, FAllocationData const& InAllocationData)
 {
 	TArray<int32> Count;
 	Count.Add(TotalElementCount);
@@ -477,7 +477,14 @@ void FOptimusPersistentBufferDataProviderProxy::AllocateResources(FRDGBuilder& G
 	ensure(Buffers.Num() == 1);
 	Buffer = Buffers[0];
 
-	BufferUAV = GraphBuilder.CreateUAV(Buffer, ERDGUnorderedAccessViewFlags::SkipBarrier);
+	// We want to use SkipBarrier to allow simultaineous execution of any subinvocations, but we want to keep barriers between kernels. 
+	// RDG will do this if we use a different UAV object per kernel based on the same underlying buffer. So we create one UAV per kernel here.
+	// This may end up being overkill for large graphs, in which case we will only want to create a UAV for each kernel that uses this data provider.
+	BufferUAVs.Reserve(InAllocationData.NumGraphKernels);
+	for (int32 BufferIndex = 0; BufferIndex < InAllocationData.NumGraphKernels; ++BufferIndex)
+	{
+		BufferUAVs.Add(GraphBuilder.CreateUAV(Buffer, ERDGUnorderedAccessViewFlags::SkipBarrier));
+	}
 }
 
 void FOptimusPersistentBufferDataProviderProxy::GatherDispatchData(FDispatchData const& InDispatchData)
@@ -488,7 +495,7 @@ void FOptimusPersistentBufferDataProviderProxy::GatherDispatchData(FDispatchData
 		FParameters& Parameters = ParameterArray[InvocationIndex];
 		Parameters.StartOffset = InDispatchData.bUnifiedDispatch ? 0 : StartOffset;
 		Parameters.BufferSize = InDispatchData.bUnifiedDispatch ? TotalElementCount : InvocationElementCounts[InvocationIndex];
-		Parameters.BufferUAV = BufferUAV;
+		Parameters.BufferUAV = BufferUAVs[InDispatchData.GraphKernelIndex];
 
 		StartOffset += InvocationElementCounts[InvocationIndex];
 	}
diff --git a/Engine/Plugins/Animation/DeformerGraph/Source/OptimusCore/Private/DataInterfaces/OptimusDataInterfaceRawBuffer.h b/Engine/Plugins/Animation/DeformerGraph/Source/OptimusCore/Private/DataInterfaces/OptimusDataInterfaceRawBuffer.h
--- a/Engine/Plugins/Animation/DeformerGraph/Source/OptimusCore/Private/DataInterfaces/OptimusDataInterfaceRawBuffer.h
+++ b/Engine/Plugins/Animation/DeformerGraph/Source/OptimusCore/Private/DataInterfaces/OptimusDataInterfaceRawBuffer.h
@@ -253,7 +253,7 @@ public:
 
 	//~ Begin FComputeDataProviderRenderProxy Interface
 	bool IsValid(FValidationData const& InValidationData) const override;
-	void AllocateResources(FRDGBuilder& GraphBuilder) override;
+	void AllocateResources(FRDGBuilder& GraphBuilder, FAllocationData const& InAllocationData);
 	void GatherDispatchData(FDispatchData const& InDispatchData) override;
 	//~ End FComputeDataProviderRenderProxy Interface
 
@@ -269,5 +269,5 @@ private:
 	const int32 LODIndex;
 
 	FRDGBufferRef Buffer;
-	FRDGBufferUAVRef BufferUAV;
+	TArray<FRDGBufferUAVRef> BufferUAVs;
 };
diff --git a/Engine/Plugins/Runtime/ComputeFramework/Source/ComputeFramework/Private/ComputeGraphWorker.cpp b/Engine/Plugins/Runtime/ComputeFramework/Source/ComputeFramework/Private/ComputeGraphWorker.cpp
--- a/Engine/Plugins/Runtime/ComputeFramework/Source/ComputeFramework/Private/ComputeGraphWorker.cpp
+++ b/Engine/Plugins/Runtime/ComputeFramework/Source/ComputeFramework/Private/ComputeGraphWorker.cpp
@@ -151,12 +151,13 @@ void FComputeGraphTaskWorker::SubmitWork(FRDGBuilder& GraphBuilder, FName InExec
 		}
 
 		// Allocate RDG resources for all the data providers in the graph.
+		FComputeDataProviderRenderProxy::FAllocationData AllocationData { NumKernels };
 		for (int32 DataProviderIndex = 0; DataProviderIndex < GraphInvocation.DataProviderRenderProxies.Num(); ++DataProviderIndex)
 		{
 			FComputeDataProviderRenderProxy* DataProvider = GraphInvocation.DataProviderRenderProxies[DataProviderIndex];
 			if (DataProvider != nullptr)
 			{
-				DataProvider->AllocateResources(GraphBuilder);
+				DataProvider->AllocateResources(GraphBuilder, AllocationData);
 			}
 		}
 	}
@@ -195,7 +196,7 @@ void FComputeGraphTaskWorker::SubmitWork(FRDGBuilder& GraphBuilder, FName InExec
 
 		// Allocate parameters buffer and fill from data providers.
 		TStridedView<FComputeKernelShader::FParameters> ParameterArray = GraphBuilder.AllocParameters<FComputeKernelShader::FParameters>(KernelInvocation.ShaderParameterMetadata, NumSubInvocations);
-		FComputeDataProviderRenderProxy::FDispatchData DispatchData{ NumSubInvocations, bIsUnifiedDispatch, 0, 0, ParameterArray.GetStride(), reinterpret_cast<uint8*>(&ParameterArray[0]) };
+		FComputeDataProviderRenderProxy::FDispatchData DispatchData{ KernelIndex, NumSubInvocations, bIsUnifiedDispatch, 0, 0, ParameterArray.GetStride(), reinterpret_cast<uint8*>(&ParameterArray[0]) };
 
 		// Iterate shader parameter members to fill the dispatch data structures.
 		// We assume that the members were filled out with a single data interface per member, and that the
diff --git a/Engine/Plugins/Runtime/ComputeFramework/Source/ComputeFramework/Public/ComputeFramework/ComputeDataProvider.h b/Engine/Plugins/Runtime/ComputeFramework/Source/ComputeFramework/Public/ComputeFramework/ComputeDataProvider.h
--- a/Engine/Plugins/Runtime/ComputeFramework/Source/ComputeFramework/Public/ComputeFramework/ComputeDataProvider.h
+++ b/Engine/Plugins/Runtime/ComputeFramework/Source/ComputeFramework/Public/ComputeFramework/ComputeDataProvider.h
@@ -74,12 +74,28 @@ public:
 	 */
 	virtual void GatherPermutations(FPermutationData& InOutPermutationData) const {}
 
-	/* Called once before any calls to GatherDispatchData() to allow any RDG resource allocation. */
+	/** Setup needed to allocate resources. */
+	struct FAllocationData
+	{
+		int32 NumGraphKernels;
+	};
+
+	// todo: Fix up derived classes to implement new AllocateResources().
+	UE_DEPRECATED(5.3, "Convert to using the new AllocateResources() that takes FAllocationData.")
 	virtual void AllocateResources(FRDGBuilder& GraphBuilder) {}
 
+	/* Called once before any calls to GatherDispatchData() to allow any RDG resource allocation. */
+	virtual void AllocateResources(FRDGBuilder& GraphBuilder, FAllocationData const& InAllocationData)
+	{
+		PRAGMA_DISABLE_DEPRECATION_WARNINGS
+		AllocateResources(GraphBuilder);
+		PRAGMA_ENABLE_DEPRECATION_WARNINGS
+	}
+
 	/** Setup needed to gather dispatch data. */
 	struct FDispatchData
 	{
+		int32 GraphKernelIndex;
 		int32 NumInvocations;
 		bool bUnifiedDispatch;
 		int32 ParameterStructSize;
