Deformer Graph: Harden against the underlying data interface class going AWOL.

#rnx

diff --git a/Engine/Plugins/Animation/DeformerGraph/Source/OptimusCore/Private/Nodes/OptimusNode_DataInterface.cpp b/Engine/Plugins/Animation/DeformerGraph/Source/OptimusCore/Private/Nodes/OptimusNode_DataInterface.cpp
--- a/Engine/Plugins/Animation/DeformerGraph/Source/OptimusCore/Private/Nodes/OptimusNode_DataInterface.cpp
+++ b/Engine/Plugins/Animation/DeformerGraph/Source/OptimusCore/Private/Nodes/OptimusNode_DataInterface.cpp
@@ -59,6 +59,10 @@ bool UOptimusNode_DataInterface::ValidateConnection(
 
 TOptional<FText> UOptimusNode_DataInterface::ValidateForCompile() const
 {
+	if (!DataInterfaceClass)
+	{
+		return LOCTEXT("NoAssociatedClass", "Node has none or invalid data interface class associated with it. Delete and re-create the node.");
+	}
 	// Ensure that we have something connected to the component binding input pin.
 	UOptimusComponentSourceBinding* PrimaryBinding = GetComponentBinding();
 	if (PrimaryBinding == nullptr)
@@ -103,6 +107,12 @@ void UOptimusNode_DataInterface::Serialize(FArchive& Ar)
 
 UOptimusComputeDataInterface* UOptimusNode_DataInterface::GetDataInterface(UObject* InOuter) const
 {
+	// Asset is probably broken, or refers to a class that no longer exists. 
+	if (!DataInterfaceClass)
+	{
+		return nullptr;
+	}
+	
 	// Legacy data may not have a DataInterfaceData object.
 	if (DataInterfaceData == nullptr || !DataInterfaceData->IsA(DataInterfaceClass))
 	{
