Optimus: Store view location + zoom in the graph object.

#rnx

diff --git a/Engine/Plugins/Animation/DeformerGraph/Source/OptimusCore/Private/OptimusNodeGraph.cpp b/Engine/Plugins/Animation/DeformerGraph/Source/OptimusCore/Private/OptimusNodeGraph.cpp
--- a/Engine/Plugins/Animation/DeformerGraph/Source/OptimusCore/Private/OptimusNodeGraph.cpp
+++ b/Engine/Plugins/Animation/DeformerGraph/Source/OptimusCore/Private/OptimusNodeGraph.cpp
@@ -1816,5 +1816,25 @@ bool UOptimusNodeGraph::RenameGraph(
 	return bSuccess;
 }
 
+#if WITH_EDITOR
+void UOptimusNodeGraph::SetViewLocationAndZoom(const FVector2D& InViewLocation, float InViewZoom)
+{
+	bViewLocationSet = true;
+	ViewLocation = InViewLocation;
+	ViewZoom = InViewZoom;
+}
+
+bool UOptimusNodeGraph::GetViewLocationAndZoom(FVector2D& OutViewLocation, float& OutViewZoom) const
+{
+	if (bViewLocationSet)
+	{
+		OutViewLocation = ViewLocation;
+		OutViewZoom = ViewZoom;
+	}
+	return bViewLocationSet;
+}
+
+#endif // WITH_EDITOR
+
 
 #undef LOCTEXT_NAMESPACE
diff --git a/Engine/Plugins/Animation/DeformerGraph/Source/OptimusCore/Public/OptimusNodeGraph.h b/Engine/Plugins/Animation/DeformerGraph/Source/OptimusCore/Public/OptimusNodeGraph.h
--- a/Engine/Plugins/Animation/DeformerGraph/Source/OptimusCore/Public/OptimusNodeGraph.h
+++ b/Engine/Plugins/Animation/DeformerGraph/Source/OptimusCore/Public/OptimusNodeGraph.h
@@ -342,6 +342,17 @@ public:
 	bool RenameGraph(
 		UOptimusNodeGraph* InGraph,
 		const FString& InNewName) override;
+
+
+#if WITH_EDITOR
+	// Set the graph view location and zoom, to ensure that the view location is stored between sessions and
+	// graph switching.
+	void SetViewLocationAndZoom(const FVector2D& InViewLocation, float InViewZoom);
+
+	// Returns the graph view location and zoom, if set. If the location has never been set, this function returns
+	// false and the resulting out values are left undefined.
+	bool GetViewLocationAndZoom(FVector2D& OutViewLocation, float& OutViewZoom) const;
+#endif
 	
 protected:
 	friend class UOptimusDeformer;
@@ -398,6 +409,17 @@ protected:
 	UPROPERTY(VisibleAnywhere, BlueprintReadOnly, Category=Overview)
 	EOptimusNodeGraphType GraphType = EOptimusNodeGraphType::Transient;
 
+#if WITH_EDITORONLY_DATA
+	UPROPERTY()
+	bool bViewLocationSet = false;
+	
+	UPROPERTY()
+	FVector2D ViewLocation = {0.0, 0.0};
+
+	UPROPERTY()
+	float ViewZoom = 0.0f;
+#endif
+	
 private:
 	IOptimusPathResolver* GetPathResolver() const;
 	
diff --git a/Engine/Plugins/Animation/DeformerGraph/Source/OptimusEditor/Private/OptimusEditor.cpp b/Engine/Plugins/Animation/DeformerGraph/Source/OptimusEditor/Private/OptimusEditor.cpp
--- a/Engine/Plugins/Animation/DeformerGraph/Source/OptimusEditor/Private/OptimusEditor.cpp
+++ b/Engine/Plugins/Animation/DeformerGraph/Source/OptimusEditor/Private/OptimusEditor.cpp
@@ -249,6 +249,8 @@ FLinearColor FOptimusEditor::GetWorldCentricTabColorScale() const
 
 void FOptimusEditor::OnClose()
 {
+	StoreCurrentViewLocation();
+	
 	IOptimusEditor::OnClose();
 }
 
@@ -259,12 +261,26 @@ bool FOptimusEditor::SetEditGraph(UOptimusNodeGraph* InNodeGraph)
 	{
 		PreviousEditedNodeGraph = EditorGraph->NodeGraph;
 
+		// Store the location/zoom for the graph we're leaving
+		StoreCurrentViewLocation();
+
 		GraphEditorWidget->ClearSelectionSet();
 
 		EditorGraph->Reset();
 		EditorGraph->InitFromNodeGraph(InNodeGraph);
 
-		// FIXME: Store pan/zoom
+		// If there was a stored view location in the node graph, use that, otherwise just frame
+		// all the nodes.
+		FVector2D Location;
+		float Zoom;
+		if (InNodeGraph->GetViewLocationAndZoom(Location, Zoom))
+		{
+			GraphEditorWidget->SetViewLocation(Location, Zoom);
+		}
+		else
+		{
+			GraphEditorWidget->ZoomToFit(false);
+		}
 
 		RefreshEvent.Broadcast();
 		return true;
@@ -1062,6 +1078,19 @@ TArray<UOptimusNode*> FOptimusEditor::GetSelectedModelNodes() const
 }
 
 
+void FOptimusEditor::StoreCurrentViewLocation()
+{
+	if (EditorGraph->NodeGraph)
+	{
+		FVector2D Location;
+		float Zoom;
+		
+		GraphEditorWidget->GetViewLocation(Location, Zoom);
+		EditorGraph->NodeGraph->SetViewLocationAndZoom(Location, Zoom);
+	}
+}
+
+
 void FOptimusEditor::OnDeformerModified(
 	EOptimusGlobalNotifyType InNotifyType, 
 	UObject* InModifiedObject
diff --git a/Engine/Plugins/Animation/DeformerGraph/Source/OptimusEditor/Private/OptimusEditor.h b/Engine/Plugins/Animation/DeformerGraph/Source/OptimusEditor/Private/OptimusEditor.h
--- a/Engine/Plugins/Animation/DeformerGraph/Source/OptimusEditor/Private/OptimusEditor.h
+++ b/Engine/Plugins/Animation/DeformerGraph/Source/OptimusEditor/Private/OptimusEditor.h
@@ -214,6 +214,8 @@ private:
 
 	TArray<UOptimusNode*> GetSelectedModelNodes() const;
 
+	void StoreCurrentViewLocation();
+
 	void OnDeformerModified(
 		EOptimusGlobalNotifyType InNotifyType, 
 		UObject *InModifiedObject
