Optimus: enabled localized resource domains, in addition to global domains
+ output pins domain can be expressed as a function of input pin domains
+ input pin domain, when left empty, is inferred from connected output pin

diff --git a/Engine/Plugins/Animation/DeformerGraph/Source/OptimusCore/Private/DataInterfaces/OptimusDataInterfaceCustomComputeKernel.cpp b/Engine/Plugins/Animation/DeformerGraph/Source/OptimusCore/Private/DataInterfaces/OptimusDataInterfaceCustomComputeKernel.cpp
--- a/Engine/Plugins/Animation/DeformerGraph/Source/OptimusCore/Private/DataInterfaces/OptimusDataInterfaceCustomComputeKernel.cpp
+++ b/Engine/Plugins/Animation/DeformerGraph/Source/OptimusCore/Private/DataInterfaces/OptimusDataInterfaceCustomComputeKernel.cpp
@@ -144,7 +144,7 @@ bool UOptimusCustomComputeKernelDataProvider::GetInvocationThreadCounts(
 		return false;
 	}
 	
-	TMap<FName, int32> EngineConstants;
+	TMap<FName, float> EngineConstants;
 	TMap<FName, TArray<int32>> ElementCountsPerDomain;
 
 	for(FName ExecutionDomain: ComponentSource->GetExecutionDomains())
@@ -170,8 +170,17 @@ bool UOptimusCustomComputeKernelDataProvider::GetInvocationThreadCounts(
 
 	using namespace Optimus::Expression;
 
-	FEngine Engine(EngineConstants);
-	TVariant<FExpressionObject, FParseError> ParseResult = Engine.Parse(NumThreadsExpression);
+	FEngine Engine;
+	TVariant<FExpressionObject, FParseError> ParseResult = Engine.Parse(NumThreadsExpression, [EngineConstants](FName InName)->TOptional<float>
+	{
+		if (const float* Value = EngineConstants.Find(InName))
+		{
+			return *Value;
+		};
+
+		return {};
+	});
+	
 	if (ParseResult.IsType<FParseError>())
 	{
 		return false;
@@ -180,16 +189,25 @@ bool UOptimusCustomComputeKernelDataProvider::GetInvocationThreadCounts(
 	OutInvocationThreadCount.Reset(NumInvocations);
 	for (int32 Index = 0; Index < NumInvocations; Index++)
 	{
-		for (TPair<FName, int32>& Constant: EngineConstants)
+		for (TPair<FName, float>& Constant: EngineConstants)
 		{
 			const FName ConstantName = Constant.Key;
-			int32& Value = Constant.Value;
+			float& Value = Constant.Value;
 			
 			const TArray<int32>& ElementCounts = ElementCountsPerDomain[ConstantName]; 
 			Value = ElementCounts.IsValidIndex(Index) ? ElementCounts[Index] : 1;
 		}
-		Engine.UpdateConstantValues(EngineConstants);
-		const int32 Count = Engine.Execute(ParseResult.Get<FExpressionObject>());
+		
+		const int32 Count = static_cast<int32>(Engine.Execute(ParseResult.Get<FExpressionObject>(),
+			[EngineConstants](FName InName)->TOptional<float>
+			{
+				if (const float* Value = EngineConstants.Find(InName))
+				{
+					return *Value;
+				};
+
+				return {};
+			}));
 
 		if (Count < 0)
 		{
@@ -209,8 +227,10 @@ bool UOptimusCustomComputeKernelDataProvider::GetInvocationThreadCounts(
 
 	// We want to make sure the results are the same for unified and non-unified
 	// Do the sum across invocations first and then evaluate the expression
-	for (auto& [ConstantName, Value]: EngineConstants)
+	for (TPair<FName, float>& Constant : EngineConstants)
 	{
+		const FName ConstantName = Constant.Key;
+		float& Value = Constant.Value;
 		Value = 0;
 		const TArray<int32>& ElementCounts = ElementCountsPerDomain[ConstantName];
 		
@@ -220,8 +240,17 @@ bool UOptimusCustomComputeKernelDataProvider::GetInvocationThreadCounts(
 		}
 	}
 	
-	Engine.UpdateConstantValues(EngineConstants);
-	const int32 TotalElementCountForUnifiedDispatch = Engine.Execute(ParseResult.Get<FExpressionObject>());
+	
+	const int32 TotalElementCountForUnifiedDispatch = static_cast<int32>(Engine.Execute(ParseResult.Get<FExpressionObject>(), 
+		[EngineConstants](FName InName)->TOptional<float>
+		{
+			if (const float* Value = EngineConstants.Find(InName))
+			{
+				return *Value;
+			};
+
+			return {};
+		}));
 	
 	if (OutTotalThreadCount != TotalElementCountForUnifiedDispatch)
 	{
diff --git a/Engine/Plugins/Animation/DeformerGraph/Source/OptimusCore/Private/DataInterfaces/OptimusDataInterfaceGraph.cpp b/Engine/Plugins/Animation/DeformerGraph/Source/OptimusCore/Private/DataInterfaces/OptimusDataInterfaceGraph.cpp
--- a/Engine/Plugins/Animation/DeformerGraph/Source/OptimusCore/Private/DataInterfaces/OptimusDataInterfaceGraph.cpp
+++ b/Engine/Plugins/Animation/DeformerGraph/Source/OptimusCore/Private/DataInterfaces/OptimusDataInterfaceGraph.cpp
@@ -118,6 +118,16 @@ FComputeDataProviderRenderProxy* UOptimusGraphDataProvider::GetRenderProxy()
 	return new FOptimusGraphDataProviderProxy(DeformerInstance, Variables, ParameterBufferSize);
 }
 
+void UOptimusGraphDataProvider::SetDeformerInstance(UOptimusDeformerInstance* InInstance)
+{
+	DeformerInstance = InInstance;
+}
+
+UOptimusDeformerInstance* UOptimusGraphDataProvider::GetDeformerInstance() const
+{
+	return DeformerInstance;
+}
+
 FOptimusGraphDataProviderProxy::FOptimusGraphDataProviderProxy(UOptimusDeformerInstance const* DeformerInstance, TArray<FOptimusGraphVariableDescription> const& Variables, int32 ParameterBufferSize)
 {
 	// Get all variables from deformer instance and fill buffer.
diff --git a/Engine/Plugins/Animation/DeformerGraph/Source/OptimusCore/Private/DataInterfaces/OptimusDataInterfaceGraph.h b/Engine/Plugins/Animation/DeformerGraph/Source/OptimusCore/Private/DataInterfaces/OptimusDataInterfaceGraph.h
--- a/Engine/Plugins/Animation/DeformerGraph/Source/OptimusCore/Private/DataInterfaces/OptimusDataInterfaceGraph.h
+++ b/Engine/Plugins/Animation/DeformerGraph/Source/OptimusCore/Private/DataInterfaces/OptimusDataInterfaceGraph.h
@@ -2,6 +2,7 @@
 
 #pragma once
 
+#include "IOptimusDeformerInstanceAccessor.h"
 #include "ComputeFramework/ComputeDataInterface.h"
 #include "ComputeFramework/ComputeDataProvider.h"
 #include "ComputeFramework/ShaderParamTypeDefinition.h"
@@ -57,7 +58,9 @@ private:
 
 /** Compute Framework Data Provider for marshaling compute graph parameters and variables. */
 UCLASS(BlueprintType, editinlinenew, Category = ComputeFramework)
-class UOptimusGraphDataProvider : public UComputeDataProvider
+class UOptimusGraphDataProvider :
+	public UComputeDataProvider,
+	public IOptimusDeformerInstanceAccessor
 {
 	GENERATED_BODY()
 
@@ -65,9 +68,6 @@ public:
 	UPROPERTY(EditAnywhere, BlueprintReadWrite, Category = Binding)
 	TObjectPtr<UMeshComponent> MeshComponent = nullptr;
 
-	UPROPERTY()
-	TObjectPtr<UOptimusDeformerInstance> DeformerInstance = nullptr;
-
 	UPROPERTY()
 	TArray<FOptimusGraphVariableDescription> Variables;
 
@@ -78,6 +78,15 @@ public:
 	//~ Begin UComputeDataProvider Interface
 	FComputeDataProviderRenderProxy* GetRenderProxy() override;
 	//~ End UComputeDataProvider Interface
+
+	//~ Begin IOptimusDeformerInstanceAccessor Interface
+	void SetDeformerInstance(UOptimusDeformerInstance* InInstance) override;
+	UOptimusDeformerInstance* GetDeformerInstance() const override;
+	//~ End IOptimusDeformerInstanceAccessor Interface
+
+private:
+	UPROPERTY()
+	TObjectPtr<UOptimusDeformerInstance> DeformerInstance = nullptr;
 };
 
 class FOptimusGraphDataProviderProxy : public FComputeDataProviderRenderProxy
diff --git a/Engine/Plugins/Animation/DeformerGraph/Source/OptimusCore/Private/DataInterfaces/OptimusDataInterfaceRawBuffer.cpp b/Engine/Plugins/Animation/DeformerGraph/Source/OptimusCore/Private/DataInterfaces/OptimusDataInterfaceRawBuffer.cpp
--- a/Engine/Plugins/Animation/DeformerGraph/Source/OptimusCore/Private/DataInterfaces/OptimusDataInterfaceRawBuffer.cpp
+++ b/Engine/Plugins/Animation/DeformerGraph/Source/OptimusCore/Private/DataInterfaces/OptimusDataInterfaceRawBuffer.cpp
@@ -226,13 +226,16 @@ bool UOptimusRawBufferDataProvider::GetLodAndInvocationElementCounts(
 	case EOptimusDataDomainType::Dimensional:
 		{
 			const FName ExecutionDomain = !DataDomain.DimensionNames.IsEmpty() ? DataDomain.DimensionNames[0] : NAME_None;
-			if (ComponentSourcePtr->GetComponentElementCountsForExecutionDomain(ExecutionDomain, ComponentPtr, OutLodIndex, OutInvocationElementCounts))
+			TArray<int32> Values;
+			if (ComponentSourcePtr->GetComponentElementCountsForExecutionDomain(ExecutionDomain, ComponentPtr, OutLodIndex, Values))
 			{
 				if (DataDomain.DimensionNames.Num() == 1)
 				{
-					for (int32& Count: OutInvocationElementCounts)
+					OutInvocationElementCounts.Reset(Values.Num());
+					for (int32& Count: Values)
 					{
 						Count *= FMath::Max(1, DataDomain.Multiplier);
+						OutInvocationElementCounts.Add(Count);
 					}
 					return true;
 				}
@@ -243,7 +246,7 @@ bool UOptimusRawBufferDataProvider::GetLodAndInvocationElementCounts(
 		}
 	case EOptimusDataDomainType::Expression:
 		{
-			TMap<FName, int32> EngineConstants;
+			TMap<FName, float> EngineConstants;
 			TMap<FName, TArray<int32>> ElementCountsPerDomain;
 
 			int32 NumInvocations = -1;
@@ -279,8 +282,17 @@ bool UOptimusRawBufferDataProvider::GetLodAndInvocationElementCounts(
 			
 			using namespace Optimus::Expression;
 
-			FEngine Engine(EngineConstants);
-			TVariant<FExpressionObject, FParseError> ParseResult = Engine.Parse(DataDomain.Expression);
+			FEngine Engine;
+			TVariant<FExpressionObject, FParseError> ParseResult = Engine.Parse(DataDomain.Expression, [EngineConstants](FName InName)->TOptional<float>
+			{
+				if (const float* Value = EngineConstants.Find(InName))
+				{
+					return *Value;
+				};
+
+				return {};
+			});
+			
 			if (ParseResult.IsType<FParseError>())
 			{
 				return false;
@@ -289,16 +301,25 @@ bool UOptimusRawBufferDataProvider::GetLodAndInvocationElementCounts(
 			OutInvocationElementCounts.Reset(NumInvocations);
 			for (int32 Index = 0; Index < NumInvocations; Index++)
 			{
-				for (TPair<FName, int32>& Constant: EngineConstants)
+				for (TPair<FName, float>& Constant: EngineConstants)
 				{
 					const FName ConstantName = Constant.Key;
-					int32& Value = Constant.Value;
+					float& Value = Constant.Value;
 					
 					const TArray<int32>& ElementCounts = ElementCountsPerDomain[ConstantName]; 
 					Value = ElementCounts.IsValidIndex(Index) ? ElementCounts[Index] : 1;
 				}
-				Engine.UpdateConstantValues(EngineConstants);
-				const int32 Count = Engine.Execute(ParseResult.Get<FExpressionObject>());
+				
+				const int32 Count = static_cast<int32>(Engine.Execute(ParseResult.Get<FExpressionObject>(),
+					[EngineConstants](FName InName)->TOptional<float>
+					{
+						if (const float* Value = EngineConstants.Find(InName))
+						{
+							return *Value;
+						};
+
+						return {};
+					}));
 
 				if (Count < 0)
 				{
@@ -318,19 +339,30 @@ bool UOptimusRawBufferDataProvider::GetLodAndInvocationElementCounts(
 
 			// We want to make sure the results are the same for unified and non-unified
 			// Do the sum across invocations first and then evaluate the expression
-			for (auto& [ConstantName, Value]: EngineConstants)
+			for (TPair<FName, float>& Constant : EngineConstants)
 			{
+				const FName ConstantName = Constant.Key;
+				float& Value = Constant.Value;
+				
 				Value = 0;
 				const TArray<int32>& ElementCounts = ElementCountsPerDomain[ConstantName];
 				
-				for (int32 Count : ElementCounts )
+				for (int32 Count : ElementCounts)
 				{
 					Value += Count;
 				}
 			}
 			
-			Engine.UpdateConstantValues(EngineConstants);
-			const int32 TotalElementCountForUnifiedDispatch = Engine.Execute(ParseResult.Get<FExpressionObject>());
+			const int32 TotalElementCountForUnifiedDispatch = static_cast<int32>(Engine.Execute(ParseResult.Get<FExpressionObject>(),
+				[EngineConstants](FName InName)->TOptional<float>
+				{
+					if (const float* Value = EngineConstants.Find(InName))
+					{
+						return *Value;
+					};
+
+					return {};
+				}));
 			
 			if (TotalElementCount != TotalElementCountForUnifiedDispatch)
 			{
@@ -344,14 +376,55 @@ bool UOptimusRawBufferDataProvider::GetLodAndInvocationElementCounts(
 	return false;
 }
 
+bool UOptimusRawBufferDataProvider::GetInvocationElementCounts(TArray<int32>& OutInvocationElementCounts) const
+{
+	TArray<float> Values = DeformerInstance->GetConstantValuePerInvocation(DomainConstantIdentifier);
+
+	// Can happen if the bound component does not have actual data, like when there is no preview mesh
+	if (Values.Num() == 0)
+	{
+		return false;
+	}
+	
+	OutInvocationElementCounts.Reset(Values.Num());
+	for (const float& Value : Values)
+	{
+		OutInvocationElementCounts.Add(static_cast<int32>(Value));
+	}
+
+	return true;
+}
+
+void UOptimusRawBufferDataProvider::SetDeformerInstance(UOptimusDeformerInstance* InInstance)
+{
+	DeformerInstance = InInstance;
+}
+
+UOptimusDeformerInstance* UOptimusRawBufferDataProvider::GetDeformerInstance() const
+{
+	return DeformerInstance;
+}
+
 
 FComputeDataProviderRenderProxy* UOptimusTransientBufferDataProvider::GetRenderProxy()
 {
 	int32 LodIndex;
 	TArray<int32> InvocationCounts;
-	if (!GetLodAndInvocationElementCounts(LodIndex, InvocationCounts))
+	
+	if (DomainConstantIdentifier.IsValid())
 	{
-		InvocationCounts.Reset();
+		// Querying the deformer instance for the domain of this buffer
+		if (!GetInvocationElementCounts(InvocationCounts))
+		{
+			InvocationCounts.Reset();
+		}
+	}
+	else
+	{
+		if (!GetLodAndInvocationElementCounts(LodIndex, InvocationCounts))
+		{
+			InvocationCounts.Reset();
+		}
 	}
 	
 	return new FOptimusTransientBufferDataProviderProxy(InvocationCounts, ElementStride, RawStride);
diff --git a/Engine/Plugins/Animation/DeformerGraph/Source/OptimusCore/Private/DataInterfaces/OptimusDataInterfaceRawBuffer.h b/Engine/Plugins/Animation/DeformerGraph/Source/OptimusCore/Private/DataInterfaces/OptimusDataInterfaceRawBuffer.h
--- a/Engine/Plugins/Animation/DeformerGraph/Source/OptimusCore/Private/DataInterfaces/OptimusDataInterfaceRawBuffer.h
+++ b/Engine/Plugins/Animation/DeformerGraph/Source/OptimusCore/Private/DataInterfaces/OptimusDataInterfaceRawBuffer.h
@@ -2,9 +2,11 @@
 
 #pragma once
 
+#include "IOptimusDeformerInstanceAccessor.h"
 #include "ComputeFramework/ComputeDataProvider.h"
 #include "ComputeFramework/ShaderParamTypeDefinition.h"
 #include "OptimusComputeDataInterface.h"
+#include "OptimusConstant.h"
 #include "OptimusDataDomain.h"
 #include "RenderGraphFwd.h"
 
@@ -63,6 +65,9 @@ public:
 	UPROPERTY()
 	TWeakObjectPtr<UOptimusComponentSourceBinding> ComponentSourceBinding;
 
+	UPROPERTY()
+	FOptimusConstantIdentifier DomainConstantIdentifier;
+	
 protected:
 	virtual bool UseSplitBuffers() const { return true; }
 
@@ -79,6 +84,7 @@ protected:
 			Provider->DataDomain = DataDomain;
 			Provider->ElementStride = ValueType->GetResourceElementSize();
 			Provider->RawStride = GetRawStride();
+			Provider->DomainConstantIdentifier = DomainConstantIdentifier;
 		}
 		return Provider;
 	}
@@ -142,7 +148,9 @@ protected:
 
 /** Compute Framework Data Provider for a transient buffer. */
 UCLASS(Abstract)
-class OPTIMUSCORE_API UOptimusRawBufferDataProvider : public UComputeDataProvider
+class OPTIMUSCORE_API UOptimusRawBufferDataProvider :
+	public UComputeDataProvider,
+	public IOptimusDeformerInstanceAccessor
 {
 	GENERATED_BODY()
 
@@ -155,6 +163,18 @@ public:
 		int32& OutLodIndex,
 		TArray<int32>& OutInvocationElementCounts
 		) const;
+
+	/** Helper function to calculate the element count given a constant identifier,
+	 *	LOD is handled by the deformer instance
+	 */
+	bool GetInvocationElementCounts(
+		TArray<int32>& OutInvocationElementCounts
+		) const;
+
+	//~ Begin IOptimusDeformerInstanceAccessor Interface
+	void SetDeformerInstance(UOptimusDeformerInstance* InInstance) override;
+	UOptimusDeformerInstance* GetDeformerInstance() const override;
+	//~ End IOptimusDeformerInstanceAccessor Interface
 	
 	/** The skinned mesh component that governs the sizing and LOD of this buffer */
 	UPROPERTY()
@@ -172,6 +192,13 @@ public:
 
 	UPROPERTY()
 	int32 RawStride = 0;
+
+	UPROPERTY()
+	FOptimusConstantIdentifier DomainConstantIdentifier;
+	
+private:
+	UPROPERTY()
+	TObjectPtr<UOptimusDeformerInstance> DeformerInstance = nullptr;
 };
 
 
diff --git a/Engine/Plugins/Animation/DeformerGraph/Source/OptimusCore/Private/Nodes/OptimusNode_ComputeKernelBase.cpp b/Engine/Plugins/Animation/DeformerGraph/Source/OptimusCore/Private/Nodes/OptimusNode_ComputeKernelBase.cpp
--- a/Engine/Plugins/Animation/DeformerGraph/Source/OptimusCore/Private/Nodes/OptimusNode_ComputeKernelBase.cpp
+++ b/Engine/Plugins/Animation/DeformerGraph/Source/OptimusCore/Private/Nodes/OptimusNode_ComputeKernelBase.cpp
@@ -10,9 +10,11 @@
 #include "DataInterfaces/OptimusDataInterfaceGraph.h"
 #include "DataInterfaces/OptimusDataInterfaceRawBuffer.h"
 #include "IOptimusValueProvider.h"
+#include "OptimusComponentSource.h"
 #include "OptimusDataTypeRegistry.h"
 #include "OptimusHelpers.h"
 #include "OptimusKernelSource.h"
+#include "OptimusConstant.h"
 
 #include UE_INLINE_GENERATED_CPP_BY_NAME(OptimusNode_ComputeKernelBase)
 
@@ -57,7 +59,8 @@ FOptimus_ComputeKernelResult UOptimusNode_ComputeKernelBase::CreateComputeKernel
 	const UOptimusComponentSourceBinding* InGraphDataComponentBinding,
 	const UComputeDataInterface* InKernelDataInterface,
 	FOptimus_InterfaceBindingMap& OutInputDataBindings,
-	FOptimus_InterfaceBindingMap& OutOutputDataBindings
+	FOptimus_InterfaceBindingMap& OutOutputDataBindings,
+	FOptimusConstantContainer& OutConstantContainer
 ) const
 {
 	// Maps friendly name to unique name for each struct type
@@ -140,6 +143,8 @@ FOptimus_ComputeKernelResult UOptimusNode_ComputeKernelBase::CreateComputeKernel
 		return ReturnError(LOCTEXT("ZeroOrMultiplePrimaryBindings", "Primary Group has zero or more than one Component Bindings"));
 	}
 
+	FOptimusKernelConstantContainer& OutKernelConstantContainer = OutConstantContainer.AddContainerForKernel();
+
 	UOptimusKernelSource* KernelSource = NewObject<UOptimusKernelSource>(InKernelSourceOuter);
 
 	// Wrap functions for unconnected resource pins (or value pins) that return default values
@@ -157,7 +162,7 @@ FOptimus_ComputeKernelResult UOptimusNode_ComputeKernelBase::CreateComputeKernel
 					InTraversalContext, SubPin, GroupIndex == 0 ? FString() : Pin->GetName(),
 					InNodeDataInterfaceMap, InLinkDataInterfaceMap, InValueNodes,
 					InGraphDataInterface, InGraphDataComponentBinding,
-					KernelSource, GeneratedFunctions, OutInputDataBindings
+					KernelSource, GeneratedFunctions, OutInputDataBindings, OutKernelConstantContainer
 					);
 				if (Result.IsSet())
 				{
@@ -172,7 +177,7 @@ FOptimus_ComputeKernelResult UOptimusNode_ComputeKernelBase::CreateComputeKernel
 			ProcessOutputPinForComputeKernel(
 				InTraversalContext, Pin,
 				InNodeDataInterfaceMap, InLinkDataInterfaceMap,
-				KernelSource, GeneratedFunctions, OutOutputDataBindings);
+				KernelSource, GeneratedFunctions, OutOutputDataBindings, OutKernelConstantContainer);
 		}
 	}
 
@@ -401,7 +406,8 @@ TOptional<FText> UOptimusNode_ComputeKernelBase::ProcessInputPinForComputeKernel
 	const UOptimusComponentSourceBinding* InGraphDataComponentBinding,
 	UOptimusKernelSource* InKernelSource,
 	TArray<FString>& OutGeneratedFunctions,
-	FOptimus_InterfaceBindingMap& OutInputDataBindings
+	FOptimus_InterfaceBindingMap& OutInputDataBindings,
+	FOptimusKernelConstantContainer& OutKernelConstantContainer
 	) const
 {
 	const UOptimusNodePin* OutputPin = nullptr;
@@ -558,6 +564,54 @@ TOptional<FText> UOptimusNode_ComputeKernelBase::ProcessInputPinForComputeKernel
 				*NamespacePrefix, *ValueType->ToString(), *InInputPin->GetName(), *OptionalParamStr, *ValueStr, *NamespaceSuffix));
 	}
 
+	if (InInputPin->GetDataDomain().IsOneDimensional())
+	{
+		// Get the component binding from the upstream connection. Unconnected input pin should use the group's component binding
+		int32 ComponentBindingIndex = INDEX_NONE;
+		
+		if (const UOptimusNodePin* GroupPin = InInputPin->GetParentPin())
+		{
+			TArray<UOptimusComponentSourceBinding*> ComponentBindings = GetGroupComponentSourceBindings(GroupPin).Array();
+			if (ensure(ComponentBindings.Num() == 1))
+			{
+				ComponentBindingIndex = ComponentBindings[0]->GetIndex();
+			}
+		}
+		
+		// Use the domain of input pin if it is defined
+		// Use the domain of connected output pin otherwise
+		if (InInputPin->GetDataDomain().IsFullyDefined())
+		{
+			OutKernelConstantContainer.AddToKernelContainer({
+				{this, *InGroupName, InInputPin->GetFName()},
+				{InInputPin->GetDataDomain().AsExpression().GetValue()},
+				ComponentBindingIndex,
+				EOptimusConstantType::Input});
+		}
+		else if (ensure(OutputPin) && ensure(OutputPin->GetDataDomain().IsOneDimensional()) && ensure(OutputPin->GetDataDomain().IsFullyDefined()))
+		{
+			if (InLinkDataInterfaceMap.Contains(OutputPin))
+			{
+				if (const UOptimusRawBufferDataInterface* RawBufferDataInterface = Cast<UOptimusRawBufferDataInterface>(InLinkDataInterfaceMap[OutputPin]))
+				{
+					OutKernelConstantContainer.AddToKernelContainer({
+						{this, *InGroupName, InInputPin->GetFName()},
+						{RawBufferDataInterface->DomainConstantIdentifier},
+						ComponentBindingIndex,
+						EOptimusConstantType::Input});	
+				}
+			}
+			else
+			{
+				OutKernelConstantContainer.AddToKernelContainer({
+					{this, *InGroupName, InInputPin->GetFName()},
+					{OutputPin->GetDataDomain().AsExpression().Get({})},
+					ComponentBindingIndex,
+					EOptimusConstantType::Input});
+			}
+		}
+	}
+
 	return {};
 }
 
@@ -569,7 +623,8 @@ void UOptimusNode_ComputeKernelBase::ProcessOutputPinForComputeKernel(
 	const FOptimus_PinToDataInterfaceMap& InLinkDataInterfaceMap,
 	UOptimusKernelSource* InKernelSource,
 	TArray<FString>& OutGeneratedFunctions,
-	FOptimus_InterfaceBindingMap& OutOutputDataBindings
+	FOptimus_InterfaceBindingMap& OutOutputDataBindings,
+	FOptimusKernelConstantContainer& OutKernelConstantContainer
 	) const
 {
 	FShaderParamTypeDefinition IndexParamDef;
@@ -712,6 +767,27 @@ void UOptimusNode_ComputeKernelBase::ProcessOutputPinForComputeKernel(
 			FString::Printf(TEXT("void Write%s(%s, %s Value) { }"),
 				*InOutputPin->GetName(), *FString::Join(StubIndexes, TEXT(", ")), *ValueType->ToString()));
 	}
+
+	// Gather the constants
+	if (InOutputPin->GetDataDomain().IsOneDimensional() && ensure(InOutputPin->GetDataDomain().IsFullyDefined()))
+	{
+		if (TOptional<FString> Expression = InOutputPin->GetDataDomain().AsExpression())
+		{
+			// Get the component binding from the upstream connection.
+			int32 ComponentBindingIndex = INDEX_NONE;
+			TArray<UOptimusComponentSourceBinding*> ComponentBindings = GetOwningGraph()->GetComponentSourceBindingsForPin(InOutputPin).Array();
+			if (ensure(ComponentBindings.Num() == 1))
+			{
+				ComponentBindingIndex = ComponentBindings[0]->GetIndex();
+			}
+			
+			OutKernelConstantContainer.AddToKernelContainer({
+				{this, NAME_None, InOutputPin->GetFName()},
+				{*Expression},
+				ComponentBindingIndex,
+				EOptimusConstantType::Output});
+		}
+	}
 }
 
 void UOptimusNode_ComputeKernelBase::BindKernelDataInterfaceForComputeKernel(
diff --git a/Engine/Plugins/Animation/DeformerGraph/Source/OptimusCore/Private/Nodes/OptimusNode_ComputeKernelBase.h b/Engine/Plugins/Animation/DeformerGraph/Source/OptimusCore/Private/Nodes/OptimusNode_ComputeKernelBase.h
--- a/Engine/Plugins/Animation/DeformerGraph/Source/OptimusCore/Private/Nodes/OptimusNode_ComputeKernelBase.h
+++ b/Engine/Plugins/Animation/DeformerGraph/Source/OptimusCore/Private/Nodes/OptimusNode_ComputeKernelBase.h
@@ -9,6 +9,7 @@
 
 #include "OptimusNode_ComputeKernelBase.generated.h"
 
+struct FOptimusKernelConstantContainer;
 class UComputeSource;
 
 
@@ -43,7 +44,8 @@ public:
 		const UOptimusComponentSourceBinding* InGraphDataComponentBinding,
 		const UComputeDataInterface* InKernelDataInterface,
 		FOptimus_InterfaceBindingMap& OutInputDataBindings,
-		FOptimus_InterfaceBindingMap& OutOutputDataBindings
+		FOptimus_InterfaceBindingMap& OutOutputDataBindings,
+		FOptimusConstantContainer& OutConstantContainer
 	) const override;
 
 	FName GetExecutionDomain() const override PURE_VIRTUAL(UOptimusNode_ComputeKernelBase::GetExecutionDomain, return NAME_None; );
@@ -88,7 +90,8 @@ private:
 		const UOptimusComponentSourceBinding* InGraphDataComponentBinding,
 		UOptimusKernelSource* InKernelSource,
 		TArray<FString>& OutGeneratedFunctions,
-		FOptimus_InterfaceBindingMap& OutInputDataBindings
+		FOptimus_InterfaceBindingMap& OutInputDataBindings,
+		FOptimusKernelConstantContainer& OutKernelConstantContainer
 		) const;
 
 	void ProcessOutputPinForComputeKernel(
@@ -98,7 +101,8 @@ private:
 		const FOptimus_PinToDataInterfaceMap& InLinkDataInterfaceMap,
 		UOptimusKernelSource* InKernelSource,
 		TArray<FString>& OutGeneratedFunctions,
-		FOptimus_InterfaceBindingMap& OutOutputDataBindings
+		FOptimus_InterfaceBindingMap& OutOutputDataBindings,
+		FOptimusKernelConstantContainer& OutKernelConstantContainer
 		) const;
 
 	void BindKernelDataInterfaceForComputeKernel(
diff --git a/Engine/Plugins/Animation/DeformerGraph/Source/OptimusCore/Private/OptimusConstant.cpp b/Engine/Plugins/Animation/DeformerGraph/Source/OptimusCore/Private/OptimusConstant.cpp
new file mode 100644
--- /dev/null
+++ b/Engine/Plugins/Animation/DeformerGraph/Source/OptimusCore/Private/OptimusConstant.cpp
@@ -0,0 +1,250 @@
+﻿// Copyright Epic Games, Inc. All Rights Reserved.
+#include "OptimusConstant.h"
+
+#include "OptimusNode.h"
+#include "OptimusExpressionEvaluator.h"
+
+
+FOptimusConstantIdentifier::FOptimusConstantIdentifier(const UOptimusNode* InNode,  const FName& InGroupName , const FName& InConstantName)
+{
+	NodePath = *InNode->GetNodePath();
+	GroupName = InGroupName;
+	ConstantName = InConstantName;
+}
+
+void FOptimusKernelConstantContainer::AddToKernelContainer(const FOptimusConstant& InConstant)
+{
+	TArray<FOptimusConstant>& Constants =
+		InConstant.Type == EOptimusConstantType::Input ?
+			InputConstants :OutputConstants;
+
+	int32 ConstantIndex = Constants.Add(InConstant);
+
+	GroupNameToBindingIndex.FindOrAdd(InConstant.Identifier.GroupName) = InConstant.ComponentBindingIndex;
+}
+
+FOptimusKernelConstantContainer& FOptimusConstantContainer::AddContainerForKernel()
+{
+	return KernelContainers.AddDefaulted_GetRef();
+}
+
+
+
+void FOptimusConstantContainer::Reset()
+{
+	KernelContainers.Reset();
+}
+
+
+bool FOptimusConstantContainerInstance::Initialize(const FOptimusConstantContainer& InContainer,
+                                                   const TMap<int32, TMap<FName, TArray<float>>>& InBindingIndexToConstantValues)
+{
+	for (int32 KernelIndex = 0; KernelIndex < InContainer.KernelContainers.Num(); KernelIndex++)
+	{
+		const FOptimusKernelConstantContainer& KernelContainer = InContainer.KernelContainers[KernelIndex];
+
+		// Process Input Constants
+		for (const FOptimusConstant& InputConstant : KernelContainer.InputConstants)
+		{
+			// Constant that directly grab value from another constant 
+			if (InputConstant.Definition.ReferencedConstant.IsValid())
+			{
+				ConstantToValuePerInvocation.Add(InputConstant.Identifier) = GetConstantValuePerInvocation(InputConstant.Definition.ReferencedConstant);
+				continue;
+			}
+
+			// Otherwise evaluate the current constant
+			if (ensure(!InputConstant.Definition.Expression.IsEmpty()))
+			{
+				TMap<FName, float> ValuePerDependentConstant;
+				TMap<FName, TArray<float>> ValuePerInvocationPerDependentConstant;
+				
+				if (!EvaluateAndSaveResult(
+					InputConstant.Identifier,
+					InputConstant.Definition.Expression,
+					InBindingIndexToConstantValues[InputConstant.ComponentBindingIndex]
+					))
+				{
+					return false;
+				}
+			}
+		}
+
+		// Process Output Constants
+		
+		// Output constants can access binding constants of all group bindings
+		TMap<FName, TArray<float>> BindingConstantsForOutput;
+
+		for (const TPair<FName, int32>& Group : KernelContainer.GroupNameToBindingIndex)
+		{
+			const FName& GroupName = Group.Key;
+			const int32 GroupBinding = Group.Value;
+			
+			const TMap<FName, TArray<float>>& BindingConstantValues = InBindingIndexToConstantValues[GroupBinding];
+			
+			for (const TPair<FName, TArray<float>>& BindingConstantValue : BindingConstantValues)
+			{
+				const FName& ConstantName = BindingConstantValue.Key;
+		
+				if (GroupName == NAME_None)
+				{
+					BindingConstantsForOutput.Add(BindingConstantValue);
+				}
+				else
+				{
+					FName LocalConstantName = *(GroupName.ToString() + TEXT(".") + ConstantName.ToString());
+
+					float Sum = 0;
+					for ( const float& Value : BindingConstantValue.Value)
+					{
+						Sum += Value;
+					}
+					
+					BindingConstantsForOutput.Add(LocalConstantName) = {Sum};
+				}
+			}
+		}
+		
+		for (const FOptimusConstant& OutputConstant : KernelContainer.OutputConstants)
+		{
+			// Constants that directly grab value from another constant 
+			if (OutputConstant.Definition.ReferencedConstant.IsValid())
+			{
+				ConstantToValuePerInvocation.Add(OutputConstant.Identifier) = GetConstantValuePerInvocation(OutputConstant.Definition.ReferencedConstant);
+				continue;
+			}
+
+			// Otherwise evaluate the current constant
+			if (ensure(!OutputConstant.Definition.Expression.IsEmpty()))
+			{
+				TMap<FName, TArray<float>> ValuePerInvocationPerDependentConstant = BindingConstantsForOutput;
+
+				for (const FOptimusConstant& InputConstant : KernelContainer.InputConstants)
+				{
+					const TArray<float>& InputValuePerInvocation = GetConstantValuePerInvocation(InputConstant.Identifier);
+					FName InputConstantLocalName = InputConstant.Identifier.GetLocalConstantName();
+
+					if (InputConstant.Identifier.GroupName == OutputConstant.Identifier.GroupName)
+					{
+						ValuePerInvocationPerDependentConstant.Add(InputConstantLocalName) = InputValuePerInvocation;	
+					}
+					else
+					{
+						float Sum = 0;
+						for (const float& Value : InputValuePerInvocation)
+						{
+							Sum+=Value;
+						}
+								
+						ValuePerInvocationPerDependentConstant.Add(InputConstantLocalName) = {Sum};	
+					}
+				}
+				
+				if (!EvaluateAndSaveResult(
+					OutputConstant.Identifier,
+					OutputConstant.Definition.Expression,
+					ValuePerInvocationPerDependentConstant
+					))
+				{
+					return false;
+				}
+			}
+		}
+	}
+
+	return true;
+}
+
+TArray<float> FOptimusConstantContainerInstance::GetConstantValuePerInvocation(const FOptimusConstantIdentifier& InIdentifier) const
+{
+	if (const TArray<float>* Values = ConstantToValuePerInvocation.Find(InIdentifier))
+	{
+		return *Values;
+	}
+
+	return {};
+}
+
+bool FOptimusConstantContainerInstance::EvaluateAndSaveResult(const FOptimusConstantIdentifier InIdentifier,
+	const FString& InExpression, const TMap<FName, TArray<float>>& InDependentConstants)
+{
+	TMap<FName, float> DependentConstantToValue;
+
+	int32 NumInvocations = 1;
+	for (const TPair<FName, TArray<float>>& DependentConstant : InDependentConstants)
+	{
+		DependentConstantToValue.Add(DependentConstant.Key, 0);
+
+		check(DependentConstant.Value.Num() != 0);
+							
+		if (DependentConstant.Value.Num() > 1)
+		{
+			if (NumInvocations == 1)
+			{
+				NumInvocations = DependentConstant.Value.Num();
+			}
+			else if (!ensure(NumInvocations == DependentConstant.Value.Num()))
+			{
+				return false;
+			}
+		}
+	}
+
+	using namespace Optimus::Expression;
+
+	FEngine Engine;
+	TVariant<FExpressionObject, FParseError> ParseResult = Engine.Parse(InExpression, [DependentConstantToValue](FName InName)->TOptional<float>
+	{
+		if (const float* Value = DependentConstantToValue.Find(InName))
+		{
+			return *Value;
+		};
+
+		return {};
+	});
+
+	if (ParseResult.IsType<FParseError>())
+	{
+		return false;
+	}
+
+	TArray<float> ValuePerInvocation;
+	ValuePerInvocation.Reset(NumInvocations);
+	
+	for (int32 Index = 0; Index < NumInvocations; Index++)
+	{
+		for (TPair<FName, float>& DependentConstant: DependentConstantToValue)
+		{
+			const FName ConstantName = DependentConstant.Key;
+			float& Value = DependentConstant.Value;
+					
+			const TArray<float>& Values = InDependentConstants[ConstantName];
+			if (ensure(Values.Num() >= 1))
+			{
+				Value = Values.IsValidIndex(Index) ? Values[Index] : Values[0];
+			}
+			else
+			{
+				return false;
+			}
+		}
+					
+		const float Result = Engine.Execute(ParseResult.Get<FExpressionObject>(),
+			[DependentConstantToValue](FName InName)->TOptional<float>
+			{
+				if (const float* Value = DependentConstantToValue.Find(InName))
+				{
+					return *Value;
+				};
+
+				return {};
+			});
+		
+		ValuePerInvocation.Add(Result);
+	}
+
+	ConstantToValuePerInvocation.Add(InIdentifier) = ValuePerInvocation;
+
+	return true;
+}
+
diff --git a/Engine/Plugins/Animation/DeformerGraph/Source/OptimusCore/Private/OptimusDataDomain.cpp b/Engine/Plugins/Animation/DeformerGraph/Source/OptimusCore/Private/OptimusDataDomain.cpp
--- a/Engine/Plugins/Animation/DeformerGraph/Source/OptimusCore/Private/OptimusDataDomain.cpp
+++ b/Engine/Plugins/Animation/DeformerGraph/Source/OptimusCore/Private/OptimusDataDomain.cpp
@@ -30,33 +30,6 @@ FString Optimus::FormatDimensionNames(const TArray<FName>& InNames)
 }
 
 
-TOptional<int32> FOptimusDataDomain::GetElementCount(TMap<FName, int32> InDomainCounts) const
-{
-	switch(Type)
-	{
-	case EOptimusDataDomainType::Dimensional:
-		{
-			if (DimensionNames.IsEmpty())
-			{
-				return 1;
-			}
-			else if (DimensionNames.Num() == 1)
-			{
-				if (const int32* Value = InDomainCounts.Find(DimensionNames[0]))
-				{
-					return *Value;
-				}
-			}
-		}
-		break;
-	case EOptimusDataDomainType::Expression:
-		return Optimus::Expression::FEngine(InDomainCounts).Evaluate(Expression);
-	}
-	
-	return {};
-}
-
-
 FString FOptimusDataDomain::ToString() const
 {
 	switch(Type)
@@ -151,3 +124,117 @@ void FOptimusDataDomain::BackCompFixupLevels()
 		LevelNames_DEPRECATED.Reset();
 	}
 }
+
+TOptional<FString> FOptimusDataDomain::AsExpression() const
+{
+	if (IsSingleton())
+	{
+		return {};
+	}
+
+	if (IsMultiDimensional())
+	{
+		return {};
+	}
+
+	if (Type == EOptimusDataDomainType::Dimensional)
+	{
+		if (Multiplier > 1)
+		{
+			return FString::Printf(TEXT("%s * %d"), *DimensionNames[0].ToString(), Multiplier);
+		}
+		
+		return DimensionNames[0].ToString();
+	}
+
+	if (Type == EOptimusDataDomainType::Expression)
+	{
+		return Expression;
+	}
+
+	return {};
+}
+
+FString FOptimusDataDomain::GetDisplayName() const
+{
+	switch(Type)
+	{
+	case EOptimusDataDomainType::Dimensional:
+		{
+			if (DimensionNames.IsEmpty())
+			{
+				return TEXT("Parameter");
+			}
+			else
+			{
+				TArray<FString> Names;
+				for (FName DomainLevelName: DimensionNames)
+				{
+					Names.Add(DomainLevelName.ToString());
+				}
+				return FString::Join(Names, *FString(UTF8TEXT(" › ")));
+			}
+		}
+		
+	case EOptimusDataDomainType::Expression:
+		if (Expression.IsEmpty())
+		{
+			return TEXT("Undefined");
+		}
+		return Expression.TrimStartAndEnd();
+	}
+
+	checkNoEntry();
+	return TEXT("");	
+}
+
+bool FOptimusDataDomain::AreCompatible(const FOptimusDataDomain& InOutput, const FOptimusDataDomain& InInput,
+                                       FString* OutReason)
+{
+	if (!InOutput.IsFullyDefined())
+	{
+		return false;
+	}
+
+	// We don't allow resource -> value connections. All other combos are legit. 
+	// Value -> Resource just means the resource gets filled with the value.
+	if (!InOutput.IsSingleton() && InInput.IsSingleton())
+	{
+		if (OutReason)
+		{
+			*OutReason = TEXT("Can't connect a resource output into a value input.");
+		}
+		return false;
+	}
+
+	if (InOutput.IsSingleton() && !InInput.IsSingleton())
+	{
+		if (!InInput.IsFullyDefined())
+		{
+			if (OutReason)
+			{
+				*OutReason = TEXT("Can't connect a value output to a input with undefined data domain");
+			}
+			return false;
+		}
+		
+	}
+	
+	// If it's resource -> resource, check that the domains are compatible
+	if (!InOutput.IsSingleton() && !InInput.IsSingleton())
+	{
+		if (InOutput != InInput && InInput.IsFullyDefined())
+		{
+			if (OutReason)
+			{
+				*OutReason = FString::Printf(TEXT("Can't connect resources with incompatible data domain types (%s vs %s)."),
+					*InOutput.GetDisplayName(), *InInput.GetDisplayName());
+			}
+			return false;
+		}
+	}
+
+	return true;
+}
+
+
diff --git a/Engine/Plugins/Animation/DeformerGraph/Source/OptimusCore/Private/OptimusDeformer.cpp b/Engine/Plugins/Animation/DeformerGraph/Source/OptimusCore/Private/OptimusDeformer.cpp
--- a/Engine/Plugins/Animation/DeformerGraph/Source/OptimusCore/Private/OptimusDeformer.cpp
+++ b/Engine/Plugins/Animation/DeformerGraph/Source/OptimusCore/Private/OptimusDeformer.cpp
@@ -1391,6 +1391,7 @@ bool UOptimusDeformer::Compile()
 	}
 
 	ComputeGraphs.Reset();
+	ConstantContainer.Reset();
 	
 	CompileBeginDelegate.Broadcast(this);
 	
@@ -1428,6 +1429,7 @@ bool UOptimusDeformer::Compile()
 	if (bCompilationFailed)
 	{
 		ComputeGraphs.Reset();
+		ConstantContainer.Reset();
 		return false;
 	}
 	
@@ -1636,6 +1638,7 @@ UOptimusComputeGraph* UOptimusDeformer::CompileNodeGraphToComputeGraph(
 							TransientBufferDI->ValueType = Pin->GetDataType()->ShaderValueType;
 							TransientBufferDI->DataDomain = Pin->GetDataDomain();
 							TransientBufferDI->ComponentSourceBinding = *ComponentSourceBindings.CreateConstIterator();
+							TransientBufferDI->DomainConstantIdentifier = {ConnectedNode.Node, NAME_None, Pin->GetFName()};
 							
 							LinkDataInterfaceMap.Add(Pin, TransientBufferDI);
 						}
@@ -1700,7 +1703,7 @@ UOptimusComputeGraph* UOptimusDeformer::CompileNodeGraphToComputeGraph(
 				ValueNodes,
 				GraphDataInterface, GraphDataComponentBinding,
 				KernelDataInterface,
-				BoundKernel.InputDataBindings, BoundKernel.OutputDataBindings
+				BoundKernel.InputDataBindings, BoundKernel.OutputDataBindings, ConstantContainer
 			);
 			if (FText* ErrorMessage = KernelSourceResult.TryGet<FText>())
 			{
diff --git a/Engine/Plugins/Animation/DeformerGraph/Source/OptimusCore/Private/OptimusDeformerInstance.cpp b/Engine/Plugins/Animation/DeformerGraph/Source/OptimusCore/Private/OptimusDeformerInstance.cpp
--- a/Engine/Plugins/Animation/DeformerGraph/Source/OptimusCore/Private/OptimusDeformerInstance.cpp
+++ b/Engine/Plugins/Animation/DeformerGraph/Source/OptimusCore/Private/OptimusDeformerInstance.cpp
@@ -301,14 +301,30 @@ void UOptimusDeformerInstance::SetupFromDeformer(UOptimusDeformer* InDeformer)
 		InstanceSettingsPtr->InitializeSettings(InDeformer, MeshComponent.Get());
 	}
 	InstanceSettingsPtr->GetComponentBindings(InDeformer, MeshComponent.Get(), BoundComponents);
+	
+	WeakBoundComponents.Reset();
+	for (UActorComponent* Component : BoundComponents)
+	{
+		WeakBoundComponents.Add(Component);
+	}
 
+	WeakComponentSources.Reset();
+	const TArray<UOptimusComponentSourceBinding*>& ComponentBindings = InDeformer->GetComponentBindings();
+	for (const UOptimusComponentSourceBinding* ComponentBinding : ComponentBindings)
+	{
+		WeakComponentSources.Add(ComponentBinding->GetComponentSource());
+	}
+	
 	// Create the persistent buffer pool
 	BufferPool = MakeShared<FOptimusPersistentBufferPool>();
 	
 	// (Re)Create and bind data providers.
 	ComputeGraphExecInfos.Reset();
 	GraphsToRunOnNextTick.Reset();
-
+	
+	ConstantContainer = InDeformer->ConstantContainer;
+	ConstantValuesPerContext.Reset();
+	
 	for (int32 GraphIndex = 0; GraphIndex < InDeformer->ComputeGraphs.Num(); ++GraphIndex)
 	{
 		FOptimusComputeGraphInfo const& ComputeGraphInfo = InDeformer->ComputeGraphs[GraphIndex];
@@ -349,9 +365,9 @@ void UOptimusDeformerInstance::SetupFromDeformer(UOptimusDeformer* InDeformer)
 			}
 
 			// Set this instance on the graph data provider so that it can query variables.
-			if (UOptimusGraphDataProvider* GraphProvider = Cast<UOptimusGraphDataProvider>(DataProvider))
+			if (IOptimusDeformerInstanceAccessor* InstanceAccessor = Cast<IOptimusDeformerInstanceAccessor>(DataProvider))
 			{
-				GraphProvider->DeformerInstance = this;
+				InstanceAccessor->SetDeformerInstance(this);
 			}
 		}
 
@@ -399,6 +415,97 @@ void UOptimusDeformerInstance::SetCanBeActive(bool bInCanBeActive)
 	bCanBeActive = bInCanBeActive;
 }
 
+TArray<float> UOptimusDeformerInstance::GetConstantValuePerInvocation(const FOptimusConstantIdentifier& InIdentifier)
+{
+	if (!InIdentifier.IsValid())
+	{
+		return {};
+	}
+	
+	FOptimusDeformerInstanceComponentContext ComponentContext;
+	
+	check(WeakBoundComponents.Num() == WeakComponentSources.Num());
+	int32 NumBindings = WeakBoundComponents.Num();
+	for (int32 BindingIndex = 0; BindingIndex < NumBindings; BindingIndex++)
+	{
+		int32 LodIndex = INDEX_NONE;
+		if (UActorComponent* Component = WeakBoundComponents[BindingIndex].Get())
+		{
+			if (const UOptimusComponentSource* ComponentSource = WeakComponentSources[BindingIndex].Get())
+			{
+				LodIndex = ComponentSource->GetLodIndex(Component);
+			}
+		}
+
+
+		if (ensure(LodIndex != INDEX_NONE))
+		{
+			ComponentContext.LodIndexPerComponent.Add(LodIndex);
+		}
+	}
+
+
+	if (const FOptimusConstantContainerInstance* ConstantValues = ConstantValuesPerContext.Find(ComponentContext))
+	{
+		return ConstantValues->GetConstantValuePerInvocation(InIdentifier);
+	}
+
+	TMap<int32, TMap<FName, TArray<float>>> BindingIndexToConstantValues;
+
+	auto CollectAllBindingConstants = [&]()
+	{
+		for (int32 BindingIndex = 0; BindingIndex < NumBindings; BindingIndex++)
+		{
+			TMap<FName, TArray<float>>& BindingConstantValues = BindingIndexToConstantValues.Add(BindingIndex);
+	
+			if (UActorComponent* Component = WeakBoundComponents[BindingIndex].Get())
+			{
+				if (const UOptimusComponentSource* ComponentSource = WeakComponentSources[BindingIndex].Get())
+				{
+					const int32 LodIndex = ComponentSource->GetLodIndex(Component);
+					TArray<FName> ComponentConstants = ComponentSource->GetExecutionDomains();
+
+					for (int32 ConstantIndex = 0; ConstantIndex < ComponentConstants.Num(); ConstantIndex++)
+					{
+						const FName& ConstantName = ComponentConstants[ConstantIndex];
+				
+						TArray<float>& Values = BindingConstantValues.Add(ConstantName);
+						
+						TArray<int32> Counts;
+						if (!ComponentSource->GetComponentElementCountsForExecutionDomain(
+							ConstantName,
+							Component, LodIndex, Counts))
+						{
+							return false;
+						}
+
+						Values.Reserve(Counts.Num());
+						for (const int32 Count : Counts)
+						{
+							Values.Add(Count);
+						}
+					}
+				}
+			}
+		}
+		
+		return true;
+	};
+
+
+	if (CollectAllBindingConstants())
+	{
+		FOptimusConstantContainerInstance& NewConstantValues = ConstantValuesPerContext.Add(ComponentContext);
+		if (NewConstantValues.Initialize(ConstantContainer, BindingIndexToConstantValues))
+		{
+			return NewConstantValues.GetConstantValuePerInvocation(InIdentifier);		
+		}
+	}
+
+	return{};
+	
+}
+
 void UOptimusDeformerInstance::AllocateResources()
 {
 }
diff --git a/Engine/Plugins/Animation/DeformerGraph/Source/OptimusCore/Private/OptimusExpressionEvaluator.cpp b/Engine/Plugins/Animation/DeformerGraph/Source/OptimusCore/Private/OptimusExpressionEvaluator.cpp
--- a/Engine/Plugins/Animation/DeformerGraph/Source/OptimusCore/Private/OptimusExpressionEvaluator.cpp
+++ b/Engine/Plugins/Animation/DeformerGraph/Source/OptimusCore/Private/OptimusExpressionEvaluator.cpp
@@ -1,84 +1,229 @@
 ﻿// Copyright Epic Games, Inc. All Rights Reserved.
 
-
 #include "OptimusExpressionEvaluator.h"
+#include <limits> // IWYU pragma: keep
 
-#include <limits>
+#include "OptimusCoreModule.h"
 
 namespace Optimus::Expression
 {
 
-FEngine::FEngine(
-	const TMap<FName, int32>& InConstants, 
-	EParseFlags InParseFlags
-	) :
-	Constants(InConstants),
-	ParseFlags(InParseFlags)
-{
-}
+#define CE_EXPR(_N_, _A_, _E_) { \
+	FEngine::FunctionNameIndex.Add(FName(#_N_), FEngine::Functions.Num()); \
+	FEngine::Functions.Add({_A_, [](TArrayView<const float> V) -> float { return _E_; } }); \
+	}  
 
+TArray<FEngine::FFunctionInfo> FEngine::Functions;
+TMap<FName, int32> FEngine::FunctionNameIndex;
 
-FEngine::FEngine(
-	TMap<FName, int32>&& InConstants,
-	EParseFlags InParseFlags
-	) :
-	Constants(MoveTemp(InConstants)),
-	ParseFlags(InParseFlags)
+struct FInitializeBuiltinFunctions
 {
-}
+	FInitializeBuiltinFunctions()
+	{
+		// clamp(value, min, max)
+		CE_EXPR(clamp, 3, FMath::Clamp(V[0], V[1], V[2]))
+		// min(a, b)
+		CE_EXPR(min, 2, FMath::Min(V[0], V[1]))
+		// max(a, b)
+		CE_EXPR(max, 2, FMath::Max(V[0], V[1]))
+		
+		// abs(value)
+		CE_EXPR(abs, 1, FMath::Abs(V[0]))
+		// round(value)
+		CE_EXPR(round, 1, FMath::RoundToFloat(V[0]))
+		// ceil(value)
+		CE_EXPR(ceil, 1, FMath::CeilToFloat(V[0]))
+		// floor(value)
+		CE_EXPR(floor, 1, FMath::FloorToFloat(V[0]))
+
+		// sin(value)
+		CE_EXPR(sin, 1, FMath::Sin(V[0]))
+		// cos(value)
+		CE_EXPR(cos, 1, FMath::Cos(V[0]))
+		// tan(value)
+		CE_EXPR(tan, 1, FMath::Tan(V[0]))
+		// asin(value)
+		CE_EXPR(asin, 1, FMath::Asin(V[0]))
+		// acos(value)
+		CE_EXPR(acos, 1, FMath::Acos(V[0]))
+		// atan(value)
+		CE_EXPR(atan, 1, FMath::Atan(V[0]))
+
+		// sqrt(value)
+		CE_EXPR(sqrt, 1, FMath::Sqrt(V[0]))
+		// isqrt(value)
+		CE_EXPR(isqrt, 1, FMath::InvSqrt(V[0]))
+
+		// log(value)
+		CE_EXPR(log, 1, FMath::Loge(V[0]))
+		// exp(value)
+		CE_EXPR(exp, 1, FMath::Exp(V[0]))
+	
+		/** pi() */
+		CE_EXPR(pi, 0, UE_PI)
+		
+		/** pi() */
+		CE_EXPR(e, 0, UE_EULERS_NUMBER)
+	}
+} GInitializeBuiltinFunctions;
+#undef CE_EXPR
 
 
-void FEngine::UpdateConstantValues(
-	const TMap<FName, int32>& InConstants
-	)
+FString FExpressionObject::ToString() const
 {
-	for (const TTuple<FName, int32>& NewConstantItem: InConstants)
+	TArray<FString> Items;
+	for (const OpElement& Element: Expression)
 	{
-		if (int32* OldConstant = Constants.Find(NewConstantItem.Key))
+		if (const EOperator* Op = Element.TryGet<EOperator>())
+		{
+			switch(*Op)
+			{
+			case EOperator::Negate: Items.Add(TEXT("Op[Negate]")); break;
+			case EOperator::Add: Items.Add(TEXT("Op[Add]")); break;
+			case EOperator::Subtract: Items.Add(TEXT("Op[Subtract]")); break; 
+			case EOperator::Multiply: Items.Add(TEXT("Op[Multiply]")); break;
+			case EOperator::Divide: Items.Add(TEXT("Op[Divide]")); break;
+			case EOperator::Modulo: Items.Add(TEXT("Op[Modulo]")); break;
+			case EOperator::Power: Items.Add(TEXT("Op[Power]")); break;
+			case EOperator::FloorDivide: Items.Add(TEXT("Op[FloorDivide]")); break;
+			}
+		}
+		else if (const FName* Constant = Element.TryGet<FName>())
+		{
+			Items.Add(FString::Printf(TEXT("C[%s]"), *Constant->ToString()));
+		}
+		else if (const FFunctionRef *FuncRef = Element.TryGet<FFunctionRef>())
 		{
-			*OldConstant = NewConstantItem.Value; 
+			Items.Add(FString::Printf(TEXT("F[%d]"), FuncRef->Index));
+		}
+		else if (const float *Value = Element.TryGet<float>())
+		{
+			Items.Add(FString::Printf(TEXT("V[%g]"), *Value));
 		}
 	}
+
+	return FString::Join(Items, TEXT("\n"));
 }
 
-// We can't use FChar::IsDigit, since it obeys locale, which would allow all sorts of non-ASCII
-// numeric digits from Unicode.
-static bool IsASCIIDigit(TCHAR InCh)
+
+TSet<FName> FExpressionObject::GetUsedConstants() const
 {
-	return InCh >= '0' && InCh <= '9';
+	TSet<FName> UsedConstants;
+	for (int32 Index = 0; Index < Expression.Num(); Index++)
+	{
+		if (const FName* ConstantName = Expression[Index].TryGet<FName>())
+		{
+			UsedConstants.Add(*ConstantName);
+		}
+	}
+	return UsedConstants;	
 }
 
-// Naive pow function for ints. Does not check for overflow, returns zero for negative exponents.
-static int32 IntPow(int32 InBase, int32 InExponent)
+
+void FExpressionObject::Load(FArchive& Ar)
 {
-	if (InExponent < 0)
+	int32 OperandCount = 0;
+	Ar << OperandCount;
+
+	Expression.Reset(OperandCount);
+
+	for (int32 OperandIndex = 0; OperandIndex < OperandCount; OperandIndex++)
 	{
-		return 0;
+		int32 OperandType = INDEX_NONE;
+		Ar << OperandType;
+		switch(OperandType)
+		{
+		case OpElement::IndexOfType<EOperator>():
+			{
+				int32 OperatorType = 0;
+				Ar << OperatorType;
+				Expression.Add(OpElement(TInPlaceType<EOperator>(), static_cast<EOperator>(OperatorType)));
+			}
+			break;
+		case OpElement::IndexOfType<FName>():
+			{
+				FName ConstantName;
+				Ar << ConstantName;
+				Expression.Add(OpElement(TInPlaceType<FName>(), ConstantName));
+			}
+			break;
+		case OpElement::IndexOfType<FFunctionRef>():
+			{
+				int32 FunctionIndex;
+				Ar << FunctionIndex;
+				Expression.Add(OpElement(TInPlaceType<FFunctionRef>(), FFunctionRef(FunctionIndex)));
+			}
+			break;
+		case OpElement::IndexOfType<float>():
+			{
+				float Value;
+				Ar << Value;
+				Expression.Add(OpElement(TInPlaceType<float>(), Value));
+			}
+			break;
+		default:
+			UE_LOG(LogOptimusCore, Error, TEXT("Invalid element type found in serialized expression"));
+			Ar.SetError();
+			return;
+		}
 	}
-	int32 Result = 1;
-	while (InExponent)
+}
+
+
+void FExpressionObject::Save(FArchive& Ar) const
+{
+	int32 OperandCount = Expression.Num(); 
+	Ar << OperandCount;
+
+	for (const OpElement& Operand: Expression)
 	{
-		if ((InExponent % 2) != 0)
-			Result *= InBase;
-		InExponent /= 2;
-		InBase *= InBase;
+		int32 OperandType = Operand.GetIndex();
+		Ar << OperandType;
+		
+		switch(OperandType)
+		{
+		case OpElement::IndexOfType<EOperator>():
+			{
+				int32 OperatorType = static_cast<int32>(Operand.Get<EOperator>());
+				Ar << OperatorType;
+			}
+			break;
+		case OpElement::IndexOfType<FName>():
+			{
+				FName ConstantName = Operand.Get<FName>();
+				Ar << ConstantName;
+			}
+			break;
+		case OpElement::IndexOfType<FFunctionRef>():
+			{
+				int32 FunctionIndex = Operand.Get<FFunctionRef>().Index;
+				Ar << FunctionIndex;
+			}
+			break;
+		case OpElement::IndexOfType<float>():
+			{
+				float Value = Operand.Get<float>();
+				Ar << Value;
+			}
+			break;
+		default:
+			checkNoEntry();
+			return;
+		}
 	}
-	return Result;
 }
 
-	
-TVariant<FEngine::FToken, FParseError> FEngine::ParseTokenError(
-	FString&& InError,
-	const TCHAR* InCurrentPos,
-	FStringView InExpression)
+
+// We can't use FChar::IsDigit, since it obeys locale, which would allow all sorts of non-ASCII
+// numeric digits from Unicode.
+static bool IsASCIIDigit(const TCHAR InCh)
 {
-	return TVariant<FToken, FParseError>(TInPlaceType<FParseError>(),
-		static_cast<int32>(InCurrentPos - InExpression.GetData()), MoveTemp(InError));
+	return InCh >= '0' && InCh <= '9';
 }
 
 
 FString FEngine::TokenToString(
-	const FEngine::FToken& InToken
+	const FToken& InToken
 	)
 {
 	const TCHAR* TypeStr = TEXT("<unknown>");
@@ -96,9 +241,14 @@ FString FEngine::TokenToString(
 		case EOperatorToken::Divide:		ValueStr = TEXT("/"); break;
 		case EOperatorToken::Modulo:		ValueStr = TEXT("%"); break;
 		case EOperatorToken::Power:			ValueStr = TEXT("**"); break;
-		// case EOperatorToken::FloorDivide:	ValueStr = TEXT("//"); break;
+		case EOperatorToken::FloorDivide:	ValueStr = TEXT("//"); break;
 		case EOperatorToken::ParenOpen:		ValueStr = TEXT("("); break;
 		case EOperatorToken::ParenClose:	ValueStr = TEXT(")"); break;
+		case EOperatorToken::Comma:			ValueStr = TEXT(","); break;
+		
+		case EOperatorToken::Max:
+			checkNoEntry();
+			break;
 		}
 	}
 	else if (const FName* Identifier = InToken.TryGet<FName>())
@@ -106,65 +256,204 @@ FString FEngine::TokenToString(
 		TypeStr = TEXT("Identifier");
 		ValueStr = Identifier->ToString();
 	}
-	else if (const int32* Value = InToken.TryGet<int32>())
+	else if (const float* Value = InToken.TryGet<float>())
 	{
 		TypeStr = TEXT("Value");
-		ValueStr = FString::Printf(TEXT("%d"), *Value);
+		ValueStr = FString::Printf(TEXT("%g"), *Value);
 	}
 
 	return FString::Printf(TEXT("%s[%s]"), TypeStr, *ValueStr);
 }
 
 
-TVariant<FEngine::FToken, FParseError> FEngine::ParseInt32(
+auto FEngine::ParseFloat(
 	FStringView& InOutParseRange,
 	FStringView InExpression
-	)
+	) -> FTokenParseResult
 {
-	int64 Value = 0;
-	constexpr int32 MaxDigits = std::numeric_limits<int32>::digits10;
+	// We only parse at most 11 significant digits, since 9 is the maximum decimal precision 
+	// a float will hold. We store 2 extra digits to make sure we round properly on the 
+	// last digit bit.
+	constexpr int32 MaxPrecision = 11;
+
+	// Count the number of digits in the mantissa and, optionally, locate the decimal point
+	// and the start of significant digits.
+	int32 DecimalPointPos = -1;
+	int32 MantissaLen = 0;
+	int32 SignificantDigitStart = -1;
+	for (; MantissaLen < InOutParseRange.Len(); MantissaLen++)
+	{
+		const TCHAR Ch = InOutParseRange[MantissaLen];
+		if (!IsASCIIDigit(Ch))
+		{
+			// If the current character is neither a digit or period, or if we've already
+			// found a period, stop counting.
+			if (Ch != '.' || DecimalPointPos >= 0)
+			{
+				break;
+			}
+			DecimalPointPos = MantissaLen;
+		}
+		else if (Ch != '0' && SignificantDigitStart < 0)
+		{
+			SignificantDigitStart = MantissaLen;
+		}
+	}
+
+	// Remember the position for the exponent parsing.
+	int32 LastPos = MantissaLen;
+	int32 SignificantDigitCount = SignificantDigitStart >= 0 ? MantissaLen - SignificantDigitStart : 0;
 
-	// Skip over leading zeroes.
-	int32 DigitStart = 0;
-	for (; DigitStart < InOutParseRange.Len() && InOutParseRange[DigitStart] == '0'; DigitStart++)
+	// If no decimal point was given, set the position of it after the mantissa. Otherwise
+	// remove one from the mantissa length, since that would be the decimal point.
+	if (DecimalPointPos < 0)
 	{
-		/**/
+		DecimalPointPos = MantissaLen;
 	}
+	else
+	{
+		MantissaLen--;		
 
-	int32 DigitIndex = DigitStart; 
-	for(; DigitIndex < InOutParseRange.Len(); DigitIndex++)
+		// If the significant digit starts before the decimal point, remove one to account
+		// for the decimal point.
+		if (SignificantDigitStart < DecimalPointPos)
+		{
+			SignificantDigitCount--;
+		}
+	}
+	
+	if (MantissaLen == 0)
 	{
-		if ((DigitIndex - DigitStart) > MaxDigits)
+		// We only parsed the period. That's not a valid float.
+		return FTokenParseResult(TEXT("Invalid floating point value"),
+			InOutParseRange.Left(1), 
+		InExpression.GetData());
+	}
+
+	// Adjustment for the exponent based on number of significant digits in the fraction part.
+	int32 FractionExponent;
+	if (SignificantDigitCount > MaxPrecision)
+	{
+		FractionExponent = DecimalPointPos - MaxPrecision;
+		SignificantDigitCount = MaxPrecision;
+	}
+	else
+	{
+		FractionExponent = DecimalPointPos - MantissaLen;
+	}
+
+	// Copy the integer and fraction digits, up to the maximum count.
+	uint64 IntegerValue = 0;
+	const TCHAR *Digits = InOutParseRange.GetData() + SignificantDigitStart;
+	for (int32 DigitIndex = 0; DigitIndex < SignificantDigitCount; DigitIndex++)
+	{
+		TCHAR Ch = *Digits++;
+		if (Ch == '.')
 		{
-			return ParseTokenError(TEXT("Integer value too large"), InOutParseRange.GetData(), InExpression);
+			Ch = *Digits++;
 		}
+		IntegerValue = 10 * IntegerValue + (Ch - '0');
+	}
+	double Fraction = static_cast<double>(IntegerValue);
 
-		TCHAR Ch = InOutParseRange[DigitIndex];
-		if (!IsASCIIDigit(Ch))
+	// Now that we have the fraction part, parse the exponent, if any.
+	int32 Exponent = 0;
+	bool IsExponentNegative = false;
+
+	// We need at least _two_ characters here ('e' + digit)
+	if ((LastPos + 1) < InOutParseRange.Len() &&
+		(InOutParseRange[LastPos] == 'e' || InOutParseRange[LastPos] == 'E'))
+	{
+		int32 ExponentPos = LastPos + 1;
+		if (InOutParseRange[ExponentPos] == '+')
 		{
-			break;
+			ExponentPos++;
+		}
+		else if (InOutParseRange[ExponentPos] == '-')
+		{
+			IsExponentNegative = true;
+			ExponentPos++;
 		}
 
-		Value *= 10;
-		Value += Ch - '0';
+		int32 DigitCount = 0;
+		while(ExponentPos < InOutParseRange.Len() && IsASCIIDigit(InOutParseRange[ExponentPos]))
+		{
+			// We don't care past 2 digits. 38 is the maximum exponent we care about for floats.
+			if (DigitCount < 2)
+			{
+				Exponent = Exponent * 10 + (InOutParseRange[ExponentPos] - '0');
+			}
+			ExponentPos++;
+			DigitCount++;
+		}
+
+		// Only if we parsed some digits do we consider the exponent parsed at all. Otherwise
+		// we leave the last parse at where the 'e' started.
+		if (DigitCount > 0)
+		{
+			LastPos = ExponentPos;
+		}
+	}
+
+	// Adjust the exponent to deal with the decimal point location.
+	if (IsExponentNegative)
+	{
+		Exponent = FractionExponent - Exponent;
+	}
+	else
+	{
+		Exponent = FractionExponent + Exponent;
+	}
+
+	// The maximum exponent for floats is 38. 
+	if (Exponent < 0) 
+	{
+		IsExponentNegative = true;
+		Exponent = -Exponent;
+	} 
+	else 
+	{
+		IsExponentNegative = false;
+	}
+
+	Exponent = FMath::Min(Exponent, std::numeric_limits<float>::max_exponent10);
+	double DoubleExponent = 1.0;
+
+	for(const double PowerOf2Exp: {1.0e1, 1.0e2, 1.0e4, 1.0e8, 1.0e16, 1.0e32})
+	{
+		if (Exponent == 0)
+		{
+			break;
+		}
+		if (Exponent & 1)
+		{
+			DoubleExponent *= PowerOf2Exp;
+		}
+		Exponent >>= 1;
 	}
 
-	if (Value > std::numeric_limits<int32>::max())
+	if (IsExponentNegative)
 	{
-		return ParseTokenError(TEXT("Integer value too large"), InOutParseRange.GetData(), InExpression);
+		Fraction /= DoubleExponent;
+	}
+	else
+	{
+		Fraction *= DoubleExponent;
 	}
 
-	InOutParseRange = InOutParseRange.Mid(DigitIndex);
+	const FStringView TokenParseRange(InOutParseRange.Left(LastPos));
+	
+	// Update the parse range for the next token. 
+	InOutParseRange = InOutParseRange.Mid(LastPos);
 	
-	return TVariant<FToken, FParseError>(
-		TInPlaceType<FToken>(), FToken(TInPlaceType<int32>(), Value));
+	return FTokenParseResult(FToken(TInPlaceType<float>(), Fraction), TokenParseRange, InExpression.GetData());
 }
 
 
-TVariant<FEngine::FToken, FParseError> FEngine::ParseIdentifier(
+auto FEngine::ParseIdentifier(
 	FStringView& InOutParseRange,
 	FStringView InExpression
-	) const
+	) -> FTokenParseResult 
 {
 	// FIXME: We need quote escaping.
 	bool bStopAtSingleQuote = false;
@@ -175,13 +464,13 @@ TVariant<FEngine::FToken, FParseError> FEngine::ParseIdentifier(
 		ParseIndex++;
 		if (ParseIndex >= InOutParseRange.Len())
 		{
-			return ParseTokenError(TEXT("Missing end quote for constant"), InOutParseRange.GetData(), InExpression);
+			return FTokenParseResult(TEXT("Missing end quote for constant"), InOutParseRange, InExpression.GetData());
 		}
 	}
 	// Unquoted constants have to start with an alphabetic character.
 	else if (!FChar::IsAlpha(InOutParseRange[0]) && InOutParseRange[0] != '_')
 	{
-		return ParseTokenError(TEXT("Unexpected character"), InOutParseRange.GetData(), InExpression);
+		return FTokenParseResult(TEXT("Unexpected character"), InOutParseRange.Left(1), InExpression.GetData());
 	}
 
 	const int32 ConstantStart = ParseIndex;
@@ -198,7 +487,7 @@ TVariant<FEngine::FToken, FParseError> FEngine::ParseIdentifier(
 
 		if (ParseIndex == InOutParseRange.Len())
 		{
-			return ParseTokenError(TEXT("Missing end quote for constant"), InOutParseRange.GetData(), InExpression);
+			return FTokenParseResult(TEXT("Missing end quote for constant"), InOutParseRange, InExpression.GetData());
 		}
 
 		ConstantEnd = ParseIndex;
@@ -215,89 +504,66 @@ TVariant<FEngine::FToken, FParseError> FEngine::ParseIdentifier(
 		ConstantEnd = ParseIndex;
 	}
 
-	// Look up this identifier. We don't want to create a new FName, to avoid polluting the
-	// FNameEntry cache.
 	const FStringView Identifier(InOutParseRange.SubStr(ConstantStart, ConstantEnd - ConstantStart));
-	const FName Name(Identifier.Len(), Identifier.GetData(), FNAME_Find);
+	const FName Name(Identifier.Len(), Identifier.GetData());
 	
-	InOutParseRange = InOutParseRange.Mid(ParseIndex); 
+	const FStringView TokenParseRange(InOutParseRange.Left(ParseIndex));
+	InOutParseRange = InOutParseRange.Mid(ParseIndex);
 
-	if (Constants.Contains(Name))
-	{
-		return TVariant<FToken, FParseError>(
-			TInPlaceType<FToken>(), FToken(TInPlaceType<FName>(), Name));
-	}
-	else if (EnumHasAllFlags(ParseFlags, EParseFlags::ValidateConstants))
-	{
-		return ParseTokenError(FString::Printf(TEXT("Unknown identifier '%s'"), *FString(Identifier)),
-			InOutParseRange.GetData(), InExpression); 
-	}
-	else
-	{
-		// If the constant value is not found, it gets a default value of 0.
-		return TVariant<FToken, FParseError>(
-			TInPlaceType<FToken>(), FToken(TInPlaceType<int32>(), 0));
-	}
+	return FTokenParseResult(FToken(TInPlaceType<FName>(), Name), TokenParseRange, InExpression.GetData());
 }
 
 
-TVariant<FEngine::FToken, FParseError> FEngine::ParseToken(
+auto FEngine::ParseToken(
 	FStringView& InOutParseRange,
 	FStringView InExpression
-	) const
+	) -> FTokenParseResult
 {
-	auto TokenResult = [](EOperatorToken InTokenKind)
+	auto TokenResult = [&InOutParseRange, Expression=InExpression.GetData()](EOperatorToken InTokenKind, int32 InTokenLength)
 	{
-		return TVariant<FToken, FParseError>(
-			TInPlaceType<FToken>(), FToken(TInPlaceType<EOperatorToken>(), InTokenKind));
+		const FStringView TokenRange(InOutParseRange.Left(InTokenLength));
+		InOutParseRange = InOutParseRange.Mid(InTokenLength); 
+		return FTokenParseResult(FToken(TInPlaceType<EOperatorToken>(), InTokenKind), TokenRange, Expression);
 	};
 
 	switch(InOutParseRange[0])
 	{
 	case '+':
-		InOutParseRange = InOutParseRange.Mid(1);
-		return TokenResult(EOperatorToken::Add);
+		return TokenResult(EOperatorToken::Add, 1);
 		
 	case '-':
-		InOutParseRange = InOutParseRange.Mid(1);
-		return TokenResult(EOperatorToken::Subtract);
+		return TokenResult(EOperatorToken::Subtract, 1);
 		
 	case '*':
 		if (InOutParseRange.Len() >= 2 && InOutParseRange[1] == '*')
 		{
-			InOutParseRange = InOutParseRange.Mid(2);
-			return TokenResult(EOperatorToken::Power);
+			return TokenResult(EOperatorToken::Power, 2);
 		}
-		InOutParseRange = InOutParseRange.Mid(1);
-		return TokenResult(EOperatorToken::Multiply);
+		return TokenResult(EOperatorToken::Multiply, 1);
 		
 	case '/':
-		/*
 		if (InOutParseRange.Len() >= 2 && InOutParseRange[1] == '/')
 		{
-			InOutParseRange = InOutParseRange.Mid(2);
-			return TokenResult(EOperatorToken::FloorDivide);
+			return TokenResult(EOperatorToken::FloorDivide, 2);
 		}
-		*/
-		InOutParseRange = InOutParseRange.Mid(1);
-		return TokenResult(EOperatorToken::Divide);
+		return TokenResult(EOperatorToken::Divide, 1);
 		
 	case '%':
-		InOutParseRange = InOutParseRange.Mid(1);
-		return TokenResult(EOperatorToken::Modulo);
+		return TokenResult(EOperatorToken::Modulo, 1);
 
 	case '(':
-		InOutParseRange = InOutParseRange.Mid(1);
-		return TokenResult(EOperatorToken::ParenOpen);
+		return TokenResult(EOperatorToken::ParenOpen, 1);
 		
 	case ')':
-		InOutParseRange = InOutParseRange.Mid(1);
-		return TokenResult(EOperatorToken::ParenClose);
+		return TokenResult(EOperatorToken::ParenClose, 1);
+
+	case ',':
+		return TokenResult(EOperatorToken::Comma, 1);
 		
 	default:
-		if (IsASCIIDigit(InOutParseRange[0]))
+		if (IsASCIIDigit(InOutParseRange[0]) || InOutParseRange[0] == '.')
 		{
-			return ParseInt32(InOutParseRange, InExpression);
+			return ParseFloat(InOutParseRange, InExpression);
 		}
 		return ParseIdentifier(InOutParseRange, InExpression);
 	}
@@ -305,7 +571,8 @@ TVariant<FEngine::FToken, FParseError> FEngine::ParseToken(
 
 
 TVariant<FExpressionObject, FParseError> FEngine::Parse(
-	FStringView InExpression
+	FStringView InExpression,
+	TFunction<TOptional<float>(FName InConstantName)> InConstantEvaluator
 	) const
 {
 	enum class EAssociativity
@@ -332,10 +599,12 @@ TVariant<FExpressionObject, FParseError> FEngine::Parse(
 			{ EOperatorToken::Divide,		2, EAssociativity::Left, FExpressionObject::EOperator::Divide},
 			{ EOperatorToken::Modulo,		3, EAssociativity::Left, FExpressionObject::EOperator::Modulo},
 			{ EOperatorToken::Power,		4, EAssociativity::Right, FExpressionObject::EOperator::Power},
-			// { EOperatorToken::FloorDivide,	3, EAssociativity::Left, FExpressionObject::EOperator::FloorDivide},
+			{ EOperatorToken::FloorDivide,	3, EAssociativity::Left, FExpressionObject::EOperator::FloorDivide},
 			{ EOperatorToken::ParenOpen,	0, EAssociativity::None},
 			{ EOperatorToken::ParenClose,	0, EAssociativity::None},
+			{ EOperatorToken::Comma,		0, EAssociativity::None},
 		};
+		static_assert(sizeof(OperatorTokenInfo) / sizeof(FOperatorTokenInfo) == (int32)EOperatorToken::Max);
 		checkSlow(OperatorTokenInfo[static_cast<int32>(InOpToken)].OperatorToken == InOpToken);
 		return OperatorTokenInfo[static_cast<int32>(InOpToken)];
 	};
@@ -349,20 +618,30 @@ TVariant<FExpressionObject, FParseError> FEngine::Parse(
 	}
 
 	// Run Dijkstra's classic Shunting Yard algorithm to convert infix expressions to RPN.
-	// TODO: Add support for calling functions.
 	// TODO: Better error pinpointing by storing the token parse result with the token.
 	TArray<FExpressionObject::OpElement, TInlineAllocator<64>> Expression;
-	TArray<EOperatorToken, TInlineAllocator<32>> OperatorStack;
+	TArray<FParseLocation, TInlineAllocator<64>> Locations;
+	TArray<TTuple<EOperatorToken, FParseLocation>, TInlineAllocator<32>> OperatorStack;
+	
+	struct FFunctionCallInfo
+	{
+		int32 FunctionIndex = INDEX_NONE;
+		int32 CountedCommas = 0;
+		int32 OpeningOpStackSize = 0;
+		int32 ExpressionSize = 0;
+		FParseLocation Location{};
+	};
+	TArray<FFunctionCallInfo, TInlineAllocator<32>> FunctionStack;
 
-	auto ParseError = [InExpression](FString&& InError, const TCHAR* InCurrentPos)
+	auto ParseError = [](FString&& InError, FParseLocation InLocation)
 	{
-		return TVariant<FExpressionObject, FParseError>(TInPlaceType<FParseError>(),
-			static_cast<int32>(InCurrentPos - InExpression.GetData()), MoveTemp(InError));
+		return TVariant<FExpressionObject, FParseError>(TInPlaceType<FParseError>(), MoveTemp(InError), InLocation);
 	};
-	auto PushOperatorTokenToExpression = [GetOperatorTokenInfo, &Expression](EOperatorToken InOpToken)
+	auto PushOperatorTokenToExpression = [GetOperatorTokenInfo, &Expression, &Locations](TTuple<EOperatorToken, FParseLocation> InTokenAndLocation)
 	{
 		Expression.Push(FExpressionObject::OpElement(TInPlaceType<FExpressionObject::EOperator>(), 
-			*GetOperatorTokenInfo(InOpToken).Operator));
+			*GetOperatorTokenInfo(InTokenAndLocation.Get<0>()).Operator));
+		Locations.Push(InTokenAndLocation.Get<1>());
 	};
 	// Used to check if we encounter a '-' to see if we need to turn it into a negate operator
 	// and also to see if we have two operators in a row (error).
@@ -370,35 +649,36 @@ TVariant<FExpressionObject, FParseError> FEngine::Parse(
 
 	do
 	{
-		const TCHAR* TokenStart = ParseRange.GetData();
-		TVariant<FToken, FParseError> ParseResult = ParseToken(ParseRange, InExpression);
+		FTokenParseResult ParseResult = ParseToken(ParseRange, InExpression);
 		
-		if (FParseError* Error = ParseResult.TryGet<FParseError>())
+		if (ParseResult.IsError())
 		{
-			return TVariant<FExpressionObject, FParseError>(TInPlaceType<FParseError>(), *Error);
+			return TVariant<FExpressionObject, FParseError>(TInPlaceType<FParseError>(), ParseResult.GetParseError());
 		}
 
-		const FStringView TokenParseRange(TokenStart, ParseRange.GetData() - TokenStart);
-		const FToken& Token = ParseResult.Get<FToken>();
+		const FToken& Token = ParseResult.GetToken();
+		
 
 		// Values / identifiers go onto the expression stack immediately.
 		if (const FName* Identifier = Token.TryGet<FName>())
 		{
 			if (LastToken.IsSet() &&
-				(LastToken->IsType<FName>() || LastToken->IsType<int32>()))
+				(LastToken->IsType<FName>() || LastToken->IsType<float>()))
 			{
-				return ParseError(TEXT("Expected an operator"), TokenParseRange.GetData());
+				return ParseError(TEXT("Expected an operator"), ParseResult.Location);
 			}
 			Expression.Push(FExpressionObject::OpElement(TInPlaceType<FName>(), *Identifier));
+			Locations.Push(ParseResult.Location);
 		}
-		else if (const int32* Value = Token.TryGet<int32>())
+		else if (const float* Value = Token.TryGet<float>())
 		{
 			if (LastToken.IsSet() &&
-				(LastToken->IsType<FName>() || LastToken->IsType<int32>()))
+				(LastToken->IsType<FName>() || LastToken->IsType<float>()))
 			{
-				return ParseError(TEXT("Expected an operator"), TokenParseRange.GetData());
+				return ParseError(TEXT("Expected an operator"), ParseResult.Location);
 			}
-			Expression.Push(FExpressionObject::OpElement(TInPlaceType<int32>(), *Value));
+			Expression.Push(FExpressionObject::OpElement(TInPlaceType<float>(), *Value));
+			Locations.Push(ParseResult.Location);
 		}
 		else
 		{
@@ -413,15 +693,16 @@ TVariant<FExpressionObject, FParseError> FEngine::Parse(
 					const FOperatorTokenInfo& LastTokenInfo = GetOperatorTokenInfo(*LastOp);
 					if (TokenInfo.Operator.IsSet() && LastTokenInfo.Operator.IsSet())
 					{
-						return ParseError(TEXT("Expected an expression"), TokenParseRange.GetData());
+						return ParseError(TEXT("Expected an expression"), ParseResult.Location);
 					}
 					if (*LastOp == EOperatorToken::ParenClose && Op == EOperatorToken::ParenOpen)
 					{
-						return ParseError(TEXT("Expected an operator"), TokenParseRange.GetData());
+						return ParseError(TEXT("Expected an operator"), ParseResult.Location);
 					}
-					if (*LastOp == EOperatorToken::ParenOpen && Op == EOperatorToken::ParenClose)
+					if (*LastOp == EOperatorToken::ParenOpen && Op == EOperatorToken::ParenClose &&
+					    (FunctionStack.IsEmpty() || FunctionStack.Top().OpeningOpStackSize != OperatorStack.Num()))
 					{
-						return ParseError(TEXT("Empty Parentheses"), TokenParseRange.GetData());
+						return ParseError(TEXT("Empty Parentheses"), ParseResult.Location);
 					}
 				}
 			}
@@ -431,42 +712,84 @@ TVariant<FExpressionObject, FParseError> FEngine::Parse(
 				(!LastToken.IsSet() || (LastToken->IsType<EOperatorToken>() && LastToken->Get<EOperatorToken>() != EOperatorToken::ParenClose)))
 			{
 				Op = EOperatorToken::Negate;
-				OperatorStack.Push(Op);
+				OperatorStack.Push({Op, ParseResult.Location});
 			}
 			else if (Op == EOperatorToken::ParenOpen)
 			{
-				OperatorStack.Push(Op);
+				OperatorStack.Push({Op, ParseResult.Location});
+
+				// If the top element on the expression is a constant, check if it is a function name and if so, 
+				// pop the constant out of the expression. 
+				if (!Expression.IsEmpty() && Expression.Top().IsType<FName>())
+				{
+					FName ConstantName = Expression.Top().Get<FName>();
+					const int32 *FunctionIndexPtr = FunctionNameIndex.Find(ConstantName);
+					if (FunctionIndexPtr == nullptr)
+					{
+						return ParseError(FString::Printf(TEXT("Unknown function '%s'"), *ConstantName.ToString()), Locations.Top());
+					}
+
+					Expression.Pop();
+					const FParseLocation LastLocation = Locations.Pop();
+					FunctionStack.Push({*FunctionIndexPtr, 0, OperatorStack.Num(), Expression.Num(), LastLocation});
+				}
 			}
 			else if (Op == EOperatorToken::ParenClose)
 			{
-				while(OperatorStack.IsEmpty() || OperatorStack.Last() != EOperatorToken::ParenOpen)
+				while(OperatorStack.IsEmpty() || OperatorStack.Last().Key != EOperatorToken::ParenOpen)
 				{
 					if (OperatorStack.IsEmpty())
 					{
-						return ParseError(TEXT("Mismatched parentheses"), TokenParseRange.GetData()); 
+						return ParseError(TEXT("Mismatched parentheses"), ParseResult.Location); 
 					}
 					PushOperatorTokenToExpression(OperatorStack.Pop(false));
 				}
+
+				if (!FunctionStack.IsEmpty() && FunctionStack.Top().OpeningOpStackSize == OperatorStack.Num())
+				{
+					// Compute the number of arguments.
+					const FFunctionCallInfo CallInfo = FunctionStack.Pop();
+					const FFunctionInfo& FunctionInfo = Functions[CallInfo.FunctionIndex];
+					const int32 ArgumentCount = CallInfo.CountedCommas + (CallInfo.ExpressionSize != Expression.Num());
+
+					if (ArgumentCount != FunctionInfo.ArgumentCount)
+					{
+						return ParseError(FString::Printf(TEXT("Invalid argument count. Expected %d, got %d"), FunctionInfo.ArgumentCount, ArgumentCount),
+							CallInfo.Location.Join(ParseResult.Location)); 
+					}
+
+					Expression.Push(FExpressionObject::OpElement(TInPlaceType<FExpressionObject::FFunctionRef>(), CallInfo.FunctionIndex));
+					Locations.Push(ParseResult.Location);
+				}
+				
 				// Remove the now unneeded open parentheses.
 				OperatorStack.Pop(false);
 			}
+			else if (Op == EOperatorToken::Comma)
+			{
+				if (FunctionStack.IsEmpty() || FunctionStack.Top().OpeningOpStackSize != OperatorStack.Num())
+				{
+					return ParseError(TEXT("Unexpected comma"), ParseResult.Location); 
+				}
+				FunctionStack.Top().CountedCommas++;
+			}
 			else if (GetOperatorTokenInfo(Op).Associativity == EAssociativity::Left)
 			{
 				// Peel off all operators that have the same precedence or higher.
-				while (!OperatorStack.IsEmpty() && GetOperatorTokenInfo(OperatorStack.Last()).Precedence >= TokenInfo.Precedence)
+				while (!OperatorStack.IsEmpty() && GetOperatorTokenInfo(OperatorStack.Last().Get<0>()).Precedence >= TokenInfo.Precedence)
 				{
 					PushOperatorTokenToExpression(OperatorStack.Pop(false));
 				}
-				OperatorStack.Push(Op);
+				OperatorStack.Push({Op, ParseResult.Location});
 			}
 			else if (GetOperatorTokenInfo(Op).Associativity == EAssociativity::Right)
 			{
 				// Peel off all operators that have higher precedence.
-				while (!OperatorStack.IsEmpty() && GetOperatorTokenInfo(OperatorStack.Last()).Precedence > TokenInfo.Precedence)
+				while (!OperatorStack.IsEmpty() && GetOperatorTokenInfo(OperatorStack.Last().Get<0>()).Precedence > TokenInfo.Precedence)
 				{
 					PushOperatorTokenToExpression(OperatorStack.Pop(false));
 				}
-				OperatorStack.Push(Op);
+				OperatorStack.Push({Op, ParseResult.Location});
 			}
 		}
 
@@ -476,36 +799,61 @@ TVariant<FExpressionObject, FParseError> FEngine::Parse(
 
 	while (!OperatorStack.IsEmpty())
 	{
-		if (!GetOperatorTokenInfo(OperatorStack.Last()).Operator.IsSet())
+		if (!GetOperatorTokenInfo(OperatorStack.Last().Get<0>()).Operator.IsSet())
 		{
-			return ParseError(TEXT("Mismatched parentheses"), InExpression.GetData()); 
+			return ParseError(TEXT("Mismatched parentheses"), OperatorStack.Last().Get<1>()); 
 		}
 		PushOperatorTokenToExpression(OperatorStack.Pop(false));
 	}
 
+	// If a constant evaluator function is given, then check if the constants exist.  
+	if (InConstantEvaluator)
+	{
+		for (int32 Index = 0; Index < Expression.Num(); Index++)
+		{
+			if (const FName* ConstantName = Expression[Index].TryGet<FName>())
+			{
+				if (!InConstantEvaluator(*ConstantName).IsSet())
+				{
+					return ParseError(*FString::Printf(TEXT("Unknown identifier '%s'"), *ConstantName->ToString()), Locations[Index]); 
+				}
+			}
+		}
+	}
+	
 	FExpressionObject S;
 	S.Expression = Expression;
 	return TVariant<FExpressionObject, FParseError>(TInPlaceType<FExpressionObject>(), S);
 }
 
 
-int32 FEngine::Execute(
-	const FExpressionObject& InExpressionObject
+float FEngine::Execute(
+	const FExpressionObject& InExpressionObject,
+	TFunctionRef<TOptional<float>(FName InConstantName)> InConstantEvaluator
 	) const
 {
 	// An empty expression object is valid.
 	if (InExpressionObject.Expression.IsEmpty())
 	{
-		return 0;
+		return 0.0f;
 	}
 
 	// We assume the codons in the expression object are correct and will always result
 	// in a single value remaining on the stack.
-	TArray<int32, TInlineAllocator<32>> ValueStack;
+	TArray<float, TInlineAllocator<32>> ValueStack;
 	
 	for (const FExpressionObject::OpElement& Token: InExpressionObject.Expression)
 	{
-		if (const FExpressionObject::EOperator* Operator = Token.TryGet<FExpressionObject::EOperator>())
+		if (const float* Value = Token.TryGet<float>())
+		{
+			ValueStack.Push(*Value);
+		}
+		else if (const FName* ConstantName = Token.TryGet<FName>())
+		{
+			TOptional<float> ConstantValue = InConstantEvaluator(*ConstantName);
+			ValueStack.Push(ConstantValue.IsSet() ? ConstantValue.GetValue() : 0.0f);
+		}
+		else if (const FExpressionObject::EOperator* Operator = Token.TryGet<FExpressionObject::EOperator>())
 		{
 			switch(*Operator)
 			{
@@ -517,28 +865,28 @@ int32 FEngine::Execute(
 				
 			case FExpressionObject::EOperator::Add:
 				{
-					const int32 V = ValueStack.Pop(false);
+					const float V = ValueStack.Pop(false);
 					ValueStack.Last() += V;
 				}
 				break;	
 			case FExpressionObject::EOperator::Subtract:
 				{
-					const int32 V = ValueStack.Pop(false);
+					const float V = ValueStack.Pop(false);
 					ValueStack.Last() -= V;
 				}
 				break;	
 			case FExpressionObject::EOperator::Multiply:
 				{
-					const int32 V = ValueStack.Pop(false);
+					const float V = ValueStack.Pop(false);
 					ValueStack.Last() *= V;
 				}
 				break;
 			case FExpressionObject::EOperator::Divide:
 				{
-					const int32 V = ValueStack.Pop(false);
-					if (V == 0)
+					const float V = ValueStack.Pop(false);
+					if (FMath::IsNearlyZero(V))
 					{
-						ValueStack.Last() = 0;
+						ValueStack.Last() = 0.0f;
 					}
 					else
 					{
@@ -548,33 +896,55 @@ int32 FEngine::Execute(
 				break;	
 			case FExpressionObject::EOperator::Modulo:
 				{
-					const int32 V = ValueStack.Pop(false);
-					if (V == 0)
+					const float V = ValueStack.Pop(false);
+					if (FMath::IsNearlyZero(V))
 					{
-						ValueStack.Last() = 0;
+						ValueStack.Last() = 0.0f;
 					}
 					else
 					{
-						ValueStack.Last() = ValueStack.Last() % V;
+						ValueStack.Last() = FMath::Fmod(ValueStack.Last(), V);
 					}
 				}
 				break;	
 			case FExpressionObject::EOperator::Power:
 				{
-					const int32 V = ValueStack.Pop(false);
-					ValueStack.Last() = IntPow(ValueStack.Last(), V);
+					const float V = ValueStack.Pop(false);
+					ValueStack.Last() = FMath::Pow(ValueStack.Last(), V);
+					if (!FMath::IsFinite(ValueStack.Last()))
+					{
+						ValueStack.Last() = 0.0f;
+					}
 				}
 				break;	
+			case FExpressionObject::EOperator::FloorDivide:
+				{
+					const float V = ValueStack.Pop(false);
+					if (FMath::IsNearlyZero(V))
+					{
+						ValueStack.Last() = 0.0f;
+					}
+					else
+					{
+						ValueStack.Last() /= V;
+					}
+					ValueStack.Last() = FMath::Floor(ValueStack.Last());
+				}
+				break;  
 			}
 		}
-		else if (const FName* Constant = Token.TryGet<FName>())
+		else if (const FExpressionObject::FFunctionRef* FuncRef = Token.TryGet<FExpressionObject::FFunctionRef>())
 		{
-			const int32 *Value = Constants.Find(*Constant);
-			ValueStack.Push(Value ? *Value : 0);
-		}
-		else if (const int32* Value = Token.TryGet<int32>())
-		{
-			ValueStack.Push(*Value);
+			const FFunctionInfo& FunctionInfo = Functions[FuncRef->Index];
+			check(FunctionInfo.ArgumentCount <= ValueStack.Num());
+
+			TArrayView<float> ValueView(ValueStack.GetData() + ValueStack.Num() - FunctionInfo.ArgumentCount, FunctionInfo.ArgumentCount);
+			const float FuncValue = FunctionInfo.FunctionPtr(ValueView);
+			for (int32 Index = 0; Index < FunctionInfo.ArgumentCount; Index++)
+			{
+				ValueStack.Pop(false);
+			}
+			ValueStack.Push(FuncValue);
 		}
 	}
 
@@ -582,14 +952,15 @@ int32 FEngine::Execute(
 }
 
 
-TOptional<int32> FEngine::Evaluate(
-	FStringView InExpression
+TOptional<float> FEngine::Evaluate(
+	FStringView InExpression,
+	TFunctionRef<TOptional<float>(FName InConstantName)> InConstantEvaluator
 	) const
 {
 	TVariant<FExpressionObject, FParseError> Result = Parse(InExpression);
 	if (const FExpressionObject* ExpressionObject = Result.TryGet<FExpressionObject>())
 	{
-		return Execute(*ExpressionObject);
+		return Execute(*ExpressionObject, InConstantEvaluator);
 	}
 
 	return {};
@@ -597,10 +968,11 @@ TOptional<int32> FEngine::Evaluate(
 
 
 TOptional<FParseError> FEngine::Verify(
-	FStringView InExpression
+	FStringView InExpression,
+	TFunction<TOptional<float>(FName InConstantName)> InConstantEvaluator 
 	) const
 {
-	TVariant<FExpressionObject, FParseError> Result = Parse(InExpression);
+	TVariant<FExpressionObject, FParseError> Result = Parse(InExpression, InConstantEvaluator);
 	if (const FParseError* Error = Result.TryGet<FParseError>())
 	{
 		return *Error;
diff --git a/Engine/Plugins/Animation/DeformerGraph/Source/OptimusCore/Private/OptimusNode.cpp b/Engine/Plugins/Animation/DeformerGraph/Source/OptimusCore/Private/OptimusNode.cpp
--- a/Engine/Plugins/Animation/DeformerGraph/Source/OptimusCore/Private/OptimusNode.cpp
+++ b/Engine/Plugins/Animation/DeformerGraph/Source/OptimusCore/Private/OptimusNode.cpp
@@ -910,11 +910,17 @@ bool UOptimusNode::SetPinDataDomain(
 			{
 				if (Pin->GetDirection() == EOptimusNodePinDirection::Input)
 				{
-					Action->AddSubAction<FOptimusNodeGraphAction_RemoveLink>(ConnectedPin, Pin);
+					if (!FOptimusDataDomain::AreCompatible(ConnectedPin->GetDataDomain(), InDataDomain, nullptr))
+					{
+						Action->AddSubAction<FOptimusNodeGraphAction_RemoveLink>(ConnectedPin, Pin);
+					}
 				}
 				else
 				{
-					Action->AddSubAction<FOptimusNodeGraphAction_RemoveLink>(Pin, ConnectedPin);
+					if (!FOptimusDataDomain::AreCompatible(InDataDomain, ConnectedPin->GetDataDomain(), nullptr))
+					{
+						Action->AddSubAction<FOptimusNodeGraphAction_RemoveLink>(Pin, ConnectedPin);
+					}
 				}
 			}
 		}
diff --git a/Engine/Plugins/Animation/DeformerGraph/Source/OptimusCore/Private/OptimusNodePin.cpp b/Engine/Plugins/Animation/DeformerGraph/Source/OptimusCore/Private/OptimusNodePin.cpp
--- a/Engine/Plugins/Animation/DeformerGraph/Source/OptimusCore/Private/OptimusNodePin.cpp
+++ b/Engine/Plugins/Animation/DeformerGraph/Source/OptimusCore/Private/OptimusNodePin.cpp
@@ -19,37 +19,6 @@
 
 #define LOCTEXT_NAMESPACE "OptimusDeformer"
 
-static FString FormatDataDomain(
-	const FOptimusDataDomain& InDataDomain
-	)
-{
-	switch(InDataDomain.Type)
-	{
-	case EOptimusDataDomainType::Dimensional:
-		{
-			if (InDataDomain.DimensionNames.IsEmpty())
-			{
-				return TEXT("Parameter");
-			}
-			else
-			{
-				TArray<FString> Names;
-				for (FName DomainLevelName: InDataDomain.DimensionNames)
-				{
-					Names.Add(DomainLevelName.ToString());
-				}
-				return FString::Join(Names, *FString(UTF8TEXT(" › ")));
-			}
-		}
-		
-	case EOptimusDataDomainType::Expression:
-		return InDataDomain.Expression.TrimStartAndEnd();
-	}
-	
-	checkNoEntry();
-	return TEXT("");
-}
-
 
 
 UOptimusNodePin* UOptimusNodePin::GetParentPin()
@@ -171,7 +140,7 @@ FText UOptimusNodePin::GetTooltipText() const
 	{
 		return FText::FormatOrdered(LOCTEXT("OptimusNodePin_Tooltip_Resource", "Name:\t{0}\nType:\t{1} ({2})\nStorage:\tResource\nDomain:\t{3}"),
 			FText::FromString(GetName()), DataType->DisplayName, FText::FromString(DataType->ShaderValueType->ToString()),
-			FText::FromString(FormatDataDomain(DataDomain)));
+			FText::FromString(DataDomain.GetDisplayName()));
 	}
 }
 
@@ -493,32 +462,7 @@ bool UOptimusNodePin::CanCannect(const UOptimusNodePin* InOtherPin, FString* Out
 	const UOptimusNodePin *OutputPin = Direction == EOptimusNodePinDirection::Output ? this : InOtherPin;
 	const UOptimusNodePin* InputPin = Direction == EOptimusNodePinDirection::Input ? this : InOtherPin;
 
-	// We don't allow resource -> value connections. All other combos are legit. 
-	// Value -> Resource just means the resource gets filled with the value.
-	if (!OutputPin->DataDomain.IsSingleton() && InputPin->DataDomain.IsSingleton())
-	{
-		if (OutReason)
-		{
-			*OutReason = TEXT("Can't connect a resource output into a value input.");
-		}
-		return false;
-	}
-
-	// If it's resource -> resource, check that the dimensionality is the same.
-	if (!OutputPin->DataDomain.IsSingleton() && !InputPin->DataDomain.IsSingleton())
-	{
-		if (OutputPin->DataDomain != InputPin->DataDomain)
-		{
-			if (OutReason)
-			{
-				*OutReason = FString::Printf(TEXT("Can't connect resources with different data domain types (%s vs %s)."),
-					*FormatDataDomain(OutputPin->DataDomain), *FormatDataDomain(InputPin->DataDomain));
-			}
-			return false;
-		}
-	}
-
-	return true;
+	return FOptimusDataDomain::AreCompatible(OutputPin->DataDomain, InputPin->DataDomain, OutReason);
 }
 
 
diff --git a/Engine/Plugins/Animation/DeformerGraph/Source/OptimusCore/Public/IOptimusComputeKernelProvider.h b/Engine/Plugins/Animation/DeformerGraph/Source/OptimusCore/Public/IOptimusComputeKernelProvider.h
--- a/Engine/Plugins/Animation/DeformerGraph/Source/OptimusCore/Public/IOptimusComputeKernelProvider.h
+++ b/Engine/Plugins/Animation/DeformerGraph/Source/OptimusCore/Public/IOptimusComputeKernelProvider.h
@@ -15,6 +15,7 @@ class UOptimusKernelSource;
 class UOptimusNode;
 class UOptimusNodePin;
 struct FOptimusPinTraversalContext;
+struct FOptimusConstantContainer;
 
 
 // Maps the data interface's data binding index to the function we would like to have present
@@ -83,7 +84,7 @@ public:
 	 * @param InKernelDataInterface
 	 * @param OutInputDataBindings
 	 * @param OutOutputDataBindings
-	 * @param OutExecutionDataInterface
+	 * @param OutConstantContainer
 	 */
 	virtual FOptimus_ComputeKernelResult CreateComputeKernel(
 		UObject* InKernelSourceOuter,
@@ -95,7 +96,8 @@ public:
 		const UOptimusComponentSourceBinding* InGraphDataComponentBinding,
 		const UComputeDataInterface* InKernelDataInterface,
 		FOptimus_InterfaceBindingMap& OutInputDataBindings,
-		FOptimus_InterfaceBindingMap& OutOutputDataBindings
+		FOptimus_InterfaceBindingMap& OutOutputDataBindings,
+		FOptimusConstantContainer& OutConstantContainer
 		) const = 0;
 
 	/** Returns the execution domain that this kernel should iterate over */
diff --git a/Engine/Plugins/Animation/DeformerGraph/Source/OptimusCore/Public/IOptimusDeformerInstanceAccessor.h b/Engine/Plugins/Animation/DeformerGraph/Source/OptimusCore/Public/IOptimusDeformerInstanceAccessor.h
new file mode 100644
--- /dev/null
+++ b/Engine/Plugins/Animation/DeformerGraph/Source/OptimusCore/Public/IOptimusDeformerInstanceAccessor.h
@@ -0,0 +1,27 @@
+﻿// Copyright Epic Games, Inc. All Rights Reserved.
+
+#pragma once
+
+#include "UObject/Interface.h"
+
+#include "IOptimusDeformerInstanceAccessor.generated.h"
+
+
+class UOptimusDeformerInstance;
+
+UINTERFACE()
+class OPTIMUSCORE_API UOptimusDeformerInstanceAccessor :
+	public UInterface
+{
+	GENERATED_BODY()
+};
+
+
+class IOptimusDeformerInstanceAccessor
+{
+	GENERATED_BODY()
+
+public:
+	virtual void SetDeformerInstance(UOptimusDeformerInstance* InInstance) = 0;
+	virtual UOptimusDeformerInstance* GetDeformerInstance() const = 0;
+};
diff --git a/Engine/Plugins/Animation/DeformerGraph/Source/OptimusCore/Public/OptimusConstant.h b/Engine/Plugins/Animation/DeformerGraph/Source/OptimusCore/Public/OptimusConstant.h
new file mode 100644
--- /dev/null
+++ b/Engine/Plugins/Animation/DeformerGraph/Source/OptimusCore/Public/OptimusConstant.h
@@ -0,0 +1,181 @@
+﻿// Copyright Epic Games, Inc. All Rights Reserved.
+#pragma once
+
+#include "OptimusConstant.generated.h"
+
+class UOptimusComponentSourceBinding;
+class UOptimusNode;
+
+USTRUCT()
+struct FOptimusConstantIdentifier
+{
+	GENERATED_BODY()
+
+	FOptimusConstantIdentifier() = default;
+	FOptimusConstantIdentifier(const UOptimusNode* InNode, const FName& InGroupName, const FName& InConstantName);
+
+	bool IsValid() const
+	{
+		return NodePath != NAME_None && ConstantName != NAME_None;
+	}
+
+	FName GetLocalConstantName() const
+	{
+		if (GroupName == NAME_None)
+		{
+			return ConstantName;
+		}
+
+		return *(GroupName.ToString() + TEXT(".") + ConstantName.ToString());
+	};
+	
+	friend uint32 GetTypeHash(const FOptimusConstantIdentifier& InIdentifier)
+	{
+		return HashCombineFast(GetTypeHash(InIdentifier.NodePath),
+			HashCombineFast(GetTypeHash(InIdentifier.GroupName), GetTypeHash(InIdentifier.ConstantName)));
+	}
+
+	bool operator==(FOptimusConstantIdentifier const& InOther) const
+	{
+		return NodePath == InOther.NodePath && GroupName == InOther.GroupName && ConstantName == InOther.ConstantName;
+	}
+	
+	UPROPERTY()
+	FName NodePath;
+
+	UPROPERTY()
+	FName GroupName;
+	
+	UPROPERTY()
+	FName ConstantName;
+};
+
+USTRUCT()
+struct FOptimusConstantDefinition
+{
+	GENERATED_BODY()
+
+	FOptimusConstantDefinition() = default;
+	FOptimusConstantDefinition(const FOptimusConstantIdentifier& InReferencedConstant) : ReferencedConstant(InReferencedConstant) {}
+	FOptimusConstantDefinition(const FString& InExpression) : Expression(InExpression) {}
+	
+	UPROPERTY()
+	FOptimusConstantIdentifier ReferencedConstant;
+
+	UPROPERTY()
+	FString Expression;
+};
+
+UENUM()
+enum class EOptimusConstantType
+{
+	Input,
+	Output
+};
+
+USTRUCT()
+struct FOptimusConstant
+{
+	GENERATED_BODY()
+	
+	FOptimusConstant() = default;
+	FOptimusConstant(
+		const FOptimusConstantIdentifier& InIdentifier,
+		const FOptimusConstantDefinition& InDefinition,
+		int32 InComponentBindingIndex,
+		EOptimusConstantType InType
+		) :
+		Identifier(InIdentifier),
+		Definition(InDefinition),
+		ComponentBindingIndex(InComponentBindingIndex),
+		Type(InType)
+		{}
+	
+	UPROPERTY()
+	FOptimusConstantIdentifier Identifier;
+
+	UPROPERTY()
+	FOptimusConstantDefinition Definition;
+	
+	UPROPERTY()
+	int32 ComponentBindingIndex;
+	
+	UPROPERTY()
+	EOptimusConstantType Type;
+};
+
+
+USTRUCT()
+struct FOptimusConstantIndex
+{
+	GENERATED_BODY()
+
+	UPROPERTY()
+	int32 KernelIndex = INDEX_NONE;
+
+	UPROPERTY()
+	EOptimusConstantType Type;
+	
+	UPROPERTY()
+	int32 ConstantIndex = INDEX_NONE;
+
+	bool IsValid() const
+	{
+		return KernelIndex != INDEX_NONE && ConstantIndex != INDEX_NONE;
+	};
+};
+
+struct FOptimusConstantContainerInstance;
+
+USTRUCT()
+struct FOptimusKernelConstantContainer
+{
+	GENERATED_BODY()
+
+	void AddToKernelContainer(const FOptimusConstant& InConstant);
+	
+	UPROPERTY();
+	TArray<FOptimusConstant> InputConstants;
+	
+	UPROPERTY();
+	TArray<FOptimusConstant> OutputConstants;
+
+	UPROPERTY()
+	TMap<FName, int32> GroupNameToBindingIndex;
+};
+
+USTRUCT()
+struct FOptimusConstantContainer
+{
+	GENERATED_BODY()
+
+public:
+	FOptimusKernelConstantContainer& AddContainerForKernel();
+
+	void Reset();
+	
+private:
+	UPROPERTY()
+	TArray<FOptimusKernelConstantContainer> KernelContainers;
+
+	friend struct FOptimusConstantContainerInstance;
+};
+
+struct FOptimusConstantContainerInstance
+{
+public:
+	bool Initialize(const FOptimusConstantContainer& InContainer, const TMap<int32, TMap<FName, TArray<float>>>& InBindingIndexToConstantValues);
+	TArray<float> GetConstantValuePerInvocation(const FOptimusConstantIdentifier& InIdentifier) const;
+	
+private:
+	TMap<FOptimusConstantIdentifier, TArray<float>> ConstantToValuePerInvocation;
+
+	bool EvaluateAndSaveResult(const FOptimusConstantIdentifier InIdentifier, const FString& InExpression, const TMap<FName, TArray<float>>& InDependentConstants);
+};
+
+
+
+
+
+
+
diff --git a/Engine/Plugins/Animation/DeformerGraph/Source/OptimusCore/Public/OptimusDataDomain.h b/Engine/Plugins/Animation/DeformerGraph/Source/OptimusCore/Public/OptimusDataDomain.h
--- a/Engine/Plugins/Animation/DeformerGraph/Source/OptimusCore/Public/OptimusDataDomain.h
+++ b/Engine/Plugins/Animation/DeformerGraph/Source/OptimusCore/Public/OptimusDataDomain.h
@@ -78,12 +78,6 @@ struct OPTIMUSCORE_API FOptimusDataDomain
 	UPROPERTY(EditAnywhere, Category = DataDomain, meta=(EditCondition="DomainType==EOptimusDataDomainType::Expression"))
 	FString Expression;
 
-	/** A convenience function to compute element counts from a data domain based on execution domain element counts.
-	 *  Current limitation is that it will only return values for zero- or one-dimensional domains, or expressions.
-	 *  If an expression is invalid, or the dimension counts are greater than one, and empty array is returned.
-	 *  \param InDomainCounts A map of execution domains and their element counts.   
-	 */
-	TOptional<int32> GetElementCount(TMap<FName, int32> InDomainCounts) const;
 	
 	/** Convenience function to check if this data domain is a singleton */
 	bool IsSingleton() const
@@ -91,6 +85,29 @@ struct OPTIMUSCORE_API FOptimusDataDomain
 		return Type == EOptimusDataDomainType::Dimensional && DimensionNames.IsEmpty();
 	}
 
+	/** Convenience function to check if this data domain is Multi-dimensional*/
+	bool IsMultiDimensional() const
+	{
+		return Type == EOptimusDataDomainType::Dimensional && DimensionNames.Num() > 1;
+	}
+	
+	/** Convenience function to check if this data domain is One-dimensional*/
+	bool IsOneDimensional() const
+	{
+		return (Type == EOptimusDataDomainType::Dimensional && DimensionNames.Num() == 1 ) ||
+			(Type == EOptimusDataDomainType::Expression);
+	}
+
+	bool IsFullyDefined() const
+	{
+		if (Type == EOptimusDataDomainType::Expression && Expression.IsEmpty())
+		{
+			return false;
+		}
+
+		return true;
+	}
+
 	/** Convert the data domain to a text serializable string */
 	FString ToString() const;
 
@@ -101,6 +118,12 @@ struct OPTIMUSCORE_API FOptimusDataDomain
 	bool operator!=(const FOptimusDataDomain& InOtherDomain) const { return !operator==(InOtherDomain); }
 	void PostSerialize(const FArchive& Ar);
 	void BackCompFixupLevels();
+
+	TOptional<FString> AsExpression() const;
+
+	FString GetDisplayName() const;
+
+	static bool AreCompatible(const FOptimusDataDomain& InOutput, const FOptimusDataDomain& InInput, FString* OutReason);
 	
 private:
 	UPROPERTY()
diff --git a/Engine/Plugins/Animation/DeformerGraph/Source/OptimusCore/Public/OptimusDeformer.h b/Engine/Plugins/Animation/DeformerGraph/Source/OptimusCore/Public/OptimusDeformer.h
--- a/Engine/Plugins/Animation/DeformerGraph/Source/OptimusCore/Public/OptimusDeformer.h
+++ b/Engine/Plugins/Animation/DeformerGraph/Source/OptimusCore/Public/OptimusDeformer.h
@@ -10,6 +10,7 @@
 #include "OptimusNodeGraph.h"
 #include "OptimusComponentSource.h"
 #include "OptimusDataDomain.h"
+#include "OptimusConstant.h"
 
 #include "Animation/MeshDeformer.h"
 #include "Interfaces/Interface_PreviewMeshProvider.h"
@@ -440,7 +441,11 @@ protected:
 	// The compute graphs to execute.
 	UPROPERTY()
 	TArray<FOptimusComputeGraphInfo> ComputeGraphs;
-	
+
+	// Constants used by resource pins in the graph
+	UPROPERTY()
+	FOptimusConstantContainer ConstantContainer;
+
 private:
 	void PostLoadFixupMissingComponentBindingsCompat();
 	void PostLoadFixupMismatchedResourceDataDomains();
diff --git a/Engine/Plugins/Animation/DeformerGraph/Source/OptimusCore/Public/OptimusDeformerInstance.h b/Engine/Plugins/Animation/DeformerGraph/Source/OptimusCore/Public/OptimusDeformerInstance.h
--- a/Engine/Plugins/Animation/DeformerGraph/Source/OptimusCore/Public/OptimusDeformerInstance.h
+++ b/Engine/Plugins/Animation/DeformerGraph/Source/OptimusCore/Public/OptimusDeformerInstance.h
@@ -4,10 +4,13 @@
 
 #include "Animation/MeshDeformerInstance.h"
 #include "ComputeFramework/ComputeGraphInstance.h"
+#include "OptimusConstant.h"
 #include "Engine/EngineTypes.h"
 
 #include "OptimusDeformerInstance.generated.h"
 
+class UOptimusComponentSource;
+struct FOptimusConstantIdentifier;
 enum class EOptimusNodeGraphType;
 struct FOptimusPersistentStructuredBuffer;
 class FRDGBuffer;
@@ -19,6 +22,31 @@ class UOptimusVariableContainer;
 class UOptimusVariableDescription;
 class UOptimusComponentSourceBinding;
 
+USTRUCT()
+struct FOptimusDeformerInstanceComponentContext
+{
+	GENERATED_BODY()
+
+	UPROPERTY()
+	TArray<int32> LodIndexPerComponent;
+
+	friend uint32 GetTypeHash(const FOptimusDeformerInstanceComponentContext& InContext)
+	{
+		uint32 Hash = GetTypeHash(InContext.LodIndexPerComponent.Num());
+		
+		for (int32 LodIndex : InContext.LodIndexPerComponent)
+		{
+			Hash = HashCombineFast(Hash, GetTypeHash(LodIndex));
+		}
+		return Hash;
+	}
+
+	bool operator==(const FOptimusDeformerInstanceComponentContext& InOther) const
+	{
+		return LodIndexPerComponent == InOther.LodIndexPerComponent;
+	}
+};
+
 class FOptimusPersistentBufferPool
 {
 public:
@@ -200,6 +228,8 @@ public:
 
 	void SetCanBeActive(bool bInCanBeActive);
 
+	TArray<float> GetConstantValuePerInvocation(const FOptimusConstantIdentifier& InIdentifier);
+
 protected:
 	/** Implementation of UMeshDeformerInstance. */
 	void AllocateResources() override;
@@ -222,7 +252,18 @@ private:
 	/** Storage for variable data. */
 	UPROPERTY()
 	TObjectPtr<UOptimusVariableContainer> Variables;
+
+	UPROPERTY()
+	TArray<TWeakObjectPtr<UActorComponent>> WeakBoundComponents;
 	
+	UPROPERTY()
+	TArray<TWeakObjectPtr<const UOptimusComponentSource>> WeakComponentSources;
+
+	UPROPERTY()
+	FOptimusConstantContainer ConstantContainer;
+	
+	TMap<FOptimusDeformerInstanceComponentContext, FOptimusConstantContainerInstance> ConstantValuesPerContext;
+
 	// List of graphs that should be run on the next tick. 
 	TSet<FName> GraphsToRunOnNextTick;
 	FCriticalSection GraphsToRunOnNextTickLock;
diff --git a/Engine/Plugins/Animation/DeformerGraph/Source/OptimusCore/Public/OptimusExpressionEvaluator.h b/Engine/Plugins/Animation/DeformerGraph/Source/OptimusCore/Public/OptimusExpressionEvaluator.h
--- a/Engine/Plugins/Animation/DeformerGraph/Source/OptimusCore/Public/OptimusExpressionEvaluator.h
+++ b/Engine/Plugins/Animation/DeformerGraph/Source/OptimusCore/Public/OptimusExpressionEvaluator.h
@@ -3,30 +3,64 @@
 #pragma once
 
 #include "Containers/Map.h"
-#include "Containers/UnrealString.h"
-#include "Misc/Optional.h"
 #include "Misc/TVariant.h"
 #include "UObject/NameTypes.h"
 
+template <typename OptionalType> struct TOptional;
+
 namespace Optimus::Expression
 {
 
+struct FParseLocation
+{
+	FParseLocation() = default;
+	FParseLocation(int32 InStart, int32 InEndLocation) :
+		Start(InStart), 
+		End(InEndLocation)
+	{}
+	FParseLocation(FStringView InParseRange, const TCHAR* InParseStart) :
+		Start(static_cast<int32>(InParseRange.begin() - InParseStart)),
+		End(static_cast<int32>(InParseRange.end() - InParseStart))
+	{}
+
+	/** Merge two parse locations and return a new one */
+	FParseLocation Join(const FParseLocation& InLocation) const
+	{
+		return FParseLocation(FMath::Min(Start, InLocation.Start), FMath::Max(End, InLocation.End));
+	}
+	
+	const int32 Start = INDEX_NONE;
+	const int32 End = INDEX_NONE;
+};
+
 struct FParseError
 {
 	FParseError() = default;
-	FParseError(const int32 InColumn, FString&& InMessage) : Column(InColumn), Message(MoveTemp(InMessage)) {}
-
-	/** The column at which the error occurs. Zero-based. */
-	const int32 Column = 0;
+	FParseError(FString&& InMessage, FParseLocation InLocation) :
+		Message(MoveTemp(InMessage)),
+		Location(InLocation)
+	{}
 
 	/** The human-readable parsing error */
-	const FString Message = FString();
+	const FString Message{};
+	
+	/** The parse location where the error took place. */
+	const FParseLocation Location{};
 };
 
 struct FExpressionObject
 {
+	FString ToString() const;
+
+	/** Returns all named constants used by this expression */
+	TSet<FName> GetUsedConstants() const;
+
+	void Load(FArchive& Ar);
+	void Save(FArchive& Ar) const;
+	
 private:
-	enum class EOperator
+	
+	enum class EOperator : int32
 	{
 		Negate,				// Negation operator.
 		Add,				// Add last two values on stack and add the result to the stack.
@@ -35,18 +69,24 @@ private:
 		Divide,				// Same but divide (div-by-zero returns zero)
 		Modulo,				// Same but modulo (mod-by-zero returns zero)
 		Power,				// Raise to power
-		// FloorDivide,    	// Divide and round the result down
+		FloorDivide,    	// Divide and round the result down
+	};
+
+	struct FFunctionRef
+	{
+		FFunctionRef(int32 InIndex) : Index(InIndex) {}
+		int32 Index;
 	};
 	
 	friend class FEngine;
-	using OpElement = TVariant<EOperator, /* Constant Ref */ FName, /* Value */ int32>;
+	using OpElement = TVariant<EOperator, FName /* Constant */, FFunctionRef, /* Value */ float>;
 	TArray<OpElement, TInlineAllocator<8>> Expression;
 };
 
 /** Simple, fast expression evaluation engine. Can evaluate simple arithmetic expressions,
  *  including user-settable constants. 
  *		1 + 1
- *		foo + (3 * bar)
+ *		foo + (3.0 * bar)
  *
  *  Allowable operands are taken from Python and are currently as follow:
  *      -x				Unary negation
@@ -55,76 +95,51 @@ private:
  *      x * y			Multiplication
  *      x / y			Division (div-by-zero returns zero)
  *      x ** y			Raise x to y's power
+ *      x // y			Division and rounding down.
  *
  *  Terms can also be grouped by parentheses '(' and ')'.
  *  
  *	Constants that begin with a number, include spaces or any of the above operands/parentheses 
  *	are allowed, but must be then enclosed in single quotes when referenced.
  *	E.g:
- *	   'foo(1)' * 2
+ *	   'foo(1)' * 2.0
  *	rather than:
- *	   foo (1) * 2
+ *	   foo (1) * 2.0
  */
-enum class EParseFlags
-{
-	Default				= 0,		/** Default parsing behavior */
-	ValidateConstants	= 1 << 0,	/** Validate that the referred-to constant exists */
-};
-ENUM_CLASS_FLAGS(EParseFlags)
 
 
 class OPTIMUSCORE_API FEngine
 {
 public:
-	explicit FEngine(EParseFlags InParseFlags = EParseFlags::Default) :
-		ParseFlags(InParseFlags)
-	{ }
-
-	/** Initialize the engine with constants, making a copy of the value constant map. */
-	explicit FEngine(
-		const TMap<FName, int32>& InConstants,
-		EParseFlags InParseFlags = EParseFlags::Default
-		);
-	
-	/** Initialize the engine with constants, taking ownership of the value constant map. */
-	explicit FEngine(
-		TMap<FName, int32>&& InConstants,
-		EParseFlags InParseFlags = EParseFlags::Default
-		);
-
-	/** Update existing constant values. No new constants are added. */
-	void UpdateConstantValues(
-		const TMap<FName, int32>& InConstants
-		);
-
-	/** Return existing constants and their values. */
-	const TMap<FName, int32>& GetConstantValues() const
-	{
-		return Constants;
-	}
-
 	/** Parse an expression and either receive an error or a expression execution object that
 	 *  matches this evaluator. */
 	TVariant<FExpressionObject, FParseError> Parse(
-		FStringView InExpression
+		FStringView InExpression,
+		TFunction<TOptional<float>(FName InConstantName)> InConstantEvaluator = {} 
 		) const;
 
-	/** Executes the given expression object */
-	int32 Execute(
-		const FExpressionObject& InExpressionObject
+	/** Executes the given expression object.
+	 *  \param InExpressionObject The parsed expression object to evaluate.
+	 *  \param InConstantEvaluator A function that returns the value of a given constant, if it's known, or none otherwise.
+	 */
+	float Execute(
+		const FExpressionObject& InExpressionObject,
+		TFunctionRef<TOptional<float>(FName InConstantName)> InConstantEvaluator
 		) const;
 	
 	/** Evaluate an expression directly and returns the result. If the expression is incorrectly
 	 *  formed, then the optional return value holds nothing.
 	 */ 
-	TOptional<int32> Evaluate(
-		FStringView InExpression
+	TOptional<float> Evaluate(
+		FStringView InExpression,
+		TFunctionRef<TOptional<float>(FName InConstantName)> InConstantEvaluator
 		) const;
 
 	/** Verify an expression. Does not check for variable validity, only that the expression
 	 *  can be evaluated */
 	TOptional<FParseError> Verify(
-		FStringView InExpression
+		FStringView InExpression,
+		TFunction<TOptional<float>(FName InConstantName)> InConstantEvaluator = {} 
 		) const;
 	
 private:
@@ -137,16 +152,50 @@ private:
 		Divide,				// '/'
 		Modulo,				// '%'
 		Power,				// '**'
-		// FloorDivide,    	// '//'
+		FloorDivide,    	// '//'
 		ParenOpen,			// '('
 		ParenClose,			// ')'
+		Comma,				// ','
+
+		Max					// Must be last
+	};
+	using FToken = TVariant<EOperatorToken, /* Identifier */ FName, /* Value */ float>;
+
+	struct FTokenParseResult
+	{
+		FTokenParseResult() = default;
+		
+		FTokenParseResult(FToken&& InToken, FStringView InTokenRange, const TCHAR* InExpression) :
+			Result(TInPlaceType<FToken>(), MoveTemp(InToken)),
+			Location(InTokenRange, InExpression)
+		{}
+		
+		FTokenParseResult(FString&& InError, FStringView InTokenRange, const TCHAR* InExpression) :
+			Result(TInPlaceType<FString>(), MoveTemp(InError)),
+			Location(InTokenRange, InExpression)
+		{}
+		
+		bool IsError() const { return Result.IsType<FString>(); }
+		bool IsToken() const { return Result.IsType<FToken>(); }
+		
+		const FToken& GetToken() const
+		{
+			return Result.Get<FToken>();
+		}
+		
+		FParseError GetParseError()
+		{
+			return FParseError(MoveTemp(Result.Get<FString>()), Location);
+		}
+		
+		TVariant<FString, FToken> Result;
+		FParseLocation Location{};
 	};
-	using FToken = TVariant<EOperatorToken, /* Identifier */ FName, /* Value */ int32>;
 	
-	/** Parse a integer value and return the result (or error). Same semantics
+	/** Parse a floating point value and return the result (or error). Same semantics
 	 *  apply as with ParseToken.
 	 */
-	static TVariant<FToken, FParseError> ParseInt32(
+	static FTokenParseResult ParseFloat(
 		FStringView& InOutParseRange,
 		FStringView InExpression
 		);
@@ -154,41 +203,50 @@ private:
 	/** Parse a constant reference and return the result (or error). Same semantics
 	 *  apply as with ParseToken.
 	 */
-	TVariant<FToken, FParseError> ParseIdentifier(
+	static FTokenParseResult ParseIdentifier(
 		FStringView& InOutParseRange,
 		FStringView InExpression
-		) const;
+		);
 	
 	/** Parse a single token from the beginning of the range, returning either the token
 	 *  or an error, if one was encountered.
 	 *  If the token was successfully parsed, then the given string view range gets updated
 	 *  with the new starting point at the end of the token just parsed
 	 */
-	TVariant<FToken, FParseError> ParseToken(
+	static FTokenParseResult ParseToken(
 		FStringView& InOutParseRange,
 		FStringView InExpression
-		) const;
-
-	/** Helper function to construct a token parse error result */
-	static TVariant<FToken, FParseError> ParseTokenError(
-		FString&& InError,
-		const TCHAR* InCurrentPos,
-		FStringView InExpression
 		);
-	
+
 	/** Helper function to get a text version of a token. */
 	static FString TokenToString(
 		const FToken& InToken
 		);
 	
-	/** The list of constants, stored using the FName key since the expression object refers
-	 *  to it directly */
-	// FUTURE: Add support for variables by allowing TFunction{Ref} objects along with direct
-	// constants.
-	TMap<FName, int32> Constants;
-
-	/** The parse flags for this engine. Only set once */
-	EParseFlags ParseFlags = EParseFlags::Default;
+	/** For now we just have built-in functions */
+	friend struct FInitializeBuiltinFunctions;
+	using FunctionType = TFunction<float(TArrayView<const float>)>;
+	struct FFunctionInfo
+	{
+		int32 ArgumentCount;
+		FunctionType FunctionPtr;
+	};
+	static TMap<FName, int32> FunctionNameIndex;
+	static TArray<FFunctionInfo> Functions;
 };
 
 }
+
+
+static FArchive& operator<<(FArchive& Ar, Optimus::Expression::FExpressionObject& InObject)
+{
+	if (Ar.IsLoading())
+	{
+		InObject.Load(Ar);
+	}
+	else
+	{
+		InObject.Save(Ar);
+	}
+	return Ar;
+}
diff --git a/Engine/Plugins/Animation/DeformerGraph/Source/OptimusEditor/Private/OptimusDetailsCustomization.cpp b/Engine/Plugins/Animation/DeformerGraph/Source/OptimusEditor/Private/OptimusDetailsCustomization.cpp
--- a/Engine/Plugins/Animation/DeformerGraph/Source/OptimusEditor/Private/OptimusDetailsCustomization.cpp
+++ b/Engine/Plugins/Animation/DeformerGraph/Source/OptimusEditor/Private/OptimusDetailsCustomization.cpp
@@ -466,13 +466,13 @@ void FOptimusDataDomainCustomization::CustomizeHeader(
 							{
 								DataDomain.Expression = "1";
 							}
-							else if (DataDomain.Multiplier > 1)
+							else if (DataDomain.IsMultiDimensional())
 							{
-								DataDomain.Expression = FString::Printf(TEXT("%s * %d"), *DataDomain.DimensionNames[0].ToString(), DataDomain.Multiplier);
+								DataDomain.Expression = DataDomain.DimensionNames[0].ToString();
 							}
 							else
 							{
-								DataDomain.Expression = DataDomain.DimensionNames[0].ToString();
+								DataDomain.Expression = DataDomain.AsExpression().Get({});
 							}
 						}
 						else
