#UE Fix some uninitialized structs

diff --git a/Engine/Plugins/Animation/DeformerGraph/Source/OptimusCore/Public/OptimusConstant.h b/Engine/Plugins/Animation/DeformerGraph/Source/OptimusCore/Public/OptimusConstant.h
--- a/Engine/Plugins/Animation/DeformerGraph/Source/OptimusCore/Public/OptimusConstant.h
+++ b/Engine/Plugins/Animation/DeformerGraph/Source/OptimusCore/Public/OptimusConstant.h
@@ -98,10 +98,10 @@ struct FOptimusConstant
 	FOptimusConstantDefinition Definition;
 	
 	UPROPERTY()
-	int32 ComponentBindingIndex;
+	int32 ComponentBindingIndex = 0;
 	
 	UPROPERTY()
-	EOptimusConstantType Type;
+	EOptimusConstantType Type = EOptimusConstantType::Input;
 };
 
 
@@ -114,7 +114,7 @@ struct FOptimusConstantIndex
 	int32 KernelIndex = INDEX_NONE;
 
 	UPROPERTY()
-	EOptimusConstantType Type;
+	EOptimusConstantType Type = EOptimusConstantType::Input;
 	
 	UPROPERTY()
 	int32 ConstantIndex = INDEX_NONE;
