Optimus: Add a dirty state indicator to the compiler button.

diff --git a/Engine/Plugins/Animation/DeformerGraph/Source/OptimusCore/Private/Nodes/OptimusNode_CustomComputeKernel.cpp b/Engine/Plugins/Animation/DeformerGraph/Source/OptimusCore/Private/Nodes/OptimusNode_CustomComputeKernel.cpp
--- a/Engine/Plugins/Animation/DeformerGraph/Source/OptimusCore/Private/Nodes/OptimusNode_CustomComputeKernel.cpp
+++ b/Engine/Plugins/Animation/DeformerGraph/Source/OptimusCore/Private/Nodes/OptimusNode_CustomComputeKernel.cpp
@@ -131,8 +131,14 @@ FString UOptimusNode_CustomComputeKernel::GetShaderText() const
 
 void UOptimusNode_CustomComputeKernel::SetShaderText(const FString& NewText) 
 {
+	PreEditChange(nullptr);
+
 	Modify();
 	ShaderSource.ShaderText = NewText;
+
+	FPropertyChangedEvent PropertyChangedEvent(FindFProperty<FProperty>(StaticClass(), GET_MEMBER_NAME_CHECKED(UOptimusNode_CustomComputeKernel, ShaderSource)),
+		EPropertyChangeType::ValueSet);
+	PostEditChangeProperty(PropertyChangedEvent);
 }
 
 #endif // WITH_EDITOR
diff --git a/Engine/Plugins/Animation/DeformerGraph/Source/OptimusCore/Private/OptimusDeformer.cpp b/Engine/Plugins/Animation/DeformerGraph/Source/OptimusCore/Private/OptimusDeformer.cpp
--- a/Engine/Plugins/Animation/DeformerGraph/Source/OptimusCore/Private/OptimusDeformer.cpp
+++ b/Engine/Plugins/Animation/DeformerGraph/Source/OptimusCore/Private/OptimusDeformer.cpp
@@ -73,7 +73,9 @@ UOptimusDeformer::UOptimusDeformer()
 	Variables = CreateDefaultSubobject<UOptimusVariableContainer>(TEXT("@Variables"));
 	Resources = CreateDefaultSubobject<UOptimusResourceContainer>(TEXT("@Resources"));
 
+#if WITH_EDITOR
 	FOptimusDataTypeRegistry::Get().GetOnDataTypeChanged().AddUObject(this, &UOptimusDeformer::OnDataTypeChanged);
+#endif
 }
 
 
@@ -1294,6 +1296,18 @@ bool UOptimusDeformer::SetComponentBindingSourceDirect(
 	return true;
 }
 
+void UOptimusDeformer::SetStatusFromDiagnostic(EOptimusDiagnosticLevel InDiagnosticLevel)
+{
+	if (InDiagnosticLevel == EOptimusDiagnosticLevel::Error)
+	{
+		Status = EOptimusDeformerStatus::HasErrors;
+	}
+	else if (InDiagnosticLevel == EOptimusDiagnosticLevel::Warning && Status == EOptimusDeformerStatus::Compiled)
+	{
+		Status = EOptimusDeformerStatus::CompiledWithWarnings;
+	}
+}
+
 
 UOptimusComponentSourceBinding* UOptimusDeformer::ResolveComponentBinding(
 	FName InBindingName
@@ -1387,6 +1401,9 @@ bool UOptimusDeformer::Compile()
 		CompileBeginDelegate.Broadcast(this);
 		CompileMessageDelegate.Broadcast(Diagnostic);
 		CompileEndDelegate.Broadcast(this);
+
+		Status = EOptimusDeformerStatus::HasErrors;
+		
 		return false;
 	}
 
@@ -1398,8 +1415,9 @@ bool UOptimusDeformer::Compile()
 	// Wait for rendering to be done.
 	FlushRenderingCommands();
 
-	bool bCompilationFailed = false;
-	auto ErrorReporter = [this, &bCompilationFailed](
+	Status = EOptimusDeformerStatus::Compiled;
+
+	auto ErrorReporter = [this](
 		EOptimusDiagnosticLevel InDiagnosticLevel, 
 		FText InMessage,
 		const UObject* InObject)
@@ -1409,7 +1427,8 @@ bool UOptimusDeformer::Compile()
 		Diagnostic.Message = InMessage;
 		Diagnostic.Object = InObject;
 		CompileMessageDelegate.Broadcast(Diagnostic);
-		bCompilationFailed = true;
+
+		SetStatusFromDiagnostic(InDiagnosticLevel);
 	};
 
 	for (const UOptimusNodeGraph* Graph: Graphs)
@@ -1426,7 +1445,7 @@ bool UOptimusDeformer::Compile()
 	
 	CompileEndDelegate.Broadcast(this);
 
-	if (bCompilationFailed)
+	if (Status == EOptimusDeformerStatus::HasErrors)
 	{
 		ComputeGraphs.Reset();
 		ConstantContainer.Reset();
@@ -2108,9 +2127,20 @@ void UOptimusDeformer::Notify(EOptimusGlobalNotifyType InNotifyType, UObject* In
 		break;
 	}
 
+	const_cast<UOptimusDeformer*>(this)->MarkModified();
+	
 	GlobalNotifyDelegate.Broadcast(InNotifyType, InObject);
 }
 
+
+void UOptimusDeformer::MarkModified()
+{
+	if (Status != EOptimusDeformerStatus::HasErrors)
+	{
+		Status = EOptimusDeformerStatus::Modified;
+	}
+}
+
 void UOptimusDeformer::SetAllInstancesCanbeActive(bool bInCanBeActive) const
 {
 	SetAllInstancesCanbeActiveDelegate.Broadcast(bInCanBeActive);
@@ -2241,6 +2271,13 @@ void UOptimusDeformer::PostLoad()
 	{
 		PostLoadRemoveDeprecatedExecutionNodes();
 	}
+
+	// If the graph was saved at any previous version, and was clean, mark the status now as modified.
+	if (Status == EOptimusDeformerStatus::Compiled &&
+		GetLinkerCustomVersion(FOptimusObjectVersion::GUID) < FOptimusObjectVersion::LatestVersion)
+	{
+		Status = EOptimusDeformerStatus::Modified;
+	}
 }
 
 
@@ -2522,10 +2559,13 @@ void UOptimusDeformer::GetAssetRegistryTags(TArray<FAssetRegistryTag>& OutTags)
 	}
 }
 
-
 #if WITH_EDITOR
+
 void UOptimusDeformer::PostEditChangeProperty(FPropertyChangedEvent& PropertyChangedEvent)
 {
+	Super::PostEditChangeProperty(PropertyChangedEvent);
+	
+	MarkModified();
 }
 #endif
 
diff --git a/Engine/Plugins/Animation/DeformerGraph/Source/OptimusCore/Private/OptimusNode.cpp b/Engine/Plugins/Animation/DeformerGraph/Source/OptimusCore/Private/OptimusNode.cpp
--- a/Engine/Plugins/Animation/DeformerGraph/Source/OptimusCore/Private/OptimusNode.cpp
+++ b/Engine/Plugins/Animation/DeformerGraph/Source/OptimusCore/Private/OptimusNode.cpp
@@ -385,6 +385,21 @@ void UOptimusNode::PostLoad()
 	}
 }
 
+#if WITH_EDITOR
+void UOptimusNode::PostEditChangeProperty(FPropertyChangedEvent& PropertyChangedEvent)
+{
+	Super::PostEditChangeProperty(PropertyChangedEvent);
+
+	if (const UOptimusNodeGraph *Graph = GetOwningGraph())
+	{
+		if (UOptimusDeformer* Deformer = Cast<UOptimusDeformer>(Graph->GetCollectionRoot()))
+		{
+			Deformer->MarkModified();
+		}
+	}
+	
+}
+#endif
 
 void UOptimusNode::Notify(EOptimusGraphNotifyType InNotifyType)
 {
diff --git a/Engine/Plugins/Animation/DeformerGraph/Source/OptimusCore/Private/OptimusNodeGraph.cpp b/Engine/Plugins/Animation/DeformerGraph/Source/OptimusCore/Private/OptimusNodeGraph.cpp
--- a/Engine/Plugins/Animation/DeformerGraph/Source/OptimusCore/Private/OptimusNodeGraph.cpp
+++ b/Engine/Plugins/Animation/DeformerGraph/Source/OptimusCore/Private/OptimusNodeGraph.cpp
@@ -1587,6 +1587,23 @@ bool UOptimusNodeGraph::DoesLinkFormCycle(const UOptimusNode* InOutputNode, cons
 
 void UOptimusNodeGraph::Notify(EOptimusGraphNotifyType InNotifyType, UObject* InSubject) const
 {
+	if (InNotifyType == EOptimusGraphNotifyType::NodeDiagnosticLevelChanged)
+	{
+		if (UOptimusDeformer* Deformer = Cast<UOptimusDeformer>(GetCollectionRoot()))
+		{
+			Deformer->SetStatusFromDiagnostic(Cast<UOptimusNode>(InSubject)->GetDiagnosticLevel());
+		}
+	}
+	else if (InNotifyType != EOptimusGraphNotifyType::NodePositionChanged &&
+			 InNotifyType != EOptimusGraphNotifyType::PinExpansionChanged
+		)
+	{
+		if (UOptimusDeformer* Deformer = Cast<UOptimusDeformer>(GetCollectionRoot()))
+		{
+			Deformer->MarkModified();
+		}
+	}
+	
 	GraphNotifyDelegate.Broadcast(InNotifyType, const_cast<UOptimusNodeGraph*>(this), InSubject);
 }
 
diff --git a/Engine/Plugins/Animation/DeformerGraph/Source/OptimusCore/Public/OptimusDeformer.h b/Engine/Plugins/Animation/DeformerGraph/Source/OptimusCore/Public/OptimusDeformer.h
--- a/Engine/Plugins/Animation/DeformerGraph/Source/OptimusCore/Public/OptimusDeformer.h
+++ b/Engine/Plugins/Animation/DeformerGraph/Source/OptimusCore/Public/OptimusDeformer.h
@@ -39,6 +39,16 @@ DECLARE_MULTICAST_DELEGATE_OneParam(FOptimusGraphCompileMessageDelegate, FOptimu
 DECLARE_MULTICAST_DELEGATE_TwoParams(FOptimusConstantValueUpdate, FString const&, TArray<uint8> const&);
 DECLARE_MULTICAST_DELEGATE_OneParam(FOptimusSetAllInstancesCanbeActive, bool);
 
+UENUM()
+enum class EOptimusDeformerStatus : int32
+{
+	Compiled,					// Compiled, no warnings, no errors
+	CompiledWithWarnings,		// Compiled, has warnings
+	Modified,					// Graph has been modified, needs recompilation
+	HasErrors					// Graph produced errors at the last compile
+};
+
+
 USTRUCT()
 struct FOptimusComputeGraphInfo
 {
@@ -115,8 +125,15 @@ class OPTIMUSCORE_API UOptimusDeformer :
 public:
 	UOptimusDeformer();
 
+	/** Get the action stack for this deformer graph */
 	UOptimusActionStack *GetActionStack();
 
+	/** Returns the current compilation/error status of the deformer */
+	EOptimusDeformerStatus GetStatus() const
+	{
+		return Status;
+	}
+
 	/** Returns the global delegate used to notify on global operations (e.g. graph, variable,
 	 *  resource lifecycle events).
 	 */
@@ -258,6 +275,9 @@ public:
 	FOptimusGraphCompileMessageDelegate& GetCompileMessageDelegate() { return CompileMessageDelegate; }
 
 	void SetAllInstancesCanbeActive(bool bInCanBeActive) const;
+
+	// Mark the deformer as modified.
+	void MarkModified();
 	
 	/// UObject overrides
 	void Serialize(FArchive& Ar) override;
@@ -435,6 +455,7 @@ protected:
 		const UOptimusComponentSource *InComponentSource
 		);
 
+	void SetStatusFromDiagnostic(EOptimusDiagnosticLevel InDiagnosticLevel);
 	
 	void Notify(EOptimusGlobalNotifyType InNotifyType, UObject *InObject) const;
 
@@ -469,10 +490,13 @@ private:
 		);
 
 	void OnDataTypeChanged(FName InTypeName);
-	
+
 	UPROPERTY(transient)
 	TObjectPtr<UOptimusActionStack> ActionStack = nullptr;
 
+	UPROPERTY()
+	EOptimusDeformerStatus Status = EOptimusDeformerStatus::Modified;
+
 	UPROPERTY()
 	TArray<TObjectPtr<UOptimusNodeGraph>> Graphs;
 
diff --git a/Engine/Plugins/Animation/DeformerGraph/Source/OptimusCore/Public/OptimusNode.h b/Engine/Plugins/Animation/DeformerGraph/Source/OptimusCore/Public/OptimusNode.h
--- a/Engine/Plugins/Animation/DeformerGraph/Source/OptimusCore/Public/OptimusNode.h
+++ b/Engine/Plugins/Animation/DeformerGraph/Source/OptimusCore/Public/OptimusNode.h
@@ -180,6 +180,9 @@ public:
 	//== UObject overrides
 	void Serialize(FArchive& Ar) override;
 	void PostLoad() override;
+#if WITH_EDITOR
+	void PostEditChangeProperty(FPropertyChangedEvent& PropertyChangedEvent) override;
+#endif
 	
 protected:
 	friend class UOptimusNodeGraph;
diff --git a/Engine/Plugins/Animation/DeformerGraph/Source/OptimusEditor/Private/OptimusEditor.cpp b/Engine/Plugins/Animation/DeformerGraph/Source/OptimusEditor/Private/OptimusEditor.cpp
--- a/Engine/Plugins/Animation/DeformerGraph/Source/OptimusEditor/Private/OptimusEditor.cpp
+++ b/Engine/Plugins/Animation/DeformerGraph/Source/OptimusEditor/Private/OptimusEditor.cpp
@@ -47,6 +47,7 @@
 #include "OptimusShaderText.h"
 #include "ISourceCodeAccessModule.h"
 #include "ISourceCodeAccessor.h"
+#include "OptimusEditorStyle.h"
 
 
 #define LOCTEXT_NAMESPACE "OptimusEditor"
@@ -783,8 +784,8 @@ void FOptimusEditor::RegisterToolbar()
 		Section.AddEntry(FToolMenuEntry::InitToolBarButton(
 			Commands.Compile,
 			TAttribute<FText>(),
-			TAttribute<FText>(),
-			FSlateIcon(FAppStyle::GetAppStyleSetName(), "Blueprint.CompileStatus.Background")));
+			TAttribute<FText>(this, &FOptimusEditor::GetCompileStatusTooltip),
+			TAttribute<FSlateIcon>(this, &FOptimusEditor::GetCompileStatusIcon)));
 	}
 
 }
@@ -1201,4 +1202,44 @@ void FOptimusEditor::OnDataTypeChanged()
 	}
 }
 
+FSlateIcon FOptimusEditor::GetCompileStatusIcon() const
+{
+	static const FName CompileStatusBackground("Toolbar.CompileStatus.Background");
+	static const FName CompileStatusModified("Toolbar.CompileStatus.Overlay.Modified");
+	static const FName CompileStatusHasError("Toolbar.CompileStatus.Overlay.Error");
+	static const FName CompileStatusCompiled("Toolbar.CompileStatus.Overlay.Good");
+	static const FName CompileStatusCompiledWithWarnings("Toolbar.CompileStatus.Overlay.Warning");
+
+	const FName StyleSetName = FOptimusEditorStyle::Get().GetStyleSetName();
+
+	switch (DeformerObject->GetStatus())
+	{
+	default:
+	case EOptimusDeformerStatus::Modified:
+		return FSlateIcon(StyleSetName, CompileStatusBackground, NAME_None, CompileStatusModified);
+	case EOptimusDeformerStatus::HasErrors:
+		return FSlateIcon(StyleSetName, CompileStatusBackground, NAME_None, CompileStatusHasError);
+	case EOptimusDeformerStatus::Compiled:
+		return FSlateIcon(StyleSetName, CompileStatusBackground, NAME_None, CompileStatusCompiled);
+	case EOptimusDeformerStatus::CompiledWithWarnings:
+		return FSlateIcon(StyleSetName, CompileStatusBackground, NAME_None, CompileStatusCompiledWithWarnings);
+	}
+}
+
+FText FOptimusEditor::GetCompileStatusTooltip() const
+{
+	switch (DeformerObject->GetStatus())
+	{
+	default:
+	case EOptimusDeformerStatus::Modified:
+		return LOCTEXT("GraphStatusTooltip_Modified", "Deformer graph has been modified and needs to be recompiled in order to update the deformation result.");
+	case EOptimusDeformerStatus::HasErrors:
+		return LOCTEXT("GraphStatusTooltip_HasErrors", "Deformer graph has errors which need to be fixed.");
+	case EOptimusDeformerStatus::Compiled:
+		return LOCTEXT("GraphStatusTooltip_Compiled", "Deformer graph is compiled and corresponds to the deformation result.");
+	case EOptimusDeformerStatus::CompiledWithWarnings:
+		return LOCTEXT("GraphStatusTooltip_CompiledWithWarnings", "Deformer graph is compiled with warnings, these warnings will likely not affect the deformation result but should be fixed regardless.");
+	}
+}
+
 #undef LOCTEXT_NAMESPACE
diff --git a/Engine/Plugins/Animation/DeformerGraph/Source/OptimusEditor/Private/OptimusEditor.h b/Engine/Plugins/Animation/DeformerGraph/Source/OptimusEditor/Private/OptimusEditor.h
--- a/Engine/Plugins/Animation/DeformerGraph/Source/OptimusEditor/Private/OptimusEditor.h
+++ b/Engine/Plugins/Animation/DeformerGraph/Source/OptimusEditor/Private/OptimusEditor.h
@@ -225,6 +225,10 @@ private:
 	void OnFinishedChangingProperties(const FPropertyChangedEvent& PropertyChangedEvent);
 
 	void OnDataTypeChanged();
+
+
+	FSlateIcon GetCompileStatusIcon() const;
+	FText GetCompileStatusTooltip() const;
 	
 private:
 	// Persona toolkit for the skelmesh preview
diff --git a/Engine/Plugins/Animation/DeformerGraph/Source/OptimusEditor/Private/OptimusEditorStyle.cpp b/Engine/Plugins/Animation/DeformerGraph/Source/OptimusEditor/Private/OptimusEditorStyle.cpp
--- a/Engine/Plugins/Animation/DeformerGraph/Source/OptimusEditor/Private/OptimusEditorStyle.cpp
+++ b/Engine/Plugins/Animation/DeformerGraph/Source/OptimusEditor/Private/OptimusEditorStyle.cpp
@@ -6,6 +6,10 @@
 #include "Styling/SlateStyleMacros.h"
 #include "Styling/SlateStyleRegistry.h"
 #include "Styling/SlateTypes.h"
+#include "Styling/StyleColors.h"
+
+
+#define CORE_IMAGE_BRUSH_SVG( RelativePath, ... ) FSlateVectorImageBrush(RootToCoreContentDir(RelativePath, TEXT(".svg")), __VA_ARGS__)
 
 
 FOptimusEditorStyle::FOptimusEditorStyle() :
@@ -18,7 +22,8 @@ FOptimusEditorStyle::FOptimusEditorStyle() :
 
 	static const FSlateColor DefaultForeground(FLinearColor(0.72f, 0.72f, 0.72f, 1.f));
 
-	SetContentRoot(FPaths::EnginePluginsDir() / TEXT("Animation/DeformerGraph/Resources"));
+	FSlateStyleSet::SetCoreContentRoot(FPaths::EngineContentDir() / TEXT("Editor/Slate"));
+	FSlateStyleSet::SetContentRoot(FPaths::EnginePluginsDir() / TEXT("Animation/DeformerGraph/Resources"));
 
 	// Asset icons
 	{
@@ -73,6 +78,19 @@ FOptimusEditorStyle::FOptimusEditorStyle() :
 
 	}
 
+	// Toolbar icons
+	{
+		Set("Toolbar.CompileStatus.Background", new CORE_IMAGE_BRUSH_SVG("Starship/Blueprints/CompileStatus_Background", IconSize20x20));
+	}
+	
+	// Compile status icons.
+	{
+		Set("Toolbar.CompileStatus.Overlay.Modified", new CORE_IMAGE_BRUSH_SVG("Starship/Blueprints/CompileStatus_Unknown_Badge", IconSize20x20, FStyleColors::AccentYellow));
+		Set("Toolbar.CompileStatus.Overlay.Warning", new CORE_IMAGE_BRUSH_SVG("Starship/Blueprints/CompileStatus_Warning_Badge", IconSize20x20, FStyleColors::Warning));
+		Set("Toolbar.CompileStatus.Overlay.Good", new CORE_IMAGE_BRUSH_SVG("Starship/Blueprints/CompileStatus_Good_Badge", IconSize20x20, FStyleColors::AccentGreen));
+		Set("Toolbar.CompileStatus.Overlay.Error", new CORE_IMAGE_BRUSH_SVG("Starship/Blueprints/CompileStatus_Fail_Badge", IconSize20x20, FStyleColors::Error));
+	}
+
 	// Graph type icons
 	{
 		Set("GraphType.Setup", new IMAGE_BRUSH_SVG("Icons/Graph_Setup", IconSize20x20, DefaultForeground));
@@ -145,3 +163,5 @@ FOptimusEditorStyle& FOptimusEditorStyle::Get()
 	static FOptimusEditorStyle Instance;
 	return Instance;
 }
+
+#undef CORE_IMAGE_BRUSH_SVG
