Optimus: Fix warnings messages for connection mismatch.

#rnx

diff --git a/Engine/Plugins/Animation/DeformerGraph/Source/OptimusCore/Private/OptimusDataDomain.cpp b/Engine/Plugins/Animation/DeformerGraph/Source/OptimusCore/Private/OptimusDataDomain.cpp
--- a/Engine/Plugins/Animation/DeformerGraph/Source/OptimusCore/Private/OptimusDataDomain.cpp
+++ b/Engine/Plugins/Animation/DeformerGraph/Source/OptimusCore/Private/OptimusDataDomain.cpp
@@ -165,15 +165,17 @@ FString FOptimusDataDomain::GetDisplayName() const
 			{
 				return TEXT("Parameter");
 			}
-			else
+			TArray<FString> Names;
+			for (FName DomainLevelName: DimensionNames)
 			{
-				TArray<FString> Names;
-				for (FName DomainLevelName: DimensionNames)
-				{
-					Names.Add(DomainLevelName.ToString());
-				}
-				return FString::Join(Names, *FString(UTF8TEXT(" › ")));
+				Names.Add(DomainLevelName.ToString());
+			}
+			FString DomainName = FString::Join(Names, *FString(UTF8TEXT(" › ")));
+			if (InDataDomain.Multiplier > 1)
+			{
+				DomainName += FString::Printf(TEXT(" x %d"), Multiplier);
 			}
+			return DomainName;
 		}
 		
 	case EOptimusDataDomainType::Expression:
@@ -181,7 +183,7 @@ FString FOptimusDataDomain::GetDisplayName() const
 		{
 			return TEXT("Undefined");
 		}
-		return Expression.TrimStartAndEnd();
+		return FString::Printf(TEXT("'%s'"), *Expression.TrimStartAndEnd());
 	}
 
 	checkNoEntry();
diff --git a/Engine/Plugins/Animation/DeformerGraph/Source/OptimusCore/Private/OptimusNodePin.cpp b/Engine/Plugins/Animation/DeformerGraph/Source/OptimusCore/Private/OptimusNodePin.cpp
--- a/Engine/Plugins/Animation/DeformerGraph/Source/OptimusCore/Private/OptimusNodePin.cpp
+++ b/Engine/Plugins/Animation/DeformerGraph/Source/OptimusCore/Private/OptimusNodePin.cpp
@@ -444,6 +444,9 @@ bool UOptimusNodePin::CanCannect(const UOptimusNodePin* InOtherPin, FString* Out
 		return false;
 	}
 	
+	const UOptimusNodePin *OutputPin = Direction == EOptimusNodePinDirection::Output ? this : InOtherPin;
+	const UOptimusNodePin* InputPin = Direction == EOptimusNodePinDirection::Input ? this : InOtherPin;
+
 	// Check for incompatible types.
 	if (!((DataType->ShaderValueType.IsValid() && InOtherPin->DataType->ShaderValueType.IsValid() &&
 		  DataType->ShaderValueType == InOtherPin->DataType->ShaderValueType) ||
@@ -453,15 +456,14 @@ bool UOptimusNodePin::CanCannect(const UOptimusNodePin* InOtherPin, FString* Out
 
 		if (OutReason)
 		{
-			*OutReason = TEXT("Incompatible pin types.");
+			*OutReason = FString::Printf(TEXT("Incompatible pin types (%s vs %s)."),
+				*OutputPin->GetDataType()->DisplayName.ToString(),
+				*InputPin->GetDataType()->DisplayName.ToString());
 		}
 		return false;
 	}
 
 
-	const UOptimusNodePin *OutputPin = Direction == EOptimusNodePinDirection::Output ? this : InOtherPin;
-	const UOptimusNodePin* InputPin = Direction == EOptimusNodePinDirection::Input ? this : InOtherPin;
-
 	return FOptimusDataDomain::AreCompatible(OutputPin->DataDomain, InputPin->DataDomain, OutReason);
 }
 
