Optimus: Fix build break.

#rnx

diff --git a/Engine/Plugins/Animation/DeformerGraph/Source/OptimusCore/Private/OptimusDataDomain.cpp b/Engine/Plugins/Animation/DeformerGraph/Source/OptimusCore/Private/OptimusDataDomain.cpp
--- a/Engine/Plugins/Animation/DeformerGraph/Source/OptimusCore/Private/OptimusDataDomain.cpp
+++ b/Engine/Plugins/Animation/DeformerGraph/Source/OptimusCore/Private/OptimusDataDomain.cpp
@@ -171,7 +171,7 @@ FString FOptimusDataDomain::GetDisplayName() const
 				Names.Add(DomainLevelName.ToString());
 			}
 			FString DomainName = FString::Join(Names, *FString(UTF8TEXT(" › ")));
-			if (InDataDomain.Multiplier > 1)
+			if (Multiplier > 1)
 			{
 				DomainName += FString::Printf(TEXT(" x %d"), Multiplier);
 			}
