Optimus: Enable use of expression for ExecutionDomain

diff --git a/Engine/Plugins/Animation/DeformerGraph/Source/OptimusCore/Private/DataInterfaces/OptimusDataInterfaceCustomComputeKernel.cpp b/Engine/Plugins/Animation/DeformerGraph/Source/OptimusCore/Private/DataInterfaces/OptimusDataInterfaceCustomComputeKernel.cpp
--- a/Engine/Plugins/Animation/DeformerGraph/Source/OptimusCore/Private/DataInterfaces/OptimusDataInterfaceCustomComputeKernel.cpp
+++ b/Engine/Plugins/Animation/DeformerGraph/Source/OptimusCore/Private/DataInterfaces/OptimusDataInterfaceCustomComputeKernel.cpp
@@ -3,6 +3,7 @@
 #include "OptimusDataInterfaceCustomComputeKernel.h"
 
 #include "OptimusComponentSource.h"
+#include "OptimusDeformerInstance.h"
 #include "OptimusExpressionEvaluator.h"
 #include "ShaderParameterMetadataBuilder.h"
 #include "ComputeFramework/ComputeMetadataBuilder.h"
@@ -15,19 +16,6 @@
 
 const FString UOptimusCustomComputeKernelDataInterface::NumThreadsReservedName = TEXT("NumThreads");
 
-void UOptimusCustomComputeKernelDataInterface::InitFromKernelNode(const UOptimusNode_CustomComputeKernel* InKernelNode)
-{
-	NumThreadsExpression = InKernelNode->ExecutionDomain.Name.ToString();
-
-	TSet<UOptimusComponentSourceBinding*> PrimaryBinding = InKernelNode->GetGroupComponentSourceBindings(InKernelNode->GetPrimaryGroupPin());
-
-	// Kernel node's ValidateForCompile() should guaranteed the existence of unique primary binding
-	if (ensure(PrimaryBinding.Num() == 1))
-	{
-		ComponentSourceBinding = PrimaryBinding.Array()[0];
-	}
-}
-
 void UOptimusCustomComputeKernelDataInterface::GetSupportedInputs(TArray<FShaderFunctionDefinition>& OutFunctions) const
 {
 	OutFunctions.AddDefaulted_GetRef()
@@ -89,6 +77,12 @@ UComputeDataProvider* UOptimusCustomComputeKernelDataInterface::CreateDataProvid
 	return Provider;
 }
 
+void UOptimusCustomComputeKernelDataInterface::SetExecutionDomainConstant(
+	const FOptimusConstantIdentifier& InExecutionDomainConstantIdentifier)
+{
+	ExecutionDomainConstantIdentifier = InExecutionDomainConstantIdentifier;
+}
+
 void UOptimusCustomComputeKernelDataProvider::InitFromDataInterface(const UOptimusCustomComputeKernelDataInterface* InDataInterface, const UObject* InBinding)
 {
 	WeakComponent = Cast<UActorComponent>(InBinding);
@@ -109,7 +103,50 @@ FComputeDataProviderRenderProxy* UOptimusCustomComputeKernelDataProvider::GetRen
 	return Proxy;
 }
 
-bool UOptimusCustomComputeKernelDataProvider::GetInvocationThreadCounts(
+void UOptimusCustomComputeKernelDataProvider::SetDeformerInstance(UOptimusDeformerInstance* InInstance)
+{
+	DeformerInstance = InInstance;
+}
+
+UOptimusDeformerInstance* UOptimusCustomComputeKernelDataProvider::GetDeformerInstance() const
+{
+	return DeformerInstance;
+}
+
+bool UOptimusCustomComputeKernelDataProvider::GetInvocationThreadCounts(TArray<int32>& OutInvocationThreadCount,
+                                                                        int32& OutTotalThreadCount) const
+{
+	if (!WeakDataInterface.IsValid())
+	{
+		return false;
+	}
+
+	const FOptimusConstantIdentifier& ExecutionDomainIdentifier = WeakDataInterface->ExecutionDomainConstantIdentifier;
+	
+	if (!ExecutionDomainIdentifier.IsValid())
+	{
+		return GetInvocationThreadCounts_DEPRECATED(OutInvocationThreadCount, OutTotalThreadCount);
+	}
+
+	TArray<float> Values = DeformerInstance->GetConstantValuePerInvocation(ExecutionDomainIdentifier);
+
+	// Can happen if the bound component does not have actual data, like when there is no preview mesh
+	if (Values.Num() == 0)
+	{
+		return false;
+	}
+	
+	OutInvocationThreadCount.Reset(Values.Num());
+	for (const float& Value : Values)
+	{
+		OutInvocationThreadCount.Add(static_cast<int32>(Value));
+	}
+
+	return true;
+	
+}
+
+bool UOptimusCustomComputeKernelDataProvider::GetInvocationThreadCounts_DEPRECATED(
 	TArray<int32>& OutInvocationThreadCount,
 	int32& OutTotalThreadCount
 	) const
@@ -119,13 +156,15 @@ bool UOptimusCustomComputeKernelDataProvider::GetInvocationThreadCounts(
 		return false;
 	}
 
-	if (!WeakDataInterface->ComponentSourceBinding.IsValid())
+PRAGMA_DISABLE_DEPRECATION_WARNINGS
+	if (!WeakDataInterface->ComponentSourceBinding_DEPRECATED.IsValid())
 	{
 		return false;
 	}
 	
-	const UOptimusComponentSource* ComponentSource = WeakDataInterface->ComponentSourceBinding->GetComponentSource();
-	const FString& NumThreadsExpression = WeakDataInterface->NumThreadsExpression;
+	const UOptimusComponentSource* ComponentSource = WeakDataInterface->ComponentSourceBinding_DEPRECATED->GetComponentSource();
+	const FString& NumThreadsExpression = WeakDataInterface->NumThreadsExpression_DEPRECATED;
+PRAGMA_ENABLE_DEPRECATION_WARNINGS
 	
 	if (!WeakComponent.IsValid() || !ComponentSource)
 	{
diff --git a/Engine/Plugins/Animation/DeformerGraph/Source/OptimusCore/Private/DataInterfaces/OptimusDataInterfaceCustomComputeKernel.h b/Engine/Plugins/Animation/DeformerGraph/Source/OptimusCore/Private/DataInterfaces/OptimusDataInterfaceCustomComputeKernel.h
--- a/Engine/Plugins/Animation/DeformerGraph/Source/OptimusCore/Private/DataInterfaces/OptimusDataInterfaceCustomComputeKernel.h
+++ b/Engine/Plugins/Animation/DeformerGraph/Source/OptimusCore/Private/DataInterfaces/OptimusDataInterfaceCustomComputeKernel.h
@@ -2,7 +2,10 @@
 
 #pragma once
 
+#include "IOptimusComputeKernelDataInterface.h"
+#include "IOptimusDeformerInstanceAccessor.h"
 #include "OptimusComputeDataInterface.h"
+#include "OptimusConstant.h"
 #include "OptimusDataType.h"
 #include "ComputeFramework/ComputeDataProvider.h"
 
@@ -13,15 +16,15 @@ class UOptimusComponentSource;
 class UOptimusComponentSourceBinding;
 
 UCLASS(Category = ComputeFramework)
-class OPTIMUSCORE_API UOptimusCustomComputeKernelDataInterface : public UComputeDataInterface
+class OPTIMUSCORE_API UOptimusCustomComputeKernelDataInterface :
+	public UComputeDataInterface,
+	public IOptimusComputeKernelDataInterface
 {
 	GENERATED_BODY()
 
 public:
 	static const FString NumThreadsReservedName;
 	
-	void InitFromKernelNode(const UOptimusNode_CustomComputeKernel* InKernelNode);
-	
 	//~ Begin UComputeDataInterface Interface
 	TCHAR const* GetClassName() const override { return TEXT("CustomComputeKernelData"); }
 	bool IsExecutionInterface() const override { return true; }
@@ -32,19 +35,28 @@ public:
 	void GetHLSL(FString& OutHLSL, FString const& InDataInterfaceName) const override;
 	UComputeDataProvider* CreateDataProvider(TObjectPtr<UObject> InBinding, uint64 InInputMask, uint64 InOutputMask) const override;
 	//~ End UComputeDataInterface Interface
+
+	//~ Begin IOptimusComputeKernelDataInterface Interface
+	void SetExecutionDomainConstant(const FOptimusConstantIdentifier& InExecutionDomainConstantIdentifier) override;
+	//~ End IOptimusComputeKernelDataInterface Interface
 	
 	UPROPERTY()
-	TWeakObjectPtr<const UOptimusComponentSourceBinding> ComponentSourceBinding;
+	TWeakObjectPtr<const UOptimusComponentSourceBinding> ComponentSourceBinding_DEPRECATED;
 	
 	UPROPERTY()
-	FString NumThreadsExpression;
+	FString NumThreadsExpression_DEPRECATED;
+
+	UPROPERTY()
+	FOptimusConstantIdentifier ExecutionDomainConstantIdentifier;
 
 protected:
 };
 
 /** Compute Framework Data Provider for each custom compute kernel. */
 UCLASS()
-class UOptimusCustomComputeKernelDataProvider : public UComputeDataProvider
+class UOptimusCustomComputeKernelDataProvider :
+	public UComputeDataProvider,
+	public IOptimusDeformerInstanceAccessor
 {
 	GENERATED_BODY()
 
@@ -55,6 +67,11 @@ public:
 	FComputeDataProviderRenderProxy* GetRenderProxy() override;
 	//~ End UComputeDataProvider Interface
 
+	//~ Begin IOptimusDeformerInstanceAccessor Interface
+	void SetDeformerInstance(UOptimusDeformerInstance* InInstance) override;
+	UOptimusDeformerInstance* GetDeformerInstance() const override;
+	//~ End IOptimusDeformerInstanceAccessor Interface
+	
 	TWeakObjectPtr<const UActorComponent> WeakComponent = nullptr;
 
 	TWeakObjectPtr<const UOptimusCustomComputeKernelDataInterface> WeakDataInterface = nullptr;
@@ -64,6 +81,16 @@ protected:
 		TArray<int32>& OutInvocationThreadCount,
 		int32& OutTotalThreadCount
 		) const;
+	
+	bool GetInvocationThreadCounts_DEPRECATED(
+		TArray<int32>& OutInvocationThreadCount,
+		int32& OutTotalThreadCount
+		) const;
+
+		
+private:
+	UPROPERTY()
+	TObjectPtr<UOptimusDeformerInstance> DeformerInstance = nullptr;
 };
 
 class FOptimusCustomComputeKernelDataProviderProxy : public FComputeDataProviderRenderProxy
diff --git a/Engine/Plugins/Animation/DeformerGraph/Source/OptimusCore/Private/Nodes/OptimusNode_ComputeKernelBase.cpp b/Engine/Plugins/Animation/DeformerGraph/Source/OptimusCore/Private/Nodes/OptimusNode_ComputeKernelBase.cpp
--- a/Engine/Plugins/Animation/DeformerGraph/Source/OptimusCore/Private/Nodes/OptimusNode_ComputeKernelBase.cpp
+++ b/Engine/Plugins/Animation/DeformerGraph/Source/OptimusCore/Private/Nodes/OptimusNode_ComputeKernelBase.cpp
@@ -2,6 +2,7 @@
 
 #include "OptimusNode_ComputeKernelBase.h"
 
+#include "IOptimusComputeKernelDataInterface.h"
 #include "IOptimusDataInterfaceProvider.h"
 #include "OptimusCoreModule.h"
 #include "OptimusNodeGraph.h"
@@ -15,6 +16,7 @@
 #include "OptimusHelpers.h"
 #include "OptimusKernelSource.h"
 #include "OptimusConstant.h"
+#include "OptimusExecutionDomain.h"
 
 #include UE_INLINE_GENERATED_CPP_BY_NAME(OptimusNode_ComputeKernelBase)
 
@@ -57,10 +59,10 @@ FOptimus_ComputeKernelResult UOptimusNode_ComputeKernelBase::CreateComputeKernel
 	const TArray<const UOptimusNode *>& InValueNodes,
 	const UComputeDataInterface* InGraphDataInterface,
 	const UOptimusComponentSourceBinding* InGraphDataComponentBinding,
-	const UComputeDataInterface* InKernelDataInterface,
+	UComputeDataInterface* InOutKernelDataInterface,
 	FOptimus_InterfaceBindingMap& OutInputDataBindings,
 	FOptimus_InterfaceBindingMap& OutOutputDataBindings,
-	FOptimusConstantContainer& OutConstantContainer
+	FOptimusKernelConstantContainer& OutKernelConstantContainer
 ) const
 {
 	// Maps friendly name to unique name for each struct type
@@ -143,7 +145,6 @@ FOptimus_ComputeKernelResult UOptimusNode_ComputeKernelBase::CreateComputeKernel
 		return ReturnError(LOCTEXT("ZeroOrMultiplePrimaryBindings", "Primary Group has zero or more than one Component Bindings"));
 	}
 
-	FOptimusKernelConstantContainer& OutKernelConstantContainer = OutConstantContainer.AddContainerForKernel();
 
 	UOptimusKernelSource* KernelSource = NewObject<UOptimusKernelSource>(InKernelSourceOuter);
 
@@ -181,12 +182,14 @@ FOptimus_ComputeKernelResult UOptimusNode_ComputeKernelBase::CreateComputeKernel
 		}
 	}
 
-	if (ensure(InKernelDataInterface))
+	if (ensure(InOutKernelDataInterface))
 	{
 		BindKernelDataInterfaceForComputeKernel(
 			PrimaryBindings.Array()[0],
-			InKernelDataInterface,
-			KernelSource, OutInputDataBindings);
+			InOutKernelDataInterface,
+			KernelSource,
+			OutInputDataBindings,
+			OutKernelConstantContainer);
 	}
 
 	FString CookedSource;
@@ -792,16 +795,32 @@ void UOptimusNode_ComputeKernelBase::ProcessOutputPinForComputeKernel(
 
 void UOptimusNode_ComputeKernelBase::BindKernelDataInterfaceForComputeKernel(
 	const UOptimusComponentSourceBinding* InKernelPrimaryComponentSourceBinding,
-	const UComputeDataInterface* InKernelDataInterface,
+	UComputeDataInterface* InOutKernelDataInterface,
 	UOptimusKernelSource* InKernelSource,
-	FOptimus_InterfaceBindingMap& OutInputDataBindings
+	FOptimus_InterfaceBindingMap& OutInputDataBindings,
+	FOptimusKernelConstantContainer& OutKernelConstantContainer
 	) const
 {
-	check(InKernelDataInterface);
+	check(InOutKernelDataInterface);
 	check(InKernelPrimaryComponentSourceBinding);
+
+	FOptimusConstantIdentifier ExecutionDomainIdentifier(this, NAME_None, TEXT("@ExecutionDomain"));
+	
+	OutKernelConstantContainer.AddToKernelContainer({
+		ExecutionDomainIdentifier,
+		{GetExecutionDomain().AsExpression()},
+		InKernelPrimaryComponentSourceBinding->GetIndex(),
+		EOptimusConstantType::Output});
+	
+	IOptimusComputeKernelDataInterface* KernelDataInterface = Cast<IOptimusComputeKernelDataInterface>(InOutKernelDataInterface);
+	if (ensure(KernelDataInterface))
+	{
+		KernelDataInterface->SetExecutionDomainConstant(ExecutionDomainIdentifier);
+	}
+
 	
 	TArray<FShaderFunctionDefinition> ReadFunctions;
-	InKernelDataInterface->GetSupportedInputs(ReadFunctions);	
+	InOutKernelDataInterface->GetSupportedInputs(ReadFunctions);	
 
 	// Simply grab everything the kernel data interface has to offer
 	for (int32 FuncIndex = 0; FuncIndex < ReadFunctions.Num(); FuncIndex++)
@@ -815,7 +834,7 @@ void UOptimusNode_ComputeKernelBase::BindKernelDataInterfaceForComputeKernel(
 		}
 
 		FOptimus_InterfaceBinding InterfaceBinding;
-		InterfaceBinding.DataInterface = InKernelDataInterface;
+		InterfaceBinding.DataInterface = InOutKernelDataInterface;
 		InterfaceBinding.ComponentBinding = InKernelPrimaryComponentSourceBinding;
 		InterfaceBinding.DataInterfaceBindingIndex = FuncIndex;
 		InterfaceBinding.BindingFunctionName = FuncDef.Name;
diff --git a/Engine/Plugins/Animation/DeformerGraph/Source/OptimusCore/Private/Nodes/OptimusNode_ComputeKernelBase.h b/Engine/Plugins/Animation/DeformerGraph/Source/OptimusCore/Private/Nodes/OptimusNode_ComputeKernelBase.h
--- a/Engine/Plugins/Animation/DeformerGraph/Source/OptimusCore/Private/Nodes/OptimusNode_ComputeKernelBase.h
+++ b/Engine/Plugins/Animation/DeformerGraph/Source/OptimusCore/Private/Nodes/OptimusNode_ComputeKernelBase.h
@@ -5,6 +5,7 @@
 #include "IOptimusComputeKernelProvider.h"
 #include "OptimusDataDomain.h"
 #include "OptimusDataType.h"
+#include "OptimusExecutionDomain.h"
 #include "OptimusNode.h"
 
 #include "OptimusNode_ComputeKernelBase.generated.h"
@@ -42,15 +43,15 @@ public:
 		const TArray<const UOptimusNode*>& InValueNodes,
 		const UComputeDataInterface* InGraphDataInterface,
 		const UOptimusComponentSourceBinding* InGraphDataComponentBinding,
-		const UComputeDataInterface* InKernelDataInterface,
+		UComputeDataInterface* InOutKernelDataInterface,
 		FOptimus_InterfaceBindingMap& OutInputDataBindings,
 		FOptimus_InterfaceBindingMap& OutOutputDataBindings,
-		FOptimusConstantContainer& OutConstantContainer
+		FOptimusKernelConstantContainer& OutKernelConstantContainer
 	) const override;
 
-	FName GetExecutionDomain() const override PURE_VIRTUAL(UOptimusNode_ComputeKernelBase::GetExecutionDomain, return NAME_None; );
+	FOptimusExecutionDomain GetExecutionDomain() const override PURE_VIRTUAL(UOptimusNode_ComputeKernelBase::GetExecutionDomain, return {}; );
 	const UOptimusNodePin* GetPrimaryGroupPin() const override PURE_VIRTUAL(UOptimusNode_ComputeKernelBase::GetPrimaryGroupPin, return {}; );  
-	UComputeDataInterface* GetKernelDataInterface(UObject* InOuter) const override PURE_VIRTUAL(UOptimusNode_ComputeKernelBase::GetKernelDataInterface, return {}; );
+	UComputeDataInterface* MakeKernelDataInterface(UObject* InOuter) const override PURE_VIRTUAL(UOptimusNode_ComputeKernelBase::MakeKernelDataInterface, return {}; );
 	
 	// -- UOptimusNode overrides
 	TOptional<FText> ValidateForCompile() const override;
@@ -107,9 +108,10 @@ private:
 
 	void BindKernelDataInterfaceForComputeKernel(
 		const UOptimusComponentSourceBinding* InKernelPrimaryComponentSourceBinding,
-		const UComputeDataInterface* InKernelDataInterface,
+		UComputeDataInterface* InOutKernelDataInterface,
 		UOptimusKernelSource* InKernelSource,
-		FOptimus_InterfaceBindingMap& OutInputDataBindings
+		FOptimus_InterfaceBindingMap& OutInputDataBindings,
+		FOptimusKernelConstantContainer& OutKernelConstantContainer
 		) const;
 };
 
diff --git a/Engine/Plugins/Animation/DeformerGraph/Source/OptimusCore/Private/Nodes/OptimusNode_ComputeKernelFunction.cpp b/Engine/Plugins/Animation/DeformerGraph/Source/OptimusCore/Private/Nodes/OptimusNode_ComputeKernelFunction.cpp
--- a/Engine/Plugins/Animation/DeformerGraph/Source/OptimusCore/Private/Nodes/OptimusNode_ComputeKernelFunction.cpp
+++ b/Engine/Plugins/Animation/DeformerGraph/Source/OptimusCore/Private/Nodes/OptimusNode_ComputeKernelFunction.cpp
@@ -204,9 +204,9 @@ void UOptimusNode_ComputeKernelFunction::ConstructNode()
 }
 
 
-FName UOptimusNode_ComputeKernelFunction::GetExecutionDomain() const
+FOptimusExecutionDomain UOptimusNode_ComputeKernelFunction::GetExecutionDomain() const
 {
-	return GetGeneratorClass()->ExecutionDomain.Name;
+	return GetGeneratorClass()->ExecutionDomain;
 }
 
 
diff --git a/Engine/Plugins/Animation/DeformerGraph/Source/OptimusCore/Private/Nodes/OptimusNode_ComputeKernelFunction.h b/Engine/Plugins/Animation/DeformerGraph/Source/OptimusCore/Private/Nodes/OptimusNode_ComputeKernelFunction.h
--- a/Engine/Plugins/Animation/DeformerGraph/Source/OptimusCore/Private/Nodes/OptimusNode_ComputeKernelFunction.h
+++ b/Engine/Plugins/Animation/DeformerGraph/Source/OptimusCore/Private/Nodes/OptimusNode_ComputeKernelFunction.h
@@ -2,9 +2,9 @@
 
 #pragma once
 
-#include "IOptimusExecutionDomainProvider.h"
 #include "OptimusNode_ComputeKernelBase.h"
 #include "OptimusBindingTypes.h"
+#include "OptimusExecutionDomain.h"
 
 #include "OptimusNode_ComputeKernelFunction.generated.h"
 
@@ -76,10 +76,10 @@ public:
 	void ConstructNode() override;
 	
 	// IOptimusComputeKernelProvider
-	FName GetExecutionDomain() const override; 
+	FOptimusExecutionDomain GetExecutionDomain() const override;
 	const UOptimusNodePin* GetPrimaryGroupPin() const override { return {}; }
-	UComputeDataInterface* GetKernelDataInterface(UObject* InOuter) const override { return nullptr; };
-	
+	UComputeDataInterface* MakeKernelDataInterface(UObject* InOuter) const override { return nullptr; };
+
 private:
 	UOptimusNode_ComputeKernelFunctionGeneratorClass *GetGeneratorClass() const;
 };
diff --git a/Engine/Plugins/Animation/DeformerGraph/Source/OptimusCore/Private/Nodes/OptimusNode_CustomComputeKernel.cpp b/Engine/Plugins/Animation/DeformerGraph/Source/OptimusCore/Private/Nodes/OptimusNode_CustomComputeKernel.cpp
--- a/Engine/Plugins/Animation/DeformerGraph/Source/OptimusCore/Private/Nodes/OptimusNode_CustomComputeKernel.cpp
+++ b/Engine/Plugins/Animation/DeformerGraph/Source/OptimusCore/Private/Nodes/OptimusNode_CustomComputeKernel.cpp
@@ -98,17 +98,20 @@ FString UOptimusNode_CustomComputeKernel::GetKernelSourceText() const
 	return GetCookedKernelSource(GetPathName(), ShaderSource.ShaderText, KernelName.ToString(), GroupSize);
 }
 
+FOptimusExecutionDomain UOptimusNode_CustomComputeKernel::GetExecutionDomain() const
+{
+	return ExecutionDomain;
+}
+
 const UOptimusNodePin* UOptimusNode_CustomComputeKernel::GetPrimaryGroupPin() const
 {
 	return GetPrimaryGroupPin_Internal();
 }
 
-UComputeDataInterface* UOptimusNode_CustomComputeKernel::GetKernelDataInterface(UObject* InOuter) const
+UComputeDataInterface* UOptimusNode_CustomComputeKernel::MakeKernelDataInterface(UObject* InOuter) const
 {
 	UOptimusCustomComputeKernelDataInterface* KernelDataInterface = NewObject<UOptimusCustomComputeKernelDataInterface>(InOuter);
-	
-	KernelDataInterface->InitFromKernelNode(this);
-	
+
 	return KernelDataInterface;
 }
 
@@ -880,9 +883,13 @@ void UOptimusNode_CustomComputeKernel::PropertyArrayItemAdded(
 		else
 		{
 			// Pick a suitable fallback for output pins.
-			FName DomainName = GetExecutionDomain();
-			if (DomainName.IsNone())
+			if (ExecutionDomain.IsDefined())
 			{
+				InBinding.DataDomain = FOptimusDataDomain(ExecutionDomain);
+			}
+			else
+			{
+				FName DomainName = NAME_None;
 				if (TArray<FName> ExecutionDomains = GetExecutionDomains(); !ExecutionDomains.IsEmpty())
 				{
 					DomainName = ExecutionDomains[0];
@@ -892,8 +899,9 @@ void UOptimusNode_CustomComputeKernel::PropertyArrayItemAdded(
 					// FIXME: There should be a generic mechanism to get the most suitable default. 
 					DomainName = UOptimusSkeletalMeshComponentSource::Domains::Vertex;
 				}
+				
+				InBinding.DataDomain = FOptimusDataDomain(TArray<FName>({DomainName}));
 			}
-			InBinding.DataDomain = FOptimusDataDomain({DomainName});
 		}
 
 		AddPin(InBinding.Name, InDirection, InBinding.DataDomain, InBinding.DataType, InBeforePin, InGroupPin);
@@ -1637,8 +1645,7 @@ UOptimusNodePin* UOptimusNode_CustomComputeKernel::GetPrimaryGroupPin_Internal()
 
 void UOptimusNode_CustomComputeKernel::PostLoadExtractExecutionDomain()
 {
-	ExecutionDomain.Name = NAME_None;
-		
+	ExecutionDomain = {};
 	// Check if there's an execution node connected and grab the domain from it.
 	FOptimusDataTypeHandle IntVector3Type = FOptimusDataTypeRegistry::Get().FindType(Optimus::GetTypeName(TBaseStructure<FIntVector3>::Get()));
 
@@ -1654,7 +1661,7 @@ void UOptimusNode_CustomComputeKernel::PostLoadExtractExecutionDomain()
 					if (const IOptimusDeprecatedExecutionDataInterface* ExecDataInterface = Cast<IOptimusDeprecatedExecutionDataInterface>(DataInterfaceNode->GetDataInterface(GetTransientPackage())))
 					{
 						FName Domain = ExecDataInterface->GetSelectedExecutionDomainName();
-						ExecutionDomain.Name = Domain;
+						ExecutionDomain = {Domain};
 					}
 				}
 			}
diff --git a/Engine/Plugins/Animation/DeformerGraph/Source/OptimusCore/Private/Nodes/OptimusNode_CustomComputeKernel.h b/Engine/Plugins/Animation/DeformerGraph/Source/OptimusCore/Private/Nodes/OptimusNode_CustomComputeKernel.h
--- a/Engine/Plugins/Animation/DeformerGraph/Source/OptimusCore/Private/Nodes/OptimusNode_CustomComputeKernel.h
+++ b/Engine/Plugins/Animation/DeformerGraph/Source/OptimusCore/Private/Nodes/OptimusNode_CustomComputeKernel.h
@@ -16,7 +16,7 @@
 
 
 enum class EOptimusNodePinDirection : uint8;
-
+struct FOptimusExecutionDomain;
 
 USTRUCT()
 struct FOptimusSecondaryInputBindingsGroup
@@ -59,9 +59,9 @@ public:
 	TArray<TObjectPtr<UComputeSource>> GetAdditionalSources() const override { return AdditionalSources; }
 
 	// IOptimusComputeKernelProvider
-	FName GetExecutionDomain() const override { return ExecutionDomain.Name; } 
+	FOptimusExecutionDomain GetExecutionDomain() const override;
 	const UOptimusNodePin* GetPrimaryGroupPin() const override;
-	UComputeDataInterface* GetKernelDataInterface(UObject* InOuter) const override;
+	UComputeDataInterface* MakeKernelDataInterface(UObject* InOuter) const override;
 
 #if WITH_EDITOR
 	// IOptimusShaderTextProvider overrides
diff --git a/Engine/Plugins/Animation/DeformerGraph/Source/OptimusCore/Private/OptimusConstant.cpp b/Engine/Plugins/Animation/DeformerGraph/Source/OptimusCore/Private/OptimusConstant.cpp
--- a/Engine/Plugins/Animation/DeformerGraph/Source/OptimusCore/Private/OptimusConstant.cpp
+++ b/Engine/Plugins/Animation/DeformerGraph/Source/OptimusCore/Private/OptimusConstant.cpp
@@ -18,8 +18,8 @@ void FOptimusKernelConstantContainer::AddToKernelContainer(const FOptimusConstan
 		InConstant.Type == EOptimusConstantType::Input ?
 			InputConstants :OutputConstants;
 
-	int32 ConstantIndex = Constants.Add(InConstant);
-
+	Constants.Add(InConstant);	
+	
 	GroupNameToBindingIndex.FindOrAdd(InConstant.Identifier.GroupName) = InConstant.ComponentBindingIndex;
 }
 
diff --git a/Engine/Plugins/Animation/DeformerGraph/Source/OptimusCore/Private/OptimusDataDomain.cpp b/Engine/Plugins/Animation/DeformerGraph/Source/OptimusCore/Private/OptimusDataDomain.cpp
--- a/Engine/Plugins/Animation/DeformerGraph/Source/OptimusCore/Private/OptimusDataDomain.cpp
+++ b/Engine/Plugins/Animation/DeformerGraph/Source/OptimusCore/Private/OptimusDataDomain.cpp
@@ -2,6 +2,7 @@
 
 #include "OptimusDataDomain.h"
 
+#include "OptimusExecutionDomain.h"
 #include "OptimusExpressionEvaluator.h"
 #include "OptimusObjectVersion.h"
 #include "String/ParseTokens.h"
@@ -30,6 +31,24 @@ FString Optimus::FormatDimensionNames(const TArray<FName>& InNames)
 }
 
 
+FOptimusDataDomain::FOptimusDataDomain(const FOptimusExecutionDomain& InExecutionDomain)
+{
+	if (InExecutionDomain.Type == EOptimusExecutionDomainType::DomainName)
+	{
+		Type = EOptimusDataDomainType::Dimensional;
+		DimensionNames = {InExecutionDomain.Name};
+	}
+	else if (InExecutionDomain.Type == EOptimusExecutionDomainType::Expression)
+	{
+		Type = EOptimusDataDomainType::Expression;
+		Expression = InExecutionDomain.Expression;
+	}
+	else
+	{
+		checkNoEntry();
+	}
+}
+
 FString FOptimusDataDomain::ToString() const
 {
 	switch(Type)
diff --git a/Engine/Plugins/Animation/DeformerGraph/Source/OptimusCore/Private/OptimusDeformer.cpp b/Engine/Plugins/Animation/DeformerGraph/Source/OptimusCore/Private/OptimusDeformer.cpp
--- a/Engine/Plugins/Animation/DeformerGraph/Source/OptimusCore/Private/OptimusDeformer.cpp
+++ b/Engine/Plugins/Animation/DeformerGraph/Source/OptimusCore/Private/OptimusDeformer.cpp
@@ -1625,7 +1625,7 @@ UOptimusComputeGraph* UOptimusDeformer::CompileNodeGraphToComputeGraph(
 		}
 		else if (const IOptimusComputeKernelProvider* KernelProvider = Cast<const IOptimusComputeKernelProvider>(ConnectedNode.Node))
 		{
-			UComputeDataInterface* KernelDataInterface = KernelProvider->GetKernelDataInterface(this);
+			UComputeDataInterface* KernelDataInterface = KernelProvider->MakeKernelDataInterface(this);
 
 			KernelDataInterfaceMap.Add(ConnectedNode.Node, KernelDataInterface);
 			
@@ -1714,6 +1714,7 @@ UOptimusComputeGraph* UOptimusDeformer::CompileNodeGraphToComputeGraph(
 			BoundKernel.KernelNodeIndex = InNodeGraph->Nodes.IndexOfByKey(ConnectedNode.Node);
 			BoundKernel.Kernel = NewObject<UComputeKernel>(this);
 
+			FOptimusKernelConstantContainer& KernelConstantContainer = ConstantContainer.AddContainerForKernel();
 			UComputeDataInterface* KernelDataInterface = KernelDataInterfaceMap[ConnectedNode.Node];
 			
 			FOptimus_ComputeKernelResult KernelSourceResult = KernelProvider->CreateComputeKernel(	
@@ -1722,7 +1723,7 @@ UOptimusComputeGraph* UOptimusDeformer::CompileNodeGraphToComputeGraph(
 				ValueNodes,
 				GraphDataInterface, GraphDataComponentBinding,
 				KernelDataInterface,
-				BoundKernel.InputDataBindings, BoundKernel.OutputDataBindings, ConstantContainer
+				BoundKernel.InputDataBindings, BoundKernel.OutputDataBindings, KernelConstantContainer
 			);
 			if (FText* ErrorMessage = KernelSourceResult.TryGet<FText>())
 			{
diff --git a/Engine/Plugins/Animation/DeformerGraph/Source/OptimusCore/Private/OptimusExecutionDomain.cpp b/Engine/Plugins/Animation/DeformerGraph/Source/OptimusCore/Private/OptimusExecutionDomain.cpp
new file mode 100644
--- /dev/null
+++ b/Engine/Plugins/Animation/DeformerGraph/Source/OptimusCore/Private/OptimusExecutionDomain.cpp
@@ -0,0 +1,25 @@
+﻿// Copyright Epic Games, Inc. All Rights Reserved.
+
+#include "OptimusExecutionDomain.h"
+#include UE_INLINE_GENERATED_CPP_BY_NAME(OptimusExecutionDomain)
+
+
+FString FOptimusExecutionDomain::AsExpression() const
+{
+	if (Type==EOptimusExecutionDomainType::DomainName)
+	{
+		return Name.ToString();
+	}
+
+	return Expression;
+}
+
+bool FOptimusExecutionDomain::IsDefined() const
+{
+	if (Type == EOptimusExecutionDomainType::DomainName)
+	{
+		return !Name.IsNone();
+	}
+
+	return !Expression.TrimStartAndEnd().IsEmpty();
+}
diff --git a/Engine/Plugins/Animation/DeformerGraph/Source/OptimusCore/Public/IOptimusComputeKernelDataInterface.h b/Engine/Plugins/Animation/DeformerGraph/Source/OptimusCore/Public/IOptimusComputeKernelDataInterface.h
new file mode 100644
--- /dev/null
+++ b/Engine/Plugins/Animation/DeformerGraph/Source/OptimusCore/Public/IOptimusComputeKernelDataInterface.h
@@ -0,0 +1,26 @@
+﻿// Copyright Epic Games, Inc. All Rights Reserved.
+
+#pragma once
+
+#include "UObject/Interface.h"
+
+#include "IOptimusComputeKernelDataInterface.generated.h"
+
+struct FOptimusConstantIdentifier;
+
+UINTERFACE()
+class OPTIMUSCORE_API UOptimusComputeKernelDataInterface :
+	public UInterface
+{
+	GENERATED_BODY()
+};
+
+/**
+ * Interface that provides a mechanism for compute kernel to setup its kernel data interface
+ */
+class OPTIMUSCORE_API IOptimusComputeKernelDataInterface
+{
+public:
+	GENERATED_BODY()
+	virtual void SetExecutionDomainConstant(const FOptimusConstantIdentifier& InExecutionDomainConstantIdentifier) = 0;
+};
diff --git a/Engine/Plugins/Animation/DeformerGraph/Source/OptimusCore/Public/IOptimusComputeKernelProvider.h b/Engine/Plugins/Animation/DeformerGraph/Source/OptimusCore/Public/IOptimusComputeKernelProvider.h
--- a/Engine/Plugins/Animation/DeformerGraph/Source/OptimusCore/Public/IOptimusComputeKernelProvider.h
+++ b/Engine/Plugins/Animation/DeformerGraph/Source/OptimusCore/Public/IOptimusComputeKernelProvider.h
@@ -15,8 +15,8 @@ class UOptimusKernelSource;
 class UOptimusNode;
 class UOptimusNodePin;
 struct FOptimusPinTraversalContext;
-struct FOptimusConstantContainer;
-
+struct FOptimusKernelConstantContainer;
+struct FOptimusExecutionDomain;
 
 // Maps the data interface's data binding index to the function we would like to have present
 // during kernel compilation to read/write values from/to that data interface's resource.
@@ -81,10 +81,10 @@ public:
 	 * @param InValueNodes
 	 * @param InGraphDataInterface
 	 * @param InGraphDataComponentBinding
-	 * @param InKernelDataInterface
+	 * @param InOutKernelDataInterface
 	 * @param OutInputDataBindings
 	 * @param OutOutputDataBindings
-	 * @param OutConstantContainer
+	 * @param OutKernelConstantContainer
 	 */
 	virtual FOptimus_ComputeKernelResult CreateComputeKernel(
 		UObject* InKernelSourceOuter,
@@ -94,18 +94,18 @@ public:
 		const TArray<const UOptimusNode*>& InValueNodes,
 		const UComputeDataInterface* InGraphDataInterface,
 		const UOptimusComponentSourceBinding* InGraphDataComponentBinding,
-		const UComputeDataInterface* InKernelDataInterface,
+		UComputeDataInterface* InOutKernelDataInterface,
 		FOptimus_InterfaceBindingMap& OutInputDataBindings,
 		FOptimus_InterfaceBindingMap& OutOutputDataBindings,
-		FOptimusConstantContainer& OutConstantContainer
-		) const = 0;
+		FOptimusKernelConstantContainer& OutKernelConstantContainer
+	) const = 0;
 
 	/** Returns the execution domain that this kernel should iterate over */
-	virtual FName GetExecutionDomain() const = 0;
-
+	virtual FOptimusExecutionDomain GetExecutionDomain() const = 0;
+	
 	/** Used for traversing up the graph to identify the component source of the compute kernel */
 	virtual const UOptimusNodePin* GetPrimaryGroupPin() const = 0;
 	
 	/** Each kernel may have its own data interface responsible for passing kernel related data to GPU */
-	virtual UComputeDataInterface* GetKernelDataInterface(UObject* InOuter) const = 0;
+	virtual UComputeDataInterface* MakeKernelDataInterface(UObject* InOuter) const = 0;
 };
diff --git a/Engine/Plugins/Animation/DeformerGraph/Source/OptimusCore/Public/IOptimusExecutionDomainProvider.h b/Engine/Plugins/Animation/DeformerGraph/Source/OptimusCore/Public/IOptimusExecutionDomainProvider.h
--- a/Engine/Plugins/Animation/DeformerGraph/Source/OptimusCore/Public/IOptimusExecutionDomainProvider.h
+++ b/Engine/Plugins/Animation/DeformerGraph/Source/OptimusCore/Public/IOptimusExecutionDomainProvider.h
@@ -6,22 +6,7 @@
 
 #include "IOptimusExecutionDomainProvider.generated.h"
 
-/** A struct to hold onto a single-level domain for controlling a kernel's execution domain. 
-  * The reason it's in a struct is so that we can apply a property panel customization for it
-  * to make it easier to select from a pre-defined list of execution domains.
-*/
-USTRUCT()
-struct OPTIMUSCORE_API FOptimusExecutionDomain
-{
-	GENERATED_BODY()
 
-	FOptimusExecutionDomain() = default;
-	FOptimusExecutionDomain(FName InExecutionDomainName) : Name(InExecutionDomainName) {}
-
-	// The name of the execution domain that this kernel operates on.
-	UPROPERTY(EditAnywhere, Category = Domain)
-	FName Name;
-};
 
 UINTERFACE()
 class OPTIMUSCORE_API UOptimusExecutionDomainProvider :
diff --git a/Engine/Plugins/Animation/DeformerGraph/Source/OptimusCore/Public/OptimusDataDomain.h b/Engine/Plugins/Animation/DeformerGraph/Source/OptimusCore/Public/OptimusDataDomain.h
--- a/Engine/Plugins/Animation/DeformerGraph/Source/OptimusCore/Public/OptimusDataDomain.h
+++ b/Engine/Plugins/Animation/DeformerGraph/Source/OptimusCore/Public/OptimusDataDomain.h
@@ -6,6 +6,8 @@
 
 #include "OptimusDataDomain.generated.h"
 
+struct FOptimusExecutionDomain;
+
 namespace Optimus::DomainName
 {
 	extern OPTIMUSCORE_API const FName Singleton;
@@ -63,6 +65,8 @@ struct OPTIMUSCORE_API FOptimusDataDomain
 		Type(EOptimusDataDomainType::Expression),
 		Expression(MoveTemp(InExpression))
 	{}
+
+	FOptimusDataDomain(const FOptimusExecutionDomain& InExecutionDomain);
 	
 	UPROPERTY(EditAnywhere, Category = DataDomain)
 	EOptimusDataDomainType Type = EOptimusDataDomainType::Dimensional;
diff --git a/Engine/Plugins/Animation/DeformerGraph/Source/OptimusCore/Public/OptimusExecutionDomain.h b/Engine/Plugins/Animation/DeformerGraph/Source/OptimusCore/Public/OptimusExecutionDomain.h
new file mode 100644
--- /dev/null
+++ b/Engine/Plugins/Animation/DeformerGraph/Source/OptimusCore/Public/OptimusExecutionDomain.h
@@ -0,0 +1,55 @@
+﻿// Copyright Epic Games, Inc. All Rights Reserved.
+
+#pragma once
+
+#include "UObject/NameTypes.h"
+
+#include "OptimusExecutionDomain.generated.h"
+
+UENUM()
+enum class EOptimusExecutionDomainType
+{
+	DomainName = 0,	
+	Expression = 1,		
+};
+
+/** A struct to hold onto a single-level domain for controlling a kernel's execution domain. 
+  * The reason it's in a struct is so that we can apply a property panel customization for it
+  * to make it easier to select from a pre-defined list of execution domains.
+*/
+USTRUCT()
+struct OPTIMUSCORE_API FOptimusExecutionDomain
+{
+	GENERATED_BODY()
+
+	FOptimusExecutionDomain() = default;
+	FOptimusExecutionDomain(FName InExecutionDomainName) :
+		Type(EOptimusExecutionDomainType::DomainName),
+		Name(InExecutionDomainName)
+	{}
+	FOptimusExecutionDomain(const FString& InExpression) :
+		Type(EOptimusExecutionDomainType::Expression),
+		Expression(InExpression)
+	{}
+	
+	FString AsExpression() const;
+
+	bool IsDefined() const;
+	
+private:
+	UPROPERTY(EditAnywhere, Category = Domain)
+	EOptimusExecutionDomainType Type = EOptimusExecutionDomainType::DomainName;
+	
+	// The name of the execution domain that this kernel operates on.
+	UPROPERTY(EditAnywhere, Category = Domain)
+	FName Name;
+
+	UPROPERTY(EditAnywhere, Category = Domain)
+	FString Expression;
+
+	friend class FOptimusExecutionDomainCustomization;
+	friend struct FOptimusDataDomain;
+};
+
+
+
diff --git a/Engine/Plugins/Animation/DeformerGraph/Source/OptimusEditor/Private/OptimusDetailsCustomization.cpp b/Engine/Plugins/Animation/DeformerGraph/Source/OptimusEditor/Private/OptimusDetailsCustomization.cpp
--- a/Engine/Plugins/Animation/DeformerGraph/Source/OptimusEditor/Private/OptimusDetailsCustomization.cpp
+++ b/Engine/Plugins/Animation/DeformerGraph/Source/OptimusEditor/Private/OptimusDetailsCustomization.cpp
@@ -27,6 +27,7 @@
 #include "OptimusSource.h"
 #include "OptimusValidatedName.h"
 #include "OptimusValueContainer.h"
+#include "OptimusExecutionDomain.h"
 #include "ScopedTransaction.h"
 #include "Styling/AppStyle.h"
 #include "Styling/SlateIconFinder.h"
@@ -47,6 +48,8 @@
 
 #define LOCTEXT_NAMESPACE "OptimusDetailCustomization"
 
+static const FName ExpressionMarkerName = TEXT("~");
+
 TSharedRef<IPropertyTypeCustomization> FOptimusDataTypeRefCustomization::MakeInstance()
 {
 	return MakeShared<FOptimusDataTypeRefCustomization>();
@@ -243,8 +246,6 @@ void FOptimusExecutionDomainCustomization::CustomizeHeader(
 	IPropertyTypeCustomizationUtils& InCustomizationUtils
 	)
 {
-	TSharedPtr<IPropertyHandle> ContextNameProperty = InPropertyHandle->GetChildHandle(GET_MEMBER_NAME_CHECKED(FOptimusExecutionDomain, Name));
-
 	TArray<UObject*> OwningObjects;
 	InPropertyHandle->GetOuterObjects(OwningObjects);
 	
@@ -256,37 +257,171 @@ void FOptimusExecutionDomainCustomization::CustomizeHeader(
 	]
 	.ValueContent()
 	[
-		SAssignNew(ComboBox, SComboBox<FName>)
+		SNew(SWidgetSwitcher)
+		.WidgetIndex_Lambda([InPropertyHandle]()
+		{
+			// Show the expression box if the domain is an expression (and all match), otherwise show the domain name drop down.
+			if (TOptional<FOptimusExecutionDomain> DataDomain = TryGetSingleExecutionDomain(InPropertyHandle, DomainType|DomainExpression))
+			{
+				return DataDomain->Type == EOptimusExecutionDomainType::Expression ? 1 : 0;
+			}
+			return 0;
+		})
+		+ SWidgetSwitcher::Slot()
+		[
+			SAssignNew(ComboBox, SComboBox<FName>)
 			.OptionsSource(&ContextNames)
 			.IsEnabled_Lambda([InPropertyHandle]() -> bool
 			{
 				return InPropertyHandle->IsEditable();
 			})
-			.OnGenerateWidget_Lambda([](FName InName)
+			.OnGenerateWidget_Lambda([this](FName InName)
 			{
-				const FText NameText = InName.IsNone() ? LOCTEXT("NoneName", "<None>") : FText::FromName(InName);
 				return SNew(STextBlock)
-					.Text(NameText)
+					.Text(FormatContextName(InName))
 					.Font(IPropertyTypeCustomizationUtils::GetRegularFont());
 			})
-			.OnSelectionChanged_Lambda([ContextNameProperty](FName InName, ESelectInfo::Type)
+			.OnSelectionChanged_Lambda([InPropertyHandle, this](FName InName, ESelectInfo::Type InSelectType)
 			{
-				ContextNameProperty->SetValue(InName);
+				// This is for the initial selection.
+				if (InSelectType == ESelectInfo::Direct)
+				{
+					return;
+				}
+
+				FOptimusExecutionDomain ExecutionDomain;
+				if (InName == ExpressionMarkerName)
+				{
+					if (TOptional<FOptimusExecutionDomain> ExecutionDomainOpt = TryGetSingleExecutionDomain(InPropertyHandle))
+					{
+						ExecutionDomain = *ExecutionDomainOpt;
+						
+						ExecutionDomain.Expression = ExecutionDomain.Name.ToString();
+					}
+					else
+					{
+						ExecutionDomain.Expression = FString();
+					}
+					ExecutionDomain.Type = EOptimusExecutionDomainType::Expression;
+				}
+				else
+				{
+					ExecutionDomain.Type = EOptimusExecutionDomainType::DomainName;
+					ExecutionDomain.Name = InName;
+				}
+
+				SetExecutionDomain(InPropertyHandle, ExecutionDomain);
 			})
 			.OnComboBoxOpening(this, &FOptimusExecutionDomainCustomization::UpdateContextNames)
 			[
 				SNew(STextBlock)
 				.Font(IPropertyTypeCustomizationUtils::GetRegularFont())
-				.Text_Lambda([ContextNameProperty]()
+				.Text_Lambda([InPropertyHandle, this]()
 				{
-					FName Name;
-					ContextNameProperty->GetValue(Name);
-					return Name.IsNone() ? LOCTEXT("NoneName", "<None>") : FText::FromName(Name);
+					if (TOptional<FOptimusExecutionDomain> ExecutionDomain = TryGetSingleExecutionDomain(InPropertyHandle))
+					{
+						return FormatContextName(ExecutionDomain->Name);
+					}
+
+					return LOCTEXT("MultipleValues", "Multiple Values");
+				})
+			]
+		]
+		+ SWidgetSwitcher::Slot()
+		[
+			SNew(SHorizontalBox)
+			+ SHorizontalBox::Slot()
+			.FillWidth(1.0)
+			.VAlign(VAlign_Center)
+			[
+				SAssignNew(ExpressionTextBox, SEditableTextBox)
+				.IsEnabled_Lambda([InPropertyHandle]() -> bool
+				{
+					return InPropertyHandle->IsEditable();
+				})
+				.Font(IPropertyTypeCustomizationUtils::GetRegularFont())
+				.HintText(LOCTEXT("ExpressionHint", "Expression..."))
+				.Text_Lambda([InPropertyHandle, this]() -> FText
+				{
+					if (TOptional<FOptimusExecutionDomain> ExecutionDomain = TryGetSingleExecutionDomain(InPropertyHandle, DomainType|DomainExpression))
+					{
+						return FText::FromString(ExecutionDomain->Expression);
+					}
+					return FText::GetEmpty();
+				})
+				.OnTextChanged_Lambda([this](const FText& InExpressionText)
+				{
+					using namespace Optimus::Expression;
+					
+					// Verify that the expression is correct.
+					const FString Expression = InExpressionText.ToString();
+					if (TOptional<FParseError> ParseError = FEngine().Verify(Expression))
+					{
+						ExpressionTextBox->SetError(ParseError->Message);
+					}
+					else
+					{
+						ExpressionTextBox->SetError(FString());
+					}
+				})
+				.OnTextCommitted_Lambda([InPropertyHandle, this](const FText& InExpressionText, ETextCommit::Type)
+				{
+					using namespace Optimus::Expression;
+					
+					// Only commit the text if the expression parses.
+					const FString Expression = InExpressionText.ToString();
+					if (!FEngine().Verify(Expression).IsSet())
+					{
+						const FOptimusExecutionDomain ExpressionDomain(Expression);
+						SetExecutionDomain(InPropertyHandle, ExpressionDomain, DomainType|DomainExpression);
+					}
+				})
+			]
+			+ SHorizontalBox::Slot()
+			.AutoWidth()
+			.Padding(5.0, 0.0, 0.0, 0.0)
+			.VAlign(VAlign_Center)
+			.HAlign(HAlign_Right)
+			[
+				SNew(SButton)
+				.ButtonStyle(FAppStyle::Get(), "HoverHintOnly")
+				.ForegroundColor(FAppStyle::GetSlateColor("DefaultForeground"))
+				.ContentPadding(FMargin(2, 2))
+				.OnClicked_Lambda([InPropertyHandle, this]() -> FReply
+				{
+					FOptimusExecutionDomain NamedDomain;
+					NamedDomain.Type = EOptimusExecutionDomainType::DomainName;
+					SetExecutionDomain(InPropertyHandle, NamedDomain, DomainType);
+
+					// Making sure Expression option can be selected again
+					ComboBox->ClearSelection();
+					
+					return FReply::Handled();
 				})
+				.ToolTipText(LOCTEXT("ClearExpression", "Clear Expression"))
+				[
+					SNew(SImage)
+					.Image(FAppStyle::GetBrush(TEXT("Cross")))
+				]
 			]
+		]
 	];
 }
 
+FText FOptimusExecutionDomainCustomization::FormatContextName(FName InName) const
+{
+	if (InName.IsNone())
+	{
+		return LOCTEXT("NoneName", "<None>");
+	}
+	if (InName == ExpressionMarkerName)
+	{
+		return LOCTEXT("ExpressionName", "Expression...");
+	}
+
+	return  FText::FromName(InName);
+}
+
 void FOptimusExecutionDomainCustomization::UpdateContextNames()
 {
 	TSet<FName> CollectedContextNames;
@@ -314,11 +449,81 @@ void FOptimusExecutionDomainCustomization::UpdateContextNames()
 	else
 	{
 		ContextNames.Sort(FNameLexicalLess{});
+		ContextNames.Add(ExpressionMarkerName);
 	}
 	
 	ComboBox->RefreshOptions();
 }
 
+void FOptimusExecutionDomainCustomization::SetExecutionDomain(TSharedRef<IPropertyHandle> InPropertyHandle,
+	const FOptimusExecutionDomain& InExecutionDomain, DomainFlags InSetFlags)
+{
+	FScopedTransaction Transaction(LOCTEXT("SetExecutionDomains", "Set Execution Domain"));
+
+	// Ideally we'd like to match up the raw data with the outers, but I'm not
+	// convinced that there's always 1-to-1 relation.
+	TArray<UObject *> OuterObjects;
+	InPropertyHandle->GetOuterObjects(OuterObjects);
+	for (UObject *OuterObject: OuterObjects)
+	{
+		// Notify the object that is has been modified so that undo/redo works.
+		OuterObject->Modify();
+	}
+				
+	InPropertyHandle->NotifyPreChange();
+	TArray<void*> RawDataPtrs;
+	InPropertyHandle->AccessRawData(RawDataPtrs);
+
+	for (void* RawPtr: RawDataPtrs)
+	{
+		FOptimusExecutionDomain* DstExecutionDomain = static_cast<FOptimusExecutionDomain*>(RawPtr);
+		if (InSetFlags & DomainType)		DstExecutionDomain->Type = InExecutionDomain.Type;
+		if (InSetFlags & DomainName)		DstExecutionDomain->Name = InExecutionDomain.Name;
+		if (InSetFlags & DomainExpression)	DstExecutionDomain->Expression = InExecutionDomain.Expression;
+	}
+
+	InPropertyHandle->NotifyPostChange(EPropertyChangeType::ValueSet);
+}
+
+TOptional<FOptimusExecutionDomain> FOptimusExecutionDomainCustomization::TryGetSingleExecutionDomain(
+	TSharedRef<IPropertyHandle> InPropertyHandle, DomainFlags InCompareFlags, bool bInCheckMultiples)
+{
+	TArray<const void *> RawDataPtrs;
+	InPropertyHandle->AccessRawData(RawDataPtrs);
+
+	bool bItemsAreAllSame = false;
+	const FOptimusExecutionDomain* ComparatorExecutionDomain = nullptr;
+	for (const void* RawPtr: RawDataPtrs)
+	{
+		// During drag & reorder, invalid binding can be created temporarily
+		if (const FOptimusExecutionDomain* ExecutionDomain = static_cast<const FOptimusExecutionDomain*>(RawPtr))
+		{
+			if (!ComparatorExecutionDomain)
+			{
+				ComparatorExecutionDomain = ExecutionDomain;
+				bItemsAreAllSame = true;
+			}
+			else 
+			{
+				if (((InCompareFlags & DomainType) && ComparatorExecutionDomain->Type != ExecutionDomain->Type) ||
+					((InCompareFlags & DomainName) && ComparatorExecutionDomain->Name != ExecutionDomain->Name) ||
+					((InCompareFlags & DomainExpression) && ComparatorExecutionDomain->Expression != ExecutionDomain->Expression))
+				{
+					bItemsAreAllSame = false;
+					break;
+				}
+			}
+		}
+	}
+
+	if (ComparatorExecutionDomain && (!bInCheckMultiples || bItemsAreAllSame))
+	{
+		return *ComparatorExecutionDomain;
+	}
+	
+	return {};
+}
+
 
 TSharedRef<IPropertyTypeCustomization> FOptimusDataDomainCustomization::MakeInstance()
 {
@@ -327,7 +532,7 @@ TSharedRef<IPropertyTypeCustomization> FOptimusDataDomainCustomization::MakeInst
 
 FOptimusDataDomainCustomization::FOptimusDataDomainCustomization() :
 	ParameterMarker(MakeShared<TArray<FName>>()),
-	ExpressionMarker(MakeShared<TArray<FName>>(TArray<FName>{FName("~")}))
+	ExpressionMarker(MakeShared<TArray<FName>>(TArray<FName>{ExpressionMarkerName}))
 {
 }
 
@@ -432,7 +637,7 @@ void FOptimusDataDomainCustomization::CustomizeHeader(
 			+ SHorizontalBox::Slot()
 			.FillWidth(0.9)
 			[
-				SNew(SComboBox<TSharedRef<TArray<FName>>>)
+				SAssignNew(DimensionalComboBox, SComboBox<TSharedRef<TArray<FName>>>)
 				.ToolTipText(DimensionSelectorTooltipText)
 				.OptionsSource(&DomainDimensionNames)
 				.InitiallySelectedItem(InitialDimensionSelection)
@@ -639,6 +844,10 @@ void FOptimusDataDomainCustomization::CustomizeHeader(
 				{
 					const FOptimusDataDomain DimensionalDomain;
 					SetDataDomain(InPropertyHandle, DimensionalDomain, DomainType);
+
+					// Making sure Expression option can be selected again
+					DimensionalComboBox->ClearSelection();
+					
 					return FReply::Handled();
 				})
 				.ToolTipText(LOCTEXT("ClearExpression", "Clear Expression"))
diff --git a/Engine/Plugins/Animation/DeformerGraph/Source/OptimusEditor/Private/OptimusDetailsCustomization.h b/Engine/Plugins/Animation/DeformerGraph/Source/OptimusEditor/Private/OptimusDetailsCustomization.h
--- a/Engine/Plugins/Animation/DeformerGraph/Source/OptimusEditor/Private/OptimusDetailsCustomization.h
+++ b/Engine/Plugins/Animation/DeformerGraph/Source/OptimusEditor/Private/OptimusDetailsCustomization.h
@@ -9,6 +9,7 @@
 #include "PropertyHandle.h"
 #include "Widgets/Input/SComboBox.h"
 
+struct FOptimusExecutionDomain;
 class IOptimusExecutionDomainProvider;
 class FOptimusHLSLSyntaxHighlighter;
 class SEditableTextBox;
@@ -73,17 +74,46 @@ public:
 		IPropertyTypeCustomizationUtils& InCustomizationUtils) override
 	{ }
 
+	enum DomainFlags
+	{
+		DomainType			= 0x1,
+		DomainName			= 0x2,
+		DomainExpression	= 0x4,
+
+		DomainAll			= DomainType + DomainName + DomainExpression     
+	};
+	
+	
 private:
 
+	FText FormatContextName(
+		FName InName
+		) const;
+	
 	void UpdateContextNames();
 
+	static void SetExecutionDomain(
+		TSharedRef<IPropertyHandle> InPropertyHandle,
+		const FOptimusExecutionDomain& InExecutionDomain,
+		DomainFlags InSetFlags = DomainAll
+		);
+	
+	static TOptional<FOptimusExecutionDomain> TryGetSingleExecutionDomain(
+		TSharedRef<IPropertyHandle> InPropertyHandle,
+		DomainFlags InCompareFlags = DomainAll,
+		bool bInCheckMultiples = true
+	);
+	
 	TSharedPtr<SComboBox<FName>> ComboBox;
+	TSharedPtr<SEditableTextBox> ExpressionTextBox;
 	
 	TArray<FName> ContextNames;
 
 	TArray<TWeakObjectPtr<UObject>> WeakOwningObjects;
 };
 
+ENUM_CLASS_FLAGS(FOptimusExecutionDomainCustomization::DomainFlags);
+
 
 class FOptimusDataDomainCustomization :
 	public IPropertyTypeCustomization
@@ -145,18 +175,15 @@ private:
 
 	TSharedRef<TArray<FName>> ParameterMarker;
 	TSharedRef<TArray<FName>> ExpressionMarker;
+	
+	TSharedPtr<SComboBox<TSharedRef<TArray<FName>>>> DimensionalComboBox;
 	TSharedPtr<SEditableTextBox> ExpressionTextBox;
 	
 	TArray<TSharedRef<TArray<FName>>> DomainDimensionNames;
 	bool bAllowParameters = false;
 };
 
-inline FOptimusDataDomainCustomization::DomainFlags operator|(FOptimusDataDomainCustomization::DomainFlags A, FOptimusDataDomainCustomization::DomainFlags B)
-{
-	return static_cast<FOptimusDataDomainCustomization::DomainFlags>(static_cast<int32>(A) | static_cast<int32>(B));
-}
-
-
+ENUM_CLASS_FLAGS(FOptimusDataDomainCustomization::DomainFlags);
 
 class FOptimusShaderTextCustomization : 
 	public IPropertyTypeCustomization
diff --git a/Engine/Plugins/Animation/DeformerGraph/Source/OptimusEditor/Private/OptimusEditorModule.cpp b/Engine/Plugins/Animation/DeformerGraph/Source/OptimusEditor/Private/OptimusEditorModule.cpp
--- a/Engine/Plugins/Animation/DeformerGraph/Source/OptimusEditor/Private/OptimusEditorModule.cpp
+++ b/Engine/Plugins/Animation/DeformerGraph/Source/OptimusEditor/Private/OptimusEditorModule.cpp
@@ -9,7 +9,7 @@
 #include "AssetViewUtils.h"
 #include "EdGraphUtilities.h"
 #include "IAssetTools.h"
-#include "IOptimusExecutionDomainProvider.h"
+#include "OptimusExecutionDomain.h"
 #include "ObjectTools.h"
 #include "OptimusBindingTypes.h"
 #include "OptimusComponentSource.h"
