Optimus: fixed a bug where ReadNumThreads().x returns 0 for kernels with unified dispatch

diff --git a/Engine/Plugins/Animation/DeformerGraph/Source/OptimusCore/Private/DataInterfaces/OptimusDataInterfaceCustomComputeKernel.cpp b/Engine/Plugins/Animation/DeformerGraph/Source/OptimusCore/Private/DataInterfaces/OptimusDataInterfaceCustomComputeKernel.cpp
--- a/Engine/Plugins/Animation/DeformerGraph/Source/OptimusCore/Private/DataInterfaces/OptimusDataInterfaceCustomComputeKernel.cpp
+++ b/Engine/Plugins/Animation/DeformerGraph/Source/OptimusCore/Private/DataInterfaces/OptimusDataInterfaceCustomComputeKernel.cpp
@@ -92,7 +92,7 @@ void UOptimusCustomComputeKernelDataProvider::InitFromDataInterface(const UOptim
 FComputeDataProviderRenderProxy* UOptimusCustomComputeKernelDataProvider::GetRenderProxy()
 {
 	TArray<int32> InvocationCounts;
-	int32 TotalThreadCount;
+	int32 TotalThreadCount = 0;
 	
 	if (!GetInvocationThreadCounts(InvocationCounts, TotalThreadCount))
 	{
@@ -135,15 +135,18 @@ bool UOptimusCustomComputeKernelDataProvider::GetInvocationThreadCounts(TArray<i
 	{
 		return false;
 	}
-	
+
+
+	OutTotalThreadCount = 0;
 	OutInvocationThreadCount.Reset(Values.Num());
 	for (const float& Value : Values)
 	{
-		OutInvocationThreadCount.Add(static_cast<int32>(Value));
+		int32 Count = static_cast<int32>(Value);
+		OutInvocationThreadCount.Add(Count);
+		OutTotalThreadCount += Count;
 	}
 
 	return true;
-	
 }
 
 bool UOptimusCustomComputeKernelDataProvider::GetInvocationThreadCounts_DEPRECATED(
@@ -314,6 +317,11 @@ bool FOptimusCustomComputeKernelDataProviderProxy::IsValid(FValidationData const
 	{
 		return false;
 	}
+
+	if (TotalThreadCount <= 0)
+	{
+		return false;
+	}
 	
 	return true;
 }
