Optimus: making sure graph travesal code scales with node count instead of pin count, which was causing significant slow down

diff --git a/Engine/Plugins/Animation/DeformerGraph/Source/OptimusCore/Private/OptimusDeformer.cpp b/Engine/Plugins/Animation/DeformerGraph/Source/OptimusCore/Private/OptimusDeformer.cpp
--- a/Engine/Plugins/Animation/DeformerGraph/Source/OptimusCore/Private/OptimusDeformer.cpp
+++ b/Engine/Plugins/Animation/DeformerGraph/Source/OptimusCore/Private/OptimusDeformer.cpp
@@ -1331,17 +1331,18 @@ static void CollectNodes(
 	TArray<FOptimusRoutedConstNode>& OutCollectedNodes
 	)
 {
-	TSet<const UOptimusNode*> VisitedNodes;
+	TSet<FOptimusRoutedConstNode> VisitedNodes;
+	TSet<FOptimusRoutedConstNode> UniqueNeighborNodes;
 	TQueue<FOptimusRoutedConstNode> WorkingSet;
 
 	for (const UOptimusNode* Node: InSeedNodes)
 	{
 		WorkingSet.Enqueue({Node, FOptimusPinTraversalContext{}});
-		VisitedNodes.Add(Node);
+		VisitedNodes.Add({Node, FOptimusPinTraversalContext{}});
 		OutCollectedNodes.Add({Node, FOptimusPinTraversalContext{}});
 	}
 
-	auto CollectFromInputPins = [&WorkingSet, &VisitedNodes, &OutCollectedNodes](const FOptimusRoutedConstNode& InWorkItem, const UOptimusNodePin* InPin)
+	auto CollectFromInputPins = [&WorkingSet, &VisitedNodes, &UniqueNeighborNodes, &OutCollectedNodes](const FOptimusRoutedConstNode& InWorkItem, const UOptimusNodePin* InPin)
 	{
 		for (const FOptimusRoutedNodePin& ConnectedPin: InPin->GetConnectedPinsWithRouting(InWorkItem.TraversalContext))
 		{
@@ -1349,17 +1350,23 @@ static void CollectNodes(
 			{
 				const UOptimusNode *NextNode = ConnectedPin.NodePin->GetOwningNode();
 				FOptimusRoutedConstNode CollectedNode{NextNode, ConnectedPin.TraversalContext};
-				WorkingSet.Enqueue(CollectedNode);
-				if (!VisitedNodes.Contains(NextNode))
+				if (!UniqueNeighborNodes.Contains(CollectedNode))
 				{
-					VisitedNodes.Add(NextNode);
-					OutCollectedNodes.Add(CollectedNode);
-				}
-				else
-				{
-					// Push the node to the back because to ensure that it is scheduled  earlier then it's referencing node.
-					OutCollectedNodes.RemoveSingle(CollectedNode);
-					OutCollectedNodes.Add(CollectedNode);
+					UniqueNeighborNodes.Add(CollectedNode);
+					
+					WorkingSet.Enqueue(CollectedNode);
+					
+					if (!VisitedNodes.Contains(CollectedNode))
+					{
+						VisitedNodes.Add(CollectedNode);
+						OutCollectedNodes.Add(CollectedNode);
+					}
+					else
+					{
+						// Push the node to the back because to ensure that it is scheduled  earlier then it's referencing node.
+						OutCollectedNodes.RemoveSingle(CollectedNode);
+						OutCollectedNodes.Add(CollectedNode);
+					}
 				}
 			}
 		}
@@ -1368,6 +1375,8 @@ static void CollectNodes(
 	FOptimusRoutedConstNode WorkItem;
 	while (WorkingSet.Dequeue(WorkItem))
 	{
+		UniqueNeighborNodes.Reset();
+		
 		// Traverse in the direction of input pins (up the graph).
 		for (const UOptimusNodePin* Pin: WorkItem.Node->GetPins())
 		{
diff --git a/Engine/Plugins/Animation/DeformerGraph/Source/OptimusCore/Private/OptimusNodeGraph.cpp b/Engine/Plugins/Animation/DeformerGraph/Source/OptimusCore/Private/OptimusNodeGraph.cpp
--- a/Engine/Plugins/Animation/DeformerGraph/Source/OptimusCore/Private/OptimusNodeGraph.cpp
+++ b/Engine/Plugins/Animation/DeformerGraph/Source/OptimusCore/Private/OptimusNodeGraph.cpp
@@ -1406,7 +1406,7 @@ TSet<UOptimusComponentSourceBinding*> UOptimusNodeGraph::GetComponentSourceBindi
 	const UOptimusNodePin* InNodePin
 	) const
 {
-	TSet<const UOptimusNode*> VisitedNodes;
+	TSet<FOptimusRoutedConstNode> VisitedNodes;
 	TSet<UOptimusComponentSourceBinding*> Bindings;
 	
 	TQueue<FOptimusRoutedConstNode> WorkingSet;
@@ -1478,10 +1478,11 @@ TSet<UOptimusComponentSourceBinding*> UOptimusNodeGraph::GetComponentSourceBindi
 				{
 					const UOptimusNode *NextNode = ConnectedPin.NodePin->GetOwningNode();
 					FOptimusRoutedConstNode CollectedNode{NextNode, ConnectedPin.TraversalContext};
-					WorkingSet.Enqueue(CollectedNode);
-					if (!VisitedNodes.Contains(NextNode))
+					
+					if (!VisitedNodes.Contains(CollectedNode))
 					{
-						VisitedNodes.Add(NextNode);
+						VisitedNodes.Add(CollectedNode);
+						WorkingSet.Enqueue(CollectedNode);
 					}
 				}
 			}
diff --git a/Engine/Plugins/Animation/DeformerGraph/Source/OptimusCore/Public/IOptimusNodePinRouter.h b/Engine/Plugins/Animation/DeformerGraph/Source/OptimusCore/Public/IOptimusNodePinRouter.h
--- a/Engine/Plugins/Animation/DeformerGraph/Source/OptimusCore/Public/IOptimusNodePinRouter.h
+++ b/Engine/Plugins/Animation/DeformerGraph/Source/OptimusCore/Public/IOptimusNodePinRouter.h
@@ -20,6 +20,22 @@ class UOptimusNodePin;
 struct FOptimusPinTraversalContext
 {
 	TArray<const IOptimusNodePinRouter*, TInlineAllocator<4>> ReferenceNesting;
+
+	bool operator==(FOptimusPinTraversalContext const& RHS) const
+	{
+		return ReferenceNesting == RHS.ReferenceNesting;
+	}
+	
+	friend uint32 GetTypeHash(const FOptimusPinTraversalContext& InContext)
+	{
+		uint32 Hash = GetTypeHash(InContext.ReferenceNesting.Num());
+		
+		for (const IOptimusNodePinRouter* Referencer : InContext.ReferenceNesting)
+		{
+			Hash = HashCombineFast(Hash, GetTypeHash(Referencer));
+		}
+		return Hash;
+	}
 };
 
 /** A utility struct for when traversing the graph through routed nodes. */
@@ -30,8 +46,16 @@ struct FOptimusRoutedNode
 
 	bool operator==(FOptimusRoutedNode const& RHS) const
 	{
-		return Node == RHS.Node;
+		return Node == RHS.Node && TraversalContext == RHS.TraversalContext;
 	}
+
+	friend uint32 GetTypeHash(const FOptimusRoutedNode& InRoutedNode)
+	{
+		uint32 Hash = GetTypeHash(InRoutedNode.Node);
+		
+		Hash = HashCombineFast(Hash, GetTypeHash(InRoutedNode.TraversalContext));
+		return Hash;
+	}	
 };
 
 struct FOptimusRoutedConstNode
@@ -41,8 +65,16 @@ struct FOptimusRoutedConstNode
 
 	bool operator==(FOptimusRoutedConstNode const& RHS) const
 	{
-		return Node == RHS.Node;
+		return Node == RHS.Node && TraversalContext == RHS.TraversalContext;
 	}
+
+	friend uint32 GetTypeHash(const FOptimusRoutedConstNode& InRoutedNode)
+	{
+		uint32 Hash = GetTypeHash(InRoutedNode.Node);
+		
+		Hash = HashCombineFast(Hash, GetTypeHash(InRoutedNode.TraversalContext));
+		return Hash;
+	}	
 };
 
 
@@ -58,8 +90,16 @@ struct FOptimusRoutedNodePin
 	
 	bool operator==(FOptimusRoutedNodePin const& RHS) const
 	{
-		return NodePin == RHS.NodePin;
+		return NodePin == RHS.NodePin && TraversalContext == RHS.TraversalContext;
 	}
+
+	friend uint32 GetTypeHash(const FOptimusRoutedNodePin& InRoutedNodePin)
+	{
+		uint32 Hash = GetTypeHash(InRoutedNodePin.NodePin);
+		
+		Hash = HashCombineFast(Hash, GetTypeHash(InRoutedNodePin.TraversalContext));
+		return Hash;
+	}		
 };
 
 struct FOptimusRoutedConstNodePin
@@ -69,8 +109,16 @@ struct FOptimusRoutedConstNodePin
 	
 	bool operator==(FOptimusRoutedConstNodePin const& RHS) const
 	{
-		return NodePin == RHS.NodePin;
+		return NodePin == RHS.NodePin && TraversalContext == RHS.TraversalContext;
 	}
+	
+	friend uint32 GetTypeHash(const FOptimusRoutedConstNodePin& InRoutedNodePin)
+	{
+		uint32 Hash = GetTypeHash(InRoutedNodePin.NodePin);
+		
+		Hash = HashCombineFast(Hash, GetTypeHash(InRoutedNodePin.TraversalContext));
+		return Hash;
+	}		
 };
 
 
