Optimus: int type resource kernel parameter bindings now can opt-in to expose atomic write functions, instead of relying connected resource nodes

diff --git a/Engine/Plugins/Animation/DeformerGraph/Shaders/Private/DataInterfaceRawBuffer.ush b/Engine/Plugins/Animation/DeformerGraph/Shaders/Private/DataInterfaceRawBuffer.ush
--- a/Engine/Plugins/Animation/DeformerGraph/Shaders/Private/DataInterfaceRawBuffer.ush
+++ b/Engine/Plugins/Animation/DeformerGraph/Shaders/Private/DataInterfaceRawBuffer.ush
@@ -55,31 +55,37 @@ void WriteValue_{DataInterfaceName}(uint Index, PUBLIC_TYPE Value)
 #endif
 }
 
-void WriteAtomicAdd_{DataInterfaceName}(uint Index, PUBLIC_TYPE Value)
+PUBLIC_TYPE WriteAtomicAdd_{DataInterfaceName}(uint Index, PUBLIC_TYPE Value)
 {
 #if BUFFER_TYPE_SUPPORTS_ATOMIC
 	uint BufferIndex = {DataInterfaceName}_StartOffset + Index;
 	PUBLIC_TYPE OldValue;
 	InterlockedAdd({DataInterfaceName}_BufferUAV[BufferIndex], Value, OldValue);
+	return OldValue;
 #endif
+	return 0;
 }
 
-void WriteAtomicMin_{DataInterfaceName}(uint Index, PUBLIC_TYPE Value)
+PUBLIC_TYPE WriteAtomicMin_{DataInterfaceName}(uint Index, PUBLIC_TYPE Value)
 {
 #if BUFFER_TYPE_SUPPORTS_ATOMIC
 	uint BufferIndex = {DataInterfaceName}_StartOffset + Index;
 	PUBLIC_TYPE OldValue;
 	InterlockedMin({DataInterfaceName}_BufferUAV[BufferIndex], Value, OldValue);
+	return OldValue;
 #endif
+	return 0;
 }
 
-void WriteAtomicMax_{DataInterfaceName}(uint Index, PUBLIC_TYPE Value)
+PUBLIC_TYPE WriteAtomicMax_{DataInterfaceName}(uint Index, PUBLIC_TYPE Value)
 {
 #if BUFFER_TYPE_SUPPORTS_ATOMIC
 	uint BufferIndex = {DataInterfaceName}_StartOffset + Index;
 	PUBLIC_TYPE OldValue;
 	InterlockedMax({DataInterfaceName}_BufferUAV[BufferIndex], Value, OldValue);
+	return OldValue;
 #endif
+	return 0;
 }
 
 #undef PUBLIC_TYPE
diff --git a/Engine/Plugins/Animation/DeformerGraph/Source/OptimusCore/Private/DataInterfaces/OptimusDataInterfaceRawBuffer.cpp b/Engine/Plugins/Animation/DeformerGraph/Source/OptimusCore/Private/DataInterfaces/OptimusDataInterfaceRawBuffer.cpp
--- a/Engine/Plugins/Animation/DeformerGraph/Source/OptimusCore/Private/DataInterfaces/OptimusDataInterfaceRawBuffer.cpp
+++ b/Engine/Plugins/Animation/DeformerGraph/Source/OptimusCore/Private/DataInterfaces/OptimusDataInterfaceRawBuffer.cpp
@@ -88,16 +88,19 @@ void UOptimusRawBufferDataInterface::GetSupportedOutputs(TArray<FShaderFunctionD
 
 	OutFunctions.AddDefaulted_GetRef()
 		.SetName(TEXT("WriteAtomicAdd"))
+		.AddReturnType(ValueType)
 		.AddParam(EShaderFundamentalType::Uint)
 		.AddParam(ValueType);
 
 	OutFunctions.AddDefaulted_GetRef()
 		.SetName(TEXT("WriteAtomicMin"))
+		.AddReturnType(ValueType)
 		.AddParam(EShaderFundamentalType::Uint)
 		.AddParam(ValueType);
 
 	OutFunctions.AddDefaulted_GetRef()
 		.SetName(TEXT("WriteAtomicMax"))
+		.AddReturnType(ValueType)
 		.AddParam(EShaderFundamentalType::Uint)
 		.AddParam(ValueType);
 }
@@ -180,6 +183,7 @@ UComputeDataProvider* UOptimusTransientBufferDataInterface::CreateDataProvider(
 	UOptimusTransientBufferDataProvider *Provider = CreateProvider<UOptimusTransientBufferDataProvider>(InBinding);
 	Provider->ElementStride = ValueType->GetResourceElementSize();
 	Provider->RawStride = GetRawStride();
+	Provider->bZeroInitForAtomicWrites = bZeroInitForAtomicWrites;
 	return Provider;
 }
 
@@ -427,7 +431,7 @@ FComputeDataProviderRenderProxy* UOptimusTransientBufferDataProvider::GetRenderP
 		}
 	}
 	
-	return new FOptimusTransientBufferDataProviderProxy(InvocationCounts, ElementStride, RawStride);
+	return new FOptimusTransientBufferDataProviderProxy(InvocationCounts, ElementStride, RawStride, bZeroInitForAtomicWrites);
 }
 
 
@@ -447,12 +451,14 @@ FComputeDataProviderRenderProxy* UOptimusPersistentBufferDataProvider::GetRender
 FOptimusTransientBufferDataProviderProxy::FOptimusTransientBufferDataProviderProxy(
 	TArray<int32> InInvocationElementCounts,
 	int32 InElementStride,
-	int32 InRawStride
+	int32 InRawStride,
+	bool bInZeroInitForAtomicWrites
 	) :
 	InvocationElementCounts(InInvocationElementCounts),
 	TotalElementCount(0),
 	ElementStride(InElementStride),
-	RawStride(InRawStride)
+	RawStride(InRawStride),
+	bZeroInitForAtomicWrites(bInZeroInitForAtomicWrites)
 {
 	for (int32 NumElements : InvocationElementCounts)
 	{
@@ -485,6 +491,11 @@ void FOptimusTransientBufferDataProviderProxy::AllocateResources(FRDGBuilder& Gr
 	Buffer = GraphBuilder.CreateBuffer(FRDGBufferDesc::CreateStructuredDesc(Stride, TotalElementCount * ElementStrideMultiplier), TEXT("TransientBuffer"), ERDGBufferFlags::None);
 	BufferSRV = GraphBuilder.CreateSRV(Buffer);
 	BufferUAV = GraphBuilder.CreateUAV(Buffer, ERDGUnorderedAccessViewFlags::SkipBarrier);
+
+	if (bZeroInitForAtomicWrites)
+	{
+		AddClearUAVPass(GraphBuilder, BufferUAV, 0);
+	}
 }
 
 void FOptimusTransientBufferDataProviderProxy::GatherDispatchData(FDispatchData const& InDispatchData)
diff --git a/Engine/Plugins/Animation/DeformerGraph/Source/OptimusCore/Private/DataInterfaces/OptimusDataInterfaceRawBuffer.h b/Engine/Plugins/Animation/DeformerGraph/Source/OptimusCore/Private/DataInterfaces/OptimusDataInterfaceRawBuffer.h
--- a/Engine/Plugins/Animation/DeformerGraph/Source/OptimusCore/Private/DataInterfaces/OptimusDataInterfaceRawBuffer.h
+++ b/Engine/Plugins/Animation/DeformerGraph/Source/OptimusCore/Private/DataInterfaces/OptimusDataInterfaceRawBuffer.h
@@ -115,6 +115,9 @@ public:
 	void GetShaderParameters(TCHAR const* UID, FShaderParametersMetadataBuilder& InOutBuilder, FShaderParametersMetadataAllocations& InOutAllocations) const override;
 	UComputeDataProvider* CreateDataProvider(TObjectPtr<UObject> InBinding, uint64 InInputMask, uint64 InOutputMask) const override;
 	//~ End UComputeDataInterface Interface
+
+	UPROPERTY()
+	bool bZeroInitForAtomicWrites = false;	
 };
 
 
@@ -212,6 +215,8 @@ public:
 	//~ Begin UComputeDataProvider Interface
 	FComputeDataProviderRenderProxy* GetRenderProxy() override;
 	//~ End UComputeDataProvider Interface
+
+	bool bZeroInitForAtomicWrites = false;
 };
 
 
@@ -242,7 +247,8 @@ public:
 	FOptimusTransientBufferDataProviderProxy(
 		TArray<int32> InInvocationElementCounts,
 		int32 InElementStride,
-		int32 InRawStride
+		int32 InRawStride,
+		bool bInZeroInitForAtomicWrites
 		);
 
 	//~ Begin FComputeDataProviderRenderProxy Interface
@@ -258,6 +264,7 @@ private:
 	int32 TotalElementCount;
 	const int32 ElementStride;
 	const int32 RawStride;
+	const bool bZeroInitForAtomicWrites;
 
 	FRDGBufferRef Buffer;
 	FRDGBufferSRVRef BufferSRV;
diff --git a/Engine/Plugins/Animation/DeformerGraph/Source/OptimusCore/Private/Nodes/OptimusNode_ComputeKernelBase.cpp b/Engine/Plugins/Animation/DeformerGraph/Source/OptimusCore/Private/Nodes/OptimusNode_ComputeKernelBase.cpp
--- a/Engine/Plugins/Animation/DeformerGraph/Source/OptimusCore/Private/Nodes/OptimusNode_ComputeKernelBase.cpp
+++ b/Engine/Plugins/Animation/DeformerGraph/Source/OptimusCore/Private/Nodes/OptimusNode_ComputeKernelBase.cpp
@@ -17,6 +17,7 @@
 #include "OptimusKernelSource.h"
 #include "OptimusConstant.h"
 #include "OptimusExecutionDomain.h"
+#include "OptimusNode_ResourceAccessorBase.h"
 
 #include UE_INLINE_GENERATED_CPP_BY_NAME(OptimusNode_ComputeKernelBase)
 
@@ -311,6 +312,40 @@ TOptional<FText> UOptimusNode_ComputeKernelBase::ValidateForCompile() const
 				}
 			}
 		}
+
+		if (Pin->GetDirection() == EOptimusNodePinDirection::Output)
+		{
+			if (GetPinSupportAtomic(Pin))
+			{
+				TArray<FOptimusRoutedNodePin> ConnectedPins = Pin->GetConnectedPinsWithRouting();
+				if (ConnectedPins.Num() > 1)
+				{
+					int32 TransientBufferCount = 0;
+					int32 PersistentBufferCount = 0;
+					for (const FOptimusRoutedNodePin& ConnectedPin : ConnectedPins)
+					{
+						if (Cast<IOptimusComputeKernelProvider>(ConnectedPin.NodePin->GetOwningNode()))
+						{
+							// At most 1 transient buffer, shared among connected kernels
+							TransientBufferCount = 1;
+						}
+						else if (Cast<UOptimusNode_ResourceAccessorBase>(ConnectedPin.NodePin->GetOwningNode()))
+						{
+							PersistentBufferCount++;
+						}
+						else
+						{
+							return FText::Format(LOCTEXT("InvalidAtomicPinConnection", "Pin '{0}' supports Atomic writes and should be connected to either a single resource node or kernel nodes"), FText::FromName(Pin->GetUniqueName()));
+						}
+					}
+
+					if (PersistentBufferCount + TransientBufferCount > 1)
+					{
+						return FText::Format(LOCTEXT("InvalidAtomicPinConnection", "Pin '{0}' supports Atomic writes and should be connected to either a single resource node or kernel nodes"), FText::FromName(Pin->GetUniqueName()));
+					}
+				}
+			}
+		}
 	}
 
 	// We should have at least the primary group, which needs to have a unique component binding
@@ -396,6 +431,31 @@ TSet<UOptimusComponentSourceBinding*> UOptimusNode_ComputeKernelBase::GetGroupCo
 	return Bindings;
 }
 
+FString UOptimusNode_ComputeKernelBase::GetAtomicWriteFunctionName(
+	EOptimusBufferWriteType InWriteType,
+	const FString& InBindingName
+	)
+{
+	FString FunctionName;
+	
+	switch(InWriteType)
+	{
+	case EOptimusBufferWriteType::WriteAtomicAdd:
+		FunctionName = FString::Printf(TEXT("AtomicAdd%s"), *InBindingName);
+		break;
+	case EOptimusBufferWriteType::WriteAtomicMin:
+		FunctionName = FString::Printf(TEXT("AtomicMin%s"), *InBindingName);
+		break;
+	case EOptimusBufferWriteType::WriteAtomicMax:
+		FunctionName = FString::Printf(TEXT("AtomicMax%s"), *InBindingName);
+		break;
+	default:
+		checkNoEntry();
+		break;
+	}
+
+	return FunctionName;
+}
 
 
 TOptional<FText> UOptimusNode_ComputeKernelBase::ProcessInputPinForComputeKernel(
@@ -502,9 +562,7 @@ TOptional<FText> UOptimusNode_ComputeKernelBase::ProcessInputPinForComputeKernel
 			// For resources we need the index parameter.
 			if (!InInputPin->GetDataDomain().IsSingleton())
 			{
-				TArray<FName> LevelNames = InInputPin->GetDataDomain().DimensionNames;
-
-				for (int32 Count = 0; Count < LevelNames.Num(); Count++)
+				for (int32 Count = 0; Count < InInputPin->GetDataDomain().NumDimensions(); Count++)
 				{
 					FuncDef.ParamTypes.Add(IndexParamDef);
 				}
@@ -539,15 +597,13 @@ TOptional<FText> UOptimusNode_ComputeKernelBase::ProcessInputPinForComputeKernel
 		}
 		else
 		{
-			TArray<FName> LevelNames = InInputPin->GetDataDomain().DimensionNames;
-			
 			ValueStr = InInputPin->GetDataType()->ShaderValueType->GetZeroValueAsString();
 
 			// No output connections, leave a stub function. The compiler will be in charge
 			// of optimizing out anything that causes us to ends up here.
 			TArray<FString> StubIndexes;
 			
-			for (const FString &IndexName: GetIndexNamesFromDataDomainLevels(LevelNames))
+			for (const FString &IndexName: GetIndexNamesFromDataDomain(InInputPin->GetDataDomain()))
 			{
 				StubIndexes.Add(*FString::Printf(TEXT("uint %s"), *IndexName));
 			}
@@ -633,11 +689,19 @@ void UOptimusNode_ComputeKernelBase::ProcessOutputPinForComputeKernel(
 	FShaderParamTypeDefinition IndexParamDef;
 	CopyValueType(FShaderValueType::Get(EShaderFundamentalType::Uint), IndexParamDef);
 
-	TArray<FName> LevelNames = InOutputPin->GetDataDomain().DimensionNames;
-	TArray<FString> IndexNames = GetIndexNamesFromDataDomainLevels(LevelNames);
+	TArray<FString> IndexNames = GetIndexNamesFromDataDomain(InOutputPin->GetDataDomain());
+
 	const FShaderValueTypeHandle ValueType = InOutputPin->GetDataType()->ShaderValueType;
 
 	TArray<FOptimusRoutedNodePin> ConnectedPins = InOutputPin->GetConnectedPinsWithRouting(InTraversalContext);
+
+	// Get the component binding from the upstream connection.
+	UOptimusComponentSourceBinding* ComponentBinding = nullptr;
+	TArray<UOptimusComponentSourceBinding*> ComponentBindings = GetOwningGraph()->GetComponentSourceBindingsForPin(InOutputPin).Array();
+	if (ensure(ComponentBindings.Num() == 1))
+	{
+		ComponentBinding = ComponentBindings[0];
+	}
 	
 	if (!ConnectedPins.IsEmpty())
 	{
@@ -647,7 +711,6 @@ void UOptimusNode_ComputeKernelBase::ProcessOutputPinForComputeKernel(
 		struct FWriteConnectionDef
 		{
 			UOptimusComputeDataInterface* DataInterface = nullptr;
-			const UOptimusComponentSourceBinding* ComponentBinding = nullptr;
 			FString DataFunctionName;
 			FString WriteToName;
 		};
@@ -662,16 +725,8 @@ void UOptimusNode_ComputeKernelBase::ProcessOutputPinForComputeKernel(
 			TArray<FShaderFunctionDefinition> WriteFunctions;
 			DataInterface->GetSupportedOutputs(WriteFunctions);
 
-			// Get the component binding from the upstream connection.
-			UOptimusComponentSourceBinding* ComponentBinding = nullptr;
-			TArray<UOptimusComponentSourceBinding*> ComponentBindings = GetOwningGraph()->GetComponentSourceBindingsForPin(InOutputPin).Array();
-			if (ensure(ComponentBindings.Num() == 1))
-			{
-				ComponentBinding = ComponentBindings[0];
-			}
-			
 			const int32 OutputIndex = UOptimusTransientBufferDataInterface::GetWriteValueOutputIndex(EOptimusBufferWriteType::Write);
-			WriteConnectionDefs.Add({DataInterface, ComponentBinding, WriteFunctions[OutputIndex].Name, TEXT("Transient")});
+			WriteConnectionDefs.Add({DataInterface, WriteFunctions[OutputIndex].Name, TEXT("Transient")});
 		}
 		
 		for (const FOptimusRoutedNodePin& ConnectedPin: ConnectedPins)
@@ -689,13 +744,12 @@ void UOptimusNode_ComputeKernelBase::ProcessOutputPinForComputeKernel(
 			// FIXME: Sub-pin write support.
 			UOptimusComputeDataInterface* DataInterface = InNodeDataInterfaceMap[ConnectedNode];
 			int32 DataFunctionIndex = InterfaceProvider->GetDataFunctionIndexFromPin(ConnectedPin.NodePin);
-			const UOptimusComponentSourceBinding* ComponentBinding = InterfaceProvider->GetComponentBinding();
 			
 			TArray<FShaderFunctionDefinition> FunctionDefinitions;
 			DataInterface->GetSupportedOutputs(FunctionDefinitions);
 			FString DataFunctionName = FunctionDefinitions[DataFunctionIndex].Name;
-			
-			WriteConnectionDefs.Add({DataInterface, ComponentBinding, DataFunctionName, ConnectedPin.NodePin->GetName()});
+
+			WriteConnectionDefs.Add({DataInterface, DataFunctionName, ConnectedPin.NodePin->GetName()});
 		}
 
 		TArray<FString> WrapFunctionNameCalls;
@@ -710,7 +764,7 @@ void UOptimusNode_ComputeKernelBase::ProcessOutputPinForComputeKernel(
 		
 			FShaderParamTypeDefinition ParamDef;
 			CopyValueType(ValueType, ParamDef);
-			for (int32 Count = 0; Count < LevelNames.Num(); Count++)
+			for (int32 Count = 0; Count < IndexNames.Num(); Count++)
 			{
 				FuncDef.ParamTypes.Add(IndexParamDef);
 			}
@@ -733,7 +787,7 @@ void UOptimusNode_ComputeKernelBase::ProcessOutputPinForComputeKernel(
 			
 			FOptimus_InterfaceBinding InterfaceBinding;
 			InterfaceBinding.DataInterface = WriteConnectionDef.DataInterface;
-			InterfaceBinding.ComponentBinding = WriteConnectionDef.ComponentBinding;
+			InterfaceBinding.ComponentBinding = ComponentBinding;
 			InterfaceBinding.DataInterfaceBindingIndex = DataInterfaceFuncIndex;
 			InterfaceBinding.BindingFunctionName = WrapFunctionName;
 			
@@ -753,6 +807,54 @@ void UOptimusNode_ComputeKernelBase::ProcessOutputPinForComputeKernel(
 			OutGeneratedFunctions.Add(
 				FString::Printf(TEXT("void Write%s(%s, %s Value)\n{\n%s;\n}"),
 					*InOutputPin->GetName(), *FString::Join(IndexParamNames, TEXT(", ")), *ValueType->ToString(), *FString::Join(WrapFunctionNameCalls, TEXT(";\n"))));
+		}	
+		
+		if (GetPinSupportAtomic(InOutputPin))
+		{
+			UOptimusRawBufferDataInterface* RawBufferDataInterface = nullptr;
+			if (InLinkDataInterfaceMap.Contains(InOutputPin))
+			{
+				RawBufferDataInterface = Cast<UOptimusRawBufferDataInterface>(InLinkDataInterfaceMap[InOutputPin]);	
+			}
+			else if (ensure(ConnectedPins.Num() == 1))
+			{
+				const UOptimusNode *ConnectedNode = ConnectedPins[0].NodePin->GetOwningNode();
+
+				// Connected to a data interface node?
+				if(ensure(InNodeDataInterfaceMap.Contains(ConnectedNode)))
+				{
+					RawBufferDataInterface = Cast<UOptimusRawBufferDataInterface>(InNodeDataInterfaceMap[ConnectedNode]);
+				}
+			}
+
+			if (ensure(RawBufferDataInterface))
+			{
+				TArray<FShaderFunctionDefinition> WriteFunctions;
+				RawBufferDataInterface->GetSupportedOutputs(WriteFunctions);
+				
+				for (uint8 WriteType = uint8(EOptimusBufferWriteType::Write) + 1; WriteType < uint8(EOptimusBufferWriteType::Count); WriteType++)
+				{
+
+					FString WrapFunctionName = GetAtomicWriteFunctionName((EOptimusBufferWriteType)WriteType, InOutputPin->GetName());
+
+					const int32 OutputIndex = UOptimusTransientBufferDataInterface::GetWriteValueOutputIndex((EOptimusBufferWriteType)WriteType);
+					
+					FShaderFunctionDefinition FuncDef = WriteFunctions[OutputIndex];
+					for (int32 Index = 0; Index < FuncDef.ParamTypes.Num(); Index++)
+					{
+						FuncDef.ParamTypes[Index].ResetTypeDeclaration();
+					}
+				
+					FOptimus_InterfaceBinding InterfaceBinding;
+					InterfaceBinding.ComponentBinding = ComponentBinding;
+					InterfaceBinding.DataInterface = RawBufferDataInterface;
+					InterfaceBinding.DataInterfaceBindingIndex = OutputIndex;
+					InterfaceBinding.BindingFunctionName = WrapFunctionName;
+				
+					OutOutputDataBindings.Add(InKernelSource->ExternalOutputs.Num(), InterfaceBinding);
+					InKernelSource->ExternalOutputs.Emplace(FuncDef);	
+				}
+			}
 		}
 	}
 	else
@@ -777,12 +879,7 @@ void UOptimusNode_ComputeKernelBase::ProcessOutputPinForComputeKernel(
 		if (TOptional<FString> Expression = InOutputPin->GetDataDomain().AsExpression())
 		{
 			// Get the component binding from the upstream connection.
-			int32 ComponentBindingIndex = INDEX_NONE;
-			TArray<UOptimusComponentSourceBinding*> ComponentBindings = GetOwningGraph()->GetComponentSourceBindingsForPin(InOutputPin).Array();
-			if (ensure(ComponentBindings.Num() == 1))
-			{
-				ComponentBindingIndex = ComponentBindings[0]->GetIndex();
-			}
+			int32 ComponentBindingIndex = ComponentBinding->GetIndex();
 			
 			OutKernelConstantContainer.AddToKernelContainer({
 				{this, NAME_None, InOutputPin->GetFName()},
diff --git a/Engine/Plugins/Animation/DeformerGraph/Source/OptimusCore/Private/Nodes/OptimusNode_ComputeKernelBase.h b/Engine/Plugins/Animation/DeformerGraph/Source/OptimusCore/Private/Nodes/OptimusNode_ComputeKernelBase.h
--- a/Engine/Plugins/Animation/DeformerGraph/Source/OptimusCore/Private/Nodes/OptimusNode_ComputeKernelBase.h
+++ b/Engine/Plugins/Animation/DeformerGraph/Source/OptimusCore/Private/Nodes/OptimusNode_ComputeKernelBase.h
@@ -10,6 +10,7 @@
 
 #include "OptimusNode_ComputeKernelBase.generated.h"
 
+enum class EOptimusBufferWriteType : uint8;
 struct FOptimusKernelConstantContainer;
 class UComputeSource;
 
@@ -52,6 +53,7 @@ public:
 	FOptimusExecutionDomain GetExecutionDomain() const override PURE_VIRTUAL(UOptimusNode_ComputeKernelBase::GetExecutionDomain, return {}; );
 	const UOptimusNodePin* GetPrimaryGroupPin() const override PURE_VIRTUAL(UOptimusNode_ComputeKernelBase::GetPrimaryGroupPin, return {}; );  
 	UComputeDataInterface* MakeKernelDataInterface(UObject* InOuter) const override PURE_VIRTUAL(UOptimusNode_ComputeKernelBase::MakeKernelDataInterface, return {}; );
+	bool GetPinSupportAtomic(const UOptimusNodePin* InPin) const override PURE_VIRTUAL(UOptimusNode_ComputeKernelBase::GetPinSupportAtomic, return {}; );
 	
 	// -- UOptimusNode overrides
 	TOptional<FText> ValidateForCompile() const override;
@@ -70,6 +72,29 @@ protected:
 		return IndexNames;
 	}
 
+	static TArray<FString> GetIndexNamesFromDataDomain(
+		const FOptimusDataDomain &InDomain
+		)
+	{
+		TArray<FString> IndexNames;
+	
+		if (InDomain.Type == EOptimusDataDomainType::Dimensional)
+		{
+			TArray<FName> LevelNames = InDomain.DimensionNames;
+			IndexNames = GetIndexNamesFromDataDomainLevels(LevelNames);
+		}
+		else if (InDomain.Type == EOptimusDataDomainType::Expression)
+		{
+			IndexNames = {TEXT("Index")};
+		}
+		else
+		{
+			checkNoEntry();
+		}
+
+		return IndexNames;
+	}
+
 	static FString GetCookedKernelSource(
 		const FString& InObjectPathName,
 		const FString& InShaderSource,
@@ -78,6 +103,8 @@ protected:
 		);
 
 	static TSet<UOptimusComponentSourceBinding*> GetGroupComponentSourceBindings(const UOptimusNodePin* InGroupPin);
+
+	static FString GetAtomicWriteFunctionName(EOptimusBufferWriteType InWriteType, const FString& InBindingName);
 	
 private:
 	TOptional<FText> ProcessInputPinForComputeKernel(
diff --git a/Engine/Plugins/Animation/DeformerGraph/Source/OptimusCore/Private/Nodes/OptimusNode_CustomComputeKernel.cpp b/Engine/Plugins/Animation/DeformerGraph/Source/OptimusCore/Private/Nodes/OptimusNode_CustomComputeKernel.cpp
--- a/Engine/Plugins/Animation/DeformerGraph/Source/OptimusCore/Private/Nodes/OptimusNode_CustomComputeKernel.cpp
+++ b/Engine/Plugins/Animation/DeformerGraph/Source/OptimusCore/Private/Nodes/OptimusNode_CustomComputeKernel.cpp
@@ -16,6 +16,7 @@
 #include "ComponentSources/OptimusSkinnedMeshComponentSource.h"
 #include "DataInterfaces/OptimusDataInterfaceCustomComputeKernel.h"
 #include "IOptimusDeprecatedExecutionDataInterface.h"
+#include "OptimusNode_ResourceAccessorBase.h"
 #include "Engine/UserDefinedStruct.h"
 
 #include UE_INLINE_GENERATED_CPP_BY_NAME(OptimusNode_CustomComputeKernel)
@@ -64,6 +65,12 @@ static bool IsSecondaryGroupInputPin(const UOptimusNodePin *InPin)
 	return InPin->GetName() != PrimaryGroupPinName && InPin->GetDirection() == EOptimusNodePinDirection::Input && InPin->IsGroupingPin();
 }
 
+static bool DoesBindingSupportAtomic(const FOptimusParameterBinding& InBinding)
+{
+	FOptimusDataTypeHandle IntType = FOptimusDataTypeRegistry::Get().FindType(*FIntProperty::StaticClass());
+	return InBinding.bSupportAtomicIfCompatibleDataType && InBinding.DataType == IntType;	
+}
+
 bool UOptimusNode_CustomComputeKernel::DoesSourceSupportUnifiedDispatch(const UOptimusNodePin& InOtherNodesPin)
 {
 	if (InOtherNodesPin.GetDirection() == EOptimusNodePinDirection::Output )
@@ -115,6 +122,23 @@ UComputeDataInterface* UOptimusNode_CustomComputeKernel::MakeKernelDataInterface
 	return KernelDataInterface;
 }
 
+bool UOptimusNode_CustomComputeKernel::GetPinSupportAtomic(const UOptimusNodePin* InPin) const
+{
+	if (ensure(InPin->GetDirection() == EOptimusNodePinDirection::Output))
+	{
+		if (const FOptimusParameterBinding* Binding = OutputBindingArray.FindByPredicate(
+			[InPin](const FOptimusParameterBinding& InBinding)
+				{
+					return InBinding.Name == InPin->GetFName(); 
+				}))
+		{
+			return DoesBindingSupportAtomic(*Binding);
+		}
+	}
+
+	return false;
+}
+
 #if WITH_EDITOR
 
 FString UOptimusNode_CustomComputeKernel::GetNameForShaderTextEditor() const
@@ -168,10 +192,41 @@ FString UOptimusNode_CustomComputeKernel::GetBindingDeclaration(
 	{
 		return GetDeclarationForBinding(*Binding, false);
 	}
+	for (const FOptimusSecondaryInputBindingsGroup& InputGroup: SecondaryInputBindingGroups)
+	{
+		if (const FOptimusParameterBinding* Binding = InputGroup.BindingArray.FindByPredicate(ParameterBindingPredicate))
+		{
+			return GetDeclarationForBinding(*Binding, true);
+		}
+	}
 
 	return FString();
 }
 
+bool UOptimusNode_CustomComputeKernel::GetBindingAtomicSupportCheckBoxVisibility(FName BindingName) const
+{
+	auto ParameterBindingPredicate = [BindingName](const FOptimusParameterBinding& InBinding)
+	{
+		if (InBinding.Name == BindingName)
+		{
+			return true;	
+		}
+			
+		return false;
+	};
+
+	if (const FOptimusParameterBinding* Binding = OutputBindingArray.FindByPredicate(ParameterBindingPredicate))
+	{
+		FOptimusDataTypeHandle IntType = FOptimusDataTypeRegistry::Get().FindType(*FIntProperty::StaticClass());
+		if (Binding->DataType == IntType)
+		{
+			return true;
+		}
+	}
+
+	return false;
+}
+
 
 TArray<IOptimusNodeAdderPinProvider::FAdderPinAction> UOptimusNode_CustomComputeKernel::GetAvailableAdderPinActions(
 	const UOptimusNodePin* InSourcePin,
@@ -181,12 +236,16 @@ TArray<IOptimusNodeAdderPinProvider::FAdderPinAction> UOptimusNode_CustomCompute
 {
 	if (InSourcePin->GetDataType().IsValid() && InSourcePin->GetDataType()->UsageFlags == EOptimusDataTypeUsageFlags::None)
 	{
-		if (OutReason)
+		// FIXME: avoid special case debug draw
+		if (InSourcePin->GetDataType() != FOptimusDataTypeRegistry::Get().FindType(FName(TEXT("FDebugDraw"))))
 		{
-			*OutReason = TEXT("Can't add pin with this type");
-		}
+			if (OutReason)
+			{
+				*OutReason = TEXT("Can't add pin with this type");
+			}
 
-		return {};
+			return {};
+		}
 	}
 	
 	TSet<UOptimusComponentSourceBinding*> SourceComponentBindings = InSourcePin->GetComponentSourceBindings();
@@ -691,7 +750,7 @@ void UOptimusNode_CustomComputeKernel::PropertyValueChanged(
 	const FPropertyChangedEvent& InPropertyChangedEvent
 	)
 {
-	auto UpdatePinNamesFromBindings = [this](
+	auto UpdatePinsFromBindings = [this](
 		UOptimusNodePin* InParentPin,
 		TArrayView<UOptimusNodePin* const> InPins,
 		TFunction<bool(const UOptimusNodePin*)> InPinFilter,
@@ -720,12 +779,12 @@ void UOptimusNode_CustomComputeKernel::PropertyValueChanged(
 		if (InPropertyChangedEvent.GetMemberPropertyName() == InputBindingsName)
 		{
 			UOptimusNodePin* PrimaryGroupPin = GetPrimaryGroupPin_Internal();	
-			UpdatePinNamesFromBindings(PrimaryGroupPin,  PrimaryGroupPin->GetSubPins(), IsConnectableInputDataPin, InputBindingArray, InApplyFunc);
+			UpdatePinsFromBindings(PrimaryGroupPin,  PrimaryGroupPin->GetSubPins(), IsConnectableInputDataPin, InputBindingArray, InApplyFunc);
 			return true;
 		}
 		if (InPropertyChangedEvent.GetMemberPropertyName() == OutputBindingsName)
 		{
-			UpdatePinNamesFromBindings(nullptr, GetPins(), IsConnectableOutputPin, OutputBindingArray, InApplyFunc);
+			UpdatePinsFromBindings(nullptr, GetPins(), IsConnectableOutputPin, OutputBindingArray, InApplyFunc);
 			return true;
 		}
 		if (InPropertyChangedEvent.GetMemberPropertyName() == ExtraInputBindingGroupsName)
@@ -736,7 +795,7 @@ void UOptimusNode_CustomComputeKernel::PropertyValueChanged(
 				if (IsSecondaryGroupInputPin(GroupPin))
 				{
 					FOptimusSecondaryInputBindingsGroup& BindingsGroup = SecondaryInputBindingGroups[GroupIndex];
-					UpdatePinNamesFromBindings(GroupPin, GroupPin->GetSubPins(), IsConnectableInputDataPin, BindingsGroup.BindingArray, InApplyFunc);
+					UpdatePinsFromBindings(GroupPin, GroupPin->GetSubPins(), IsConnectableInputDataPin, BindingsGroup.BindingArray, InApplyFunc);
 
 					GroupIndex++;
 				}
@@ -837,6 +896,11 @@ void UOptimusNode_CustomComputeKernel::PropertyValueChanged(
 			return;
 		}
 	}
+	else if (InPropertyChangedEvent.GetPropertyName() == GET_MEMBER_NAME_STRING_CHECKED(FOptimusParameterBinding, bSupportAtomicIfCompatibleDataType))
+	{
+		UpdatePreamble();
+		return;
+	}
 	else if ( InPropertyChangedEvent.GetPropertyName() == GET_MEMBER_NAME_STRING_CHECKED(FOptimusParameterBindingArray, InnerArray))
 	{
 		// Pasting into a specific binding
@@ -1416,6 +1480,11 @@ PRAGMA_ENABLE_DEPRECATION_WARNINGS
 		PostLoadAddMissingPrimaryGroupPin();
 	}
 
+	if (GetLinkerCustomVersion(FOptimusObjectVersion::GUID) < FOptimusObjectVersion::KernelParameterBindingToggleAtomic)
+	{
+		PostLoadExtractAtomicModeFromConnectedResource();
+	}
+
 	SetDisplayName(FText::FromName(KernelName));
 }
 
@@ -1572,12 +1641,6 @@ void UOptimusNode_CustomComputeKernel::UpdatePreamble()
 		{
 			AddCountFunctionIfNeeded(Binding.DataDomain.DimensionNames);
 		
-			TArray<FString> Indexes;
-			for (FString IndexName: GetIndexNamesFromDataDomainLevels(Binding.DataDomain.DimensionNames))
-			{
-				Indexes.Add(FString::Printf(TEXT("uint %s"), *IndexName));
-			}
-
 			if (Binding.DataType->ShaderValueType.IsValid())
 			{
 				Declarations.Add(GetDeclarationForBinding(Binding, true));
@@ -1704,6 +1767,77 @@ void UOptimusNode_CustomComputeKernel::PostLoadAddMissingPrimaryGroupPin()
 
 }
 
+void UOptimusNode_CustomComputeKernel::PostLoadExtractAtomicModeFromConnectedResource()
+{
+	FOptimusDataTypeHandle IntType = FOptimusDataTypeRegistry::Get().FindType(*FIntProperty::StaticClass());
+
+	TArray<TPair<UOptimusNodePin*, EOptimusBufferWriteType>> PinsToProcess;
+	
+	for (UOptimusNodePin* Pin: GetPins())
+	{
+		if (Pin->GetDirection() == EOptimusNodePinDirection::Output && Pin->GetDataType() == IntType)
+		{
+			TArray<FOptimusRoutedNodePin> ConnectedPins = Pin->GetConnectedPinsWithRouting();
+
+			for (const FOptimusRoutedNodePin& ConnectedPin : ConnectedPins)
+			{
+				if (const UOptimusNode_ResourceAccessorBase* ResourceNode = Cast<UOptimusNode_ResourceAccessorBase>(ConnectedPin.NodePin->GetOwningNode()))
+				{
+					if(ResourceNode->GetDeprecatedBufferWriteType() != EOptimusBufferWriteType::Write)
+					{
+						if (FOptimusParameterBinding* Binding = OutputBindingArray.FindByPredicate([Pin](const FOptimusParameterBinding& InBinding)
+							{
+								return InBinding.Name == Pin->GetFName(); 
+							}))
+						{
+							Binding->bSupportAtomicIfCompatibleDataType = true;
+
+							// Auto fixup is only possible for valid connections,
+							// where only a single resource node is connected to the atomic pin
+							if (ConnectedPins.Num() == 1)
+							{
+								PinsToProcess.Add({Pin, ResourceNode->GetDeprecatedBufferWriteType()});
+							}
+						}
+					}
+				}
+			}
+		}
+	}
+
+	if (!PinsToProcess.IsEmpty())
+	{
+		TArray<FString> Statements;
+		ShaderSource.ShaderText.ParseIntoArray(Statements, TEXT(";"));
+
+		bool bModified = false;
+		for (FString& Statement : Statements)
+		{
+			for (const TPair<UOptimusNodePin*, EOptimusBufferWriteType> PinInfo : PinsToProcess)
+			{
+				// See UOptimusNode_ComputeKernelBase::ProcessOutputPinForComputeKernel
+				FString OldName = FString::Printf(TEXT("Write%s"), *(PinInfo.Key->GetName()));;
+				
+				FString NewName = GetAtomicWriteFunctionName(PinInfo.Value, PinInfo.Key->GetName());
+
+				int32 NumReplaced = Statement.ReplaceInline(*OldName, *NewName, ESearchCase::CaseSensitive );
+				if (NumReplaced > 0)
+				{
+					// One function call per statement at most, no need to look further
+					bModified = true;
+					break;
+				}
+			}
+		}
+
+		if (bModified)
+		{
+			ShaderSource.ShaderText = FString::Join(Statements, TEXT(";"));
+			UpdatePreamble();	
+		}
+	}
+}
+
 
 bool UOptimusNode_CustomComputeKernel::PostLoadRemoveDeprecatedNumThreadsPin()
 {
@@ -1749,7 +1883,7 @@ bool UOptimusNode_CustomComputeKernel::PostLoadRemoveDeprecatedNumThreadsPin()
 FString UOptimusNode_CustomComputeKernel::GetDeclarationForBinding(const FOptimusParameterBinding& Binding, bool bIsInput)
 {
 	TArray<FString> Indexes;
-	for (FString IndexName: GetIndexNamesFromDataDomainLevels(Binding.DataDomain.DimensionNames))
+	for (FString IndexName: GetIndexNamesFromDataDomain(Binding.DataDomain))
 	{
 		Indexes.Add(FString::Printf(TEXT("uint %s"), *IndexName));
 	}
@@ -1761,9 +1895,26 @@ FString UOptimusNode_CustomComputeKernel::GetDeclarationForBinding(const FOptimu
 	}
 	else
 	{
-		return FString::Printf(TEXT("void Write%s(%s, %s Value);"),
-					*Binding.Name.ToString(), *FString::Join(Indexes, TEXT(", ")), *GetShaderValueTypeFriendlyName(Binding.DataType));
+		TArray<FString> Declarations;
+		Declarations.Add(FString::Printf(TEXT("void Write%s(%s, %s Value);"),
+			*Binding.Name.ToString(), *FString::Join(Indexes, TEXT(", ")), *GetShaderValueTypeFriendlyName(Binding.DataType)));
+
+		if (DoesBindingSupportAtomic(Binding))
+		{
+			for (uint8 WriteType = uint8(EOptimusBufferWriteType::Write) + 1; WriteType < uint8(EOptimusBufferWriteType::Count); WriteType++)
+			{
+				FString FunctionName = GetAtomicWriteFunctionName((EOptimusBufferWriteType)WriteType, Binding.Name.ToString());
+				
+				FString TypeName = GetShaderValueTypeFriendlyName(Binding.DataType);
+				Declarations.Add(FString::Printf(TEXT("%s %s(%s, %s Value);"),
+					*TypeName, *FunctionName, *FString::Join(Indexes, TEXT(", ")), *TypeName));	
+			}
+		}
+
+		return FString::Join(Declarations, TEXT("\n"));
 	}
+
+	return {};
 }
 
 #undef LOCTEXT_NAMESPACE
\ No newline at end of file
diff --git a/Engine/Plugins/Animation/DeformerGraph/Source/OptimusCore/Private/Nodes/OptimusNode_CustomComputeKernel.h b/Engine/Plugins/Animation/DeformerGraph/Source/OptimusCore/Private/Nodes/OptimusNode_CustomComputeKernel.h
--- a/Engine/Plugins/Animation/DeformerGraph/Source/OptimusCore/Private/Nodes/OptimusNode_CustomComputeKernel.h
+++ b/Engine/Plugins/Animation/DeformerGraph/Source/OptimusCore/Private/Nodes/OptimusNode_CustomComputeKernel.h
@@ -62,6 +62,7 @@ public:
 	FOptimusExecutionDomain GetExecutionDomain() const override;
 	const UOptimusNodePin* GetPrimaryGroupPin() const override;
 	UComputeDataInterface* MakeKernelDataInterface(UObject* InOuter) const override;
+	bool GetPinSupportAtomic(const UOptimusNodePin* InPin) const override;
 
 #if WITH_EDITOR
 	// IOptimusShaderTextProvider overrides
@@ -74,6 +75,7 @@ public:
 	
 	// IOptimusParameterBindingProvider
 	FString GetBindingDeclaration(FName BindingName) const override;
+	bool GetBindingAtomicSupportCheckBoxVisibility(FName BindingName) const override;
 
 	// IOptimusNodeAdderPinProvider
 	TArray<FAdderPinAction> GetAvailableAdderPinActions(
@@ -186,6 +188,7 @@ private:
 	UOptimusNodePin* GetPrimaryGroupPin_Internal() const;
 	void PostLoadExtractExecutionDomain();
 	void PostLoadAddMissingPrimaryGroupPin();
+	void PostLoadExtractAtomicModeFromConnectedResource();
 	
 	// Called during UOptimusDeformer::PostLoad()
 	bool PostLoadRemoveDeprecatedNumThreadsPin();
diff --git a/Engine/Plugins/Animation/DeformerGraph/Source/OptimusCore/Private/Nodes/OptimusNode_Resource.cpp b/Engine/Plugins/Animation/DeformerGraph/Source/OptimusCore/Private/Nodes/OptimusNode_Resource.cpp
--- a/Engine/Plugins/Animation/DeformerGraph/Source/OptimusCore/Private/Nodes/OptimusNode_Resource.cpp
+++ b/Engine/Plugins/Animation/DeformerGraph/Source/OptimusCore/Private/Nodes/OptimusNode_Resource.cpp
@@ -22,7 +22,7 @@ int32 UOptimusNode_Resource::GetDataFunctionIndexFromPin(const UOptimusNodePin*
 
 	return InPin->GetDirection() == EOptimusNodePinDirection::Output 
 		? UOptimusRawBufferDataInterface::GetReadValueInputIndex() 
-		: UOptimusRawBufferDataInterface::GetWriteValueOutputIndex(WriteType);
+		: UOptimusRawBufferDataInterface::GetWriteValueOutputIndex(EOptimusBufferWriteType::Write);
 }
 
 
diff --git a/Engine/Plugins/Animation/DeformerGraph/Source/OptimusCore/Private/Nodes/OptimusNode_ResourceAccessorBase.h b/Engine/Plugins/Animation/DeformerGraph/Source/OptimusCore/Private/Nodes/OptimusNode_ResourceAccessorBase.h
--- a/Engine/Plugins/Animation/DeformerGraph/Source/OptimusCore/Private/Nodes/OptimusNode_ResourceAccessorBase.h
+++ b/Engine/Plugins/Animation/DeformerGraph/Source/OptimusCore/Private/Nodes/OptimusNode_ResourceAccessorBase.h
@@ -60,6 +60,9 @@ public:
 	// Also IOptimusComponentBindingProvider implementation
 	UOptimusComponentSourceBinding* GetComponentBinding() const override;
 
+	PRAGMA_DISABLE_DEPRECATION_WARNINGS
+	EOptimusBufferWriteType GetDeprecatedBufferWriteType() const { return WriteType_DEPRECATED; }
+	PRAGMA_ENABLE_DEPRECATION_WARNINGS
 	
 protected:
 	void PreDuplicateRequirementActions(const UOptimusNodeGraph* InTargetGraph, FOptimusCompoundAction* InCompoundAction) override;
@@ -74,8 +77,8 @@ protected:
 	TWeakObjectPtr<UOptimusResourceDescription> ResourceDesc;
 
 	/** Logical operation when writing to the resource. */
-	UPROPERTY(EditAnywhere, Category = Resource)
-	EOptimusBufferWriteType WriteType = EOptimusBufferWriteType::Write;
+	UPROPERTY(meta=(DeprecatedProperty))
+	EOptimusBufferWriteType WriteType_DEPRECATED= EOptimusBufferWriteType::Write;
 
 	UPROPERTY(DuplicateTransient)
 	FOptimusNode_ResourceAccessorBase_DuplicationInfo DuplicationInfo;
diff --git a/Engine/Plugins/Animation/DeformerGraph/Source/OptimusCore/Private/Nodes/OptimusNode_SetResource.cpp b/Engine/Plugins/Animation/DeformerGraph/Source/OptimusCore/Private/Nodes/OptimusNode_SetResource.cpp
--- a/Engine/Plugins/Animation/DeformerGraph/Source/OptimusCore/Private/Nodes/OptimusNode_SetResource.cpp
+++ b/Engine/Plugins/Animation/DeformerGraph/Source/OptimusCore/Private/Nodes/OptimusNode_SetResource.cpp
@@ -20,7 +20,7 @@ int32 UOptimusNode_SetResource::GetDataFunctionIndexFromPin(const UOptimusNodePi
 		return INDEX_NONE;
 	}
 
-	return UOptimusRawBufferDataInterface::GetWriteValueOutputIndex(WriteType);
+	return UOptimusRawBufferDataInterface::GetWriteValueOutputIndex(EOptimusBufferWriteType::Write);
 }
 
 
diff --git a/Engine/Plugins/Animation/DeformerGraph/Source/OptimusCore/Private/OptimusDeformer.cpp b/Engine/Plugins/Animation/DeformerGraph/Source/OptimusCore/Private/OptimusDeformer.cpp
--- a/Engine/Plugins/Animation/DeformerGraph/Source/OptimusCore/Private/OptimusDeformer.cpp
+++ b/Engine/Plugins/Animation/DeformerGraph/Source/OptimusCore/Private/OptimusDeformer.cpp
@@ -1667,6 +1667,10 @@ UOptimusComputeGraph* UOptimusDeformer::CompileNodeGraphToComputeGraph(
 							TransientBufferDI->DataDomain = Pin->GetDataDomain();
 							TransientBufferDI->ComponentSourceBinding = *ComponentSourceBindings.CreateConstIterator();
 							TransientBufferDI->DomainConstantIdentifier = {ConnectedNode.Node, NAME_None, Pin->GetFName()};
+							if (KernelProvider->GetPinSupportAtomic(Pin))
+							{
+								TransientBufferDI->bZeroInitForAtomicWrites = true;
+							}
 							
 							LinkDataInterfaceMap.Add(Pin, TransientBufferDI);
 						}
diff --git a/Engine/Plugins/Animation/DeformerGraph/Source/OptimusCore/Private/OptimusDeformerInstance.cpp b/Engine/Plugins/Animation/DeformerGraph/Source/OptimusCore/Private/OptimusDeformerInstance.cpp
--- a/Engine/Plugins/Animation/DeformerGraph/Source/OptimusCore/Private/OptimusDeformerInstance.cpp
+++ b/Engine/Plugins/Animation/DeformerGraph/Source/OptimusCore/Private/OptimusDeformerInstance.cpp
@@ -438,10 +438,14 @@ TArray<float> UOptimusDeformerInstance::GetConstantValuePerInvocation(const FOpt
 		}
 
 
-		if (ensure(LodIndex != INDEX_NONE))
+		if (LodIndex != INDEX_NONE)
 		{
 			ComponentContext.LodIndexPerComponent.Add(LodIndex);
 		}
+		else
+		{
+			return {};
+		}
 	}
 
 
diff --git a/Engine/Plugins/Animation/DeformerGraph/Source/OptimusCore/Public/IOptimusComputeKernelProvider.h b/Engine/Plugins/Animation/DeformerGraph/Source/OptimusCore/Public/IOptimusComputeKernelProvider.h
--- a/Engine/Plugins/Animation/DeformerGraph/Source/OptimusCore/Public/IOptimusComputeKernelProvider.h
+++ b/Engine/Plugins/Animation/DeformerGraph/Source/OptimusCore/Public/IOptimusComputeKernelProvider.h
@@ -108,4 +108,7 @@ public:
 	
 	/** Each kernel may have its own data interface responsible for passing kernel related data to GPU */
 	virtual UComputeDataInterface* MakeKernelDataInterface(UObject* InOuter) const = 0;
+
+	/** Check if a specific pin needs to support atomic operation */
+	virtual bool GetPinSupportAtomic(const UOptimusNodePin* InPin) const = 0;
 };
diff --git a/Engine/Plugins/Animation/DeformerGraph/Source/OptimusCore/Public/IOptimusParameterBindingProvider.h b/Engine/Plugins/Animation/DeformerGraph/Source/OptimusCore/Public/IOptimusParameterBindingProvider.h
--- a/Engine/Plugins/Animation/DeformerGraph/Source/OptimusCore/Public/IOptimusParameterBindingProvider.h
+++ b/Engine/Plugins/Animation/DeformerGraph/Source/OptimusCore/Public/IOptimusParameterBindingProvider.h
@@ -23,4 +23,5 @@ class OPTIMUSCORE_API IOptimusParameterBindingProvider
 
 public:
 	virtual FString GetBindingDeclaration(FName BindingName) const = 0;
+	virtual bool GetBindingAtomicSupportCheckBoxVisibility(FName BindingName) const = 0;
 };
diff --git a/Engine/Plugins/Animation/DeformerGraph/Source/OptimusCore/Public/OptimusBindingTypes.h b/Engine/Plugins/Animation/DeformerGraph/Source/OptimusCore/Public/OptimusBindingTypes.h
--- a/Engine/Plugins/Animation/DeformerGraph/Source/OptimusCore/Public/OptimusBindingTypes.h
+++ b/Engine/Plugins/Animation/DeformerGraph/Source/OptimusCore/Public/OptimusBindingTypes.h
@@ -45,6 +45,10 @@ struct FOptimusParameterBinding
 	
 	UPROPERTY(EditAnywhere, Category = Binding)
 	FOptimusDataDomain DataDomain;
+
+	// Int type resource can optionally support atomic writes, memory is zero-initialized
+	UPROPERTY(EditAnywhere, DisplayName="Support Atomic", Category = Binding)
+	bool bSupportAtomicIfCompatibleDataType;
 };
 
 
@@ -62,6 +66,12 @@ struct FOptimusParameterBindingArray
 		return InnerArray.FindByPredicate(Pred);
 	}
 
+	template <typename Predicate>
+	FOptimusParameterBinding* FindByPredicate(Predicate Pred)
+	{
+		return InnerArray.FindByPredicate(Pred);
+	}
+
 	bool IsEmpty() const
 	{
 		return InnerArray.IsEmpty();
diff --git a/Engine/Plugins/Animation/DeformerGraph/Source/OptimusCore/Public/OptimusDataDomain.h b/Engine/Plugins/Animation/DeformerGraph/Source/OptimusCore/Public/OptimusDataDomain.h
--- a/Engine/Plugins/Animation/DeformerGraph/Source/OptimusCore/Public/OptimusDataDomain.h
+++ b/Engine/Plugins/Animation/DeformerGraph/Source/OptimusCore/Public/OptimusDataDomain.h
@@ -94,6 +94,22 @@ struct OPTIMUSCORE_API FOptimusDataDomain
 	{
 		return Type == EOptimusDataDomainType::Dimensional && DimensionNames.Num() > 1;
 	}
+
+	int NumDimensions() const
+	{
+		if (Type == EOptimusDataDomainType::Dimensional)
+		{
+			return DimensionNames.Num();
+		}
+
+		if (Type == EOptimusDataDomainType::Expression)
+		{
+			return 1;
+		}
+
+		checkNoEntry();
+		return 0;
+	}
 	
 	/** Convenience function to check if this data domain is One-dimensional*/
 	bool IsOneDimensional() const
diff --git a/Engine/Plugins/Animation/DeformerGraph/Source/OptimusCore/Public/OptimusObjectVersion.h b/Engine/Plugins/Animation/DeformerGraph/Source/OptimusCore/Public/OptimusObjectVersion.h
--- a/Engine/Plugins/Animation/DeformerGraph/Source/OptimusCore/Public/OptimusObjectVersion.h
+++ b/Engine/Plugins/Animation/DeformerGraph/Source/OptimusCore/Public/OptimusObjectVersion.h
@@ -20,6 +20,7 @@ struct OPTIMUSCORE_API FOptimusObjectVersion
 		SetPrimaryBindingName,
 		DataDomainExpansion,
 		KernelDataInterface,
+		KernelParameterBindingToggleAtomic,
 
 		// -----<new versions can be added above this line>-------------------------------------------------
 		VersionPlusOne,
diff --git a/Engine/Plugins/Animation/DeformerGraph/Source/OptimusEditor/Private/OptimusDetailsCustomization.cpp b/Engine/Plugins/Animation/DeformerGraph/Source/OptimusEditor/Private/OptimusDetailsCustomization.cpp
--- a/Engine/Plugins/Animation/DeformerGraph/Source/OptimusEditor/Private/OptimusDetailsCustomization.cpp
+++ b/Engine/Plugins/Animation/DeformerGraph/Source/OptimusEditor/Private/OptimusDetailsCustomization.cpp
@@ -1100,8 +1100,18 @@ public:
 		CustomizationUtils = &InCustomizationUtils;
 		
 		TSharedPtr<IPropertyHandle> DataTypeProperty = BindingPropertyHandle->GetChildHandle(GET_MEMBER_NAME_CHECKED(FOptimusParameterBinding, DataType));
-		const TSharedPtr<IPropertyHandle> DataDomainProperty = BindingPropertyHandle->GetChildHandle(GET_MEMBER_NAME_CHECKED(FOptimusParameterBinding, DataDomain));	
-	
+		const TSharedPtr<IPropertyHandle> DataDomainProperty = BindingPropertyHandle->GetChildHandle(GET_MEMBER_NAME_CHECKED(FOptimusParameterBinding, DataDomain));
+		const TSharedPtr<IPropertyHandle> SupportAtomicProperty = BindingPropertyHandle->GetChildHandle(GET_MEMBER_NAME_CHECKED(FOptimusParameterBinding, bSupportAtomicIfCompatibleDataType));
+
+		const TArray<TWeakObjectPtr<UObject>>& SelectedObjects = InCustomizationUtils.GetPropertyUtilities()->GetSelectedObjects();
+
+		auto IsAtomicCheckBoxVisible = [SelectedObjects, InBindingPropertyHandle]()
+		{
+			return FOptimusParameterBindingCustomization::IsAtomicCheckBoxVisible(SelectedObjects, InBindingPropertyHandle);
+		};
+
+		TSharedRef<SWidget> AtomicToggleWidget = SupportAtomicProperty->CreatePropertyValueWidget(false);
+		
 		FDetailWidgetRow DataTypeHeaderRow;
 		DataTypeRefCustomizationInstance = FOptimusDataTypeRefCustomization::MakeInstance();
 		DataTypeRefCustomizationInstance->CustomizeHeader(DataTypeProperty.ToSharedRef(), DataTypeHeaderRow, InCustomizationUtils);
@@ -1148,7 +1158,22 @@ public:
 					// padding values grabbed from DetailWidgetConstants
 					.Padding(0,0,10,0)
 					[
-						DataTypeHeaderRow.ValueContent().Widget
+						SNew(SHorizontalBox)
+						+SHorizontalBox::Slot()
+						[
+							DataTypeHeaderRow.ValueContent().Widget
+						]
+						+SHorizontalBox::Slot()
+						.AutoWidth()
+						[
+							SNew(SHorizontalBox)
+							.Visibility_Lambda(IsAtomicCheckBoxVisible)
+							+SHorizontalBox::Slot()
+							.Padding(5,0,0,0)
+							[
+								AtomicToggleWidget
+							]
+						]
 					]
 				]
 				+SSplitter::Slot()
@@ -1237,6 +1262,31 @@ TSharedRef<IPropertyTypeCustomization> FOptimusParameterBindingCustomization::Ma
 	return MakeShared<FOptimusParameterBindingCustomization>();
 }
 
+EVisibility FOptimusParameterBindingCustomization::IsAtomicCheckBoxVisible(const TArray<TWeakObjectPtr<UObject>>& InSelectedObjects,  TSharedRef<IPropertyHandle> InPropertyHandle)
+{
+	for (TWeakObjectPtr<UObject> Object : InSelectedObjects)
+	{
+		if (IOptimusParameterBindingProvider* BindingProvider = Cast<IOptimusParameterBindingProvider>(Object))
+		{
+			TArray<const void *> RawData;
+
+			InPropertyHandle->AccessRawData(RawData);
+			if (RawData.Num() > 0)
+			{
+				const FOptimusParameterBinding* Binding = static_cast<const FOptimusParameterBinding*>(RawData[0]);
+				// During drag & reorder, we can have invalid bindings in the property
+				if (Binding && Binding->Name != NAME_None)
+				{
+					return BindingProvider->GetBindingAtomicSupportCheckBoxVisibility(Binding->Name) ? EVisibility::Visible : EVisibility::Collapsed;
+				}
+			}
+			break;
+		}
+	}
+		
+	return EVisibility::Collapsed;
+}
+
 FOptimusParameterBindingCustomization::FOptimusParameterBindingCustomization()
 {
 	
@@ -1302,36 +1352,44 @@ void FOptimusParameterBindingCustomization::CustomizeChildren(
 	IPropertyTypeCustomizationUtils& InCustomizationUtils
 	)
 {
-	FString Declaration;
+	
 	const TArray<TWeakObjectPtr<UObject>>& SelectedObjects = InCustomizationUtils.GetPropertyUtilities()->GetSelectedObjects();
-	for (TWeakObjectPtr<UObject> Object : SelectedObjects)
+
+	auto GetDeclarationText = [SelectedObjects, InPropertyHandle]()
 	{
-		if (IOptimusParameterBindingProvider* BindingProvider = Cast<IOptimusParameterBindingProvider>(Object))
+		FString Declaration;
+		
+		for (TWeakObjectPtr<UObject> Object : SelectedObjects)
 		{
-			TArray<const void *> RawData;
-
-			InPropertyHandle->AccessRawData(RawData);
-			if (ensure(RawData.Num() > 0))
+			if (IOptimusParameterBindingProvider* BindingProvider = Cast<IOptimusParameterBindingProvider>(Object))
 			{
-				const FOptimusParameterBinding* Binding = static_cast<const FOptimusParameterBinding*>(RawData[0]);
-				// During drag & reorder, we can have invalid bindings in the property
-				if (Binding->Name != NAME_None)
+				TArray<const void *> RawData;
+
+				InPropertyHandle->AccessRawData(RawData);
+				if (ensure(RawData.Num() > 0))
 				{
-					if (Binding->DataType->ShaderValueType.IsValid())
-					{
-						Declaration = BindingProvider->GetBindingDeclaration(Binding->Name);
-					}
-					else
+					const FOptimusParameterBinding* Binding = static_cast<const FOptimusParameterBinding*>(RawData[0]);
+					// During drag & reorder, we can have invalid bindings in the property
+					if (Binding && Binding->Name != NAME_None)
 					{
-						Declaration = FString::Printf(TEXT("Type is not supported"));
+						if (Binding->DataType->ShaderValueType.IsValid())
+						{
+							Declaration = BindingProvider->GetBindingDeclaration(Binding->Name);
+						}
+						else
+						{
+							Declaration = FString::Printf(TEXT("Type is not supported"));
+						}
 					}
 				}
+				break;
 			}
-			break;
 		}
-	}
 
-	if (!Declaration.IsEmpty())
+		return FText::FromString(Declaration);
+	};
+	
+	if (!GetDeclarationText().IsEmpty())
 	{
 		FDetailWidgetRow& DeclarationRow = InChildBuilder.AddCustomRow(FText::GetEmpty());
 		DeclarationRow
@@ -1346,12 +1404,21 @@ void FOptimusParameterBindingCustomization::CustomizeChildren(
 			.MinDesiredWidth(180.0f)
 			[
 				SNew(SMultiLineEditableTextBox)
-				.Text(FText::FromString(Declaration))
+				.Text_Lambda(GetDeclarationText)
 				.Font(FCoreStyle::GetDefaultFontStyle("Mono",InCustomizationUtils.GetRegularFont().Size))
 				.IsReadOnly(true)
 			]
 		];		
 	}
+
+	auto IsAtomicRowVisible = [SelectedObjects, InPropertyHandle]()
+	{
+		return IsAtomicCheckBoxVisible(SelectedObjects, InPropertyHandle);
+	};
+	
+	TSharedPtr<IPropertyHandle> SupportAtomicHandle = InPropertyHandle->GetChildHandle(GET_MEMBER_NAME_CHECKED(FOptimusParameterBinding, bSupportAtomicIfCompatibleDataType), false);
+	IDetailPropertyRow& AtomicRow = InChildBuilder.AddProperty(SupportAtomicHandle.ToSharedRef());
+	AtomicRow.Visibility(TAttribute<EVisibility>::CreateLambda(IsAtomicRowVisible));
 }
 
 TSharedRef<FOptimusParameterBindingArrayBuilder> FOptimusParameterBindingArrayBuilder::MakeInstance(
diff --git a/Engine/Plugins/Animation/DeformerGraph/Source/OptimusEditor/Private/OptimusDetailsCustomization.h b/Engine/Plugins/Animation/DeformerGraph/Source/OptimusEditor/Private/OptimusDetailsCustomization.h
--- a/Engine/Plugins/Animation/DeformerGraph/Source/OptimusEditor/Private/OptimusDetailsCustomization.h
+++ b/Engine/Plugins/Animation/DeformerGraph/Source/OptimusEditor/Private/OptimusDetailsCustomization.h
@@ -237,6 +237,7 @@ public:
 	};
 	
 	static TSharedRef<IPropertyTypeCustomization> MakeInstance();
+	static EVisibility IsAtomicCheckBoxVisible(const TArray<TWeakObjectPtr<UObject>>& InSelectedObjects,  TSharedRef<IPropertyHandle> InPropertyHandle);
 
 	FOptimusParameterBindingCustomization();
 
