Optimus: fix compilation error

#rnx

diff --git a/Engine/Plugins/Animation/DeformerGraph/Source/OptimusCore/Private/Nodes/OptimusNode_ComputeKernelFunction.h b/Engine/Plugins/Animation/DeformerGraph/Source/OptimusCore/Private/Nodes/OptimusNode_ComputeKernelFunction.h
--- a/Engine/Plugins/Animation/DeformerGraph/Source/OptimusCore/Private/Nodes/OptimusNode_ComputeKernelFunction.h
+++ b/Engine/Plugins/Animation/DeformerGraph/Source/OptimusCore/Private/Nodes/OptimusNode_ComputeKernelFunction.h
@@ -79,6 +79,7 @@ public:
 	FOptimusExecutionDomain GetExecutionDomain() const override;
 	const UOptimusNodePin* GetPrimaryGroupPin() const override { return {}; }
 	UComputeDataInterface* MakeKernelDataInterface(UObject* InOuter) const override { return nullptr; };
+	bool GetPinSupportAtomic(const UOptimusNodePin* InPin) const override {return false;};
 
 private:
 	UOptimusNode_ComputeKernelFunctionGeneratorClass *GetGeneratorClass() const;
