Optimus: Fix uninitialized member property.

#rnx

diff --git a/Engine/Plugins/Animation/DeformerGraph/Source/OptimusCore/Public/OptimusBindingTypes.h b/Engine/Plugins/Animation/DeformerGraph/Source/OptimusCore/Public/OptimusBindingTypes.h
--- a/Engine/Plugins/Animation/DeformerGraph/Source/OptimusCore/Public/OptimusBindingTypes.h
+++ b/Engine/Plugins/Animation/DeformerGraph/Source/OptimusCore/Public/OptimusBindingTypes.h
@@ -48,7 +48,7 @@ struct FOptimusParameterBinding
 
 	// Int type resource can optionally support atomic writes, memory is zero-initialized
 	UPROPERTY(EditAnywhere, DisplayName="Support Atomic", Category = Binding)
-	bool bSupportAtomicIfCompatibleDataType;
+	bool bSupportAtomicIfCompatibleDataType = false;
 };
 
 
