Optimus: fix nullptr
#rnx

diff --git a/Engine/Plugins/Animation/DeformerGraph/Source/OptimusCore/Private/Nodes/OptimusNode_ComputeKernelBase.cpp b/Engine/Plugins/Animation/DeformerGraph/Source/OptimusCore/Private/Nodes/OptimusNode_ComputeKernelBase.cpp
--- a/Engine/Plugins/Animation/DeformerGraph/Source/OptimusCore/Private/Nodes/OptimusNode_ComputeKernelBase.cpp
+++ b/Engine/Plugins/Animation/DeformerGraph/Source/OptimusCore/Private/Nodes/OptimusNode_ComputeKernelBase.cpp
@@ -702,6 +702,13 @@ void UOptimusNode_ComputeKernelBase::ProcessOutputPinForComputeKernel(
 	{
 		ComponentBinding = ComponentBindings[0];
 	}
+
+	if (!ensure(ComponentBinding))
+	{
+		// Validate For Compile should have caught this
+		return;
+	}
+	
 	
 	if (!ConnectedPins.IsEmpty())
 	{
