Optimus: kernels with constant inputs are extracted to run just once. their outputs are saved to a persistent resource and reused every tick.

+support read from a output pin

diff --git a/Engine/Plugins/Animation/DeformerGraph/Shaders/Private/DataInterfaceRawBuffer.ush b/Engine/Plugins/Animation/DeformerGraph/Shaders/Private/DataInterfaceRawBuffer.ush
--- a/Engine/Plugins/Animation/DeformerGraph/Shaders/Private/DataInterfaceRawBuffer.ush
+++ b/Engine/Plugins/Animation/DeformerGraph/Shaders/Private/DataInterfaceRawBuffer.ush
@@ -45,6 +45,20 @@ PUBLIC_TYPE ReadValue_{DataInterfaceName}(uint Index)
 	return ReturnValue;
 }
 
+PUBLIC_TYPE ReadValueUAV_{DataInterfaceName}(uint Index)
+{
+	uint BufferIndex = {DataInterfaceName}_StartOffset + Index;
+	PUBLIC_TYPE ReturnValue;
+
+#if BUFFER_TYPE_RAW
+	ReadRawBuffer({DataInterfaceName}_BufferUAV, BufferIndex, ReturnValue);
+#else
+	ReturnValue = {DataInterfaceName}_BufferUAV[BufferIndex];
+#endif
+	
+	return ReturnValue;
+}
+
 void WriteValue_{DataInterfaceName}(uint Index, PUBLIC_TYPE Value)
 {
 	uint BufferIndex = {DataInterfaceName}_StartOffset + Index;
diff --git a/Engine/Plugins/Animation/DeformerGraph/Source/OptimusCore/Private/DataInterfaces/OptimusDataInterfaceCustomComputeKernel.cpp b/Engine/Plugins/Animation/DeformerGraph/Source/OptimusCore/Private/DataInterfaces/OptimusDataInterfaceCustomComputeKernel.cpp
--- a/Engine/Plugins/Animation/DeformerGraph/Source/OptimusCore/Private/DataInterfaces/OptimusDataInterfaceCustomComputeKernel.cpp
+++ b/Engine/Plugins/Animation/DeformerGraph/Source/OptimusCore/Private/DataInterfaces/OptimusDataInterfaceCustomComputeKernel.cpp
@@ -128,8 +128,10 @@ bool UOptimusCustomComputeKernelDataProvider::GetInvocationThreadCounts(TArray<i
 		return GetInvocationThreadCounts_DEPRECATED(OutInvocationThreadCount, OutTotalThreadCount);
 	}
 
-	TArray<float> Values = DeformerInstance->GetConstantValuePerInvocation(ExecutionDomainIdentifier);
+	const FOptimusConstantEvaluationResult Result = DeformerInstance->GetConstantValuePerInvocation(ExecutionDomainIdentifier);
 
+	const TArray<float>& Values = Result.ValuePerInvocation;
+	
 	// Can happen if the bound component does not have actual data, like when there is no preview mesh
 	if (Values.Num() == 0)
 	{
diff --git a/Engine/Plugins/Animation/DeformerGraph/Source/OptimusCore/Private/DataInterfaces/OptimusDataInterfaceHalfEdge.cpp b/Engine/Plugins/Animation/DeformerGraph/Source/OptimusCore/Private/DataInterfaces/OptimusDataInterfaceHalfEdge.cpp
--- a/Engine/Plugins/Animation/DeformerGraph/Source/OptimusCore/Private/DataInterfaces/OptimusDataInterfaceHalfEdge.cpp
+++ b/Engine/Plugins/Animation/DeformerGraph/Source/OptimusCore/Private/DataInterfaces/OptimusDataInterfaceHalfEdge.cpp
@@ -298,8 +298,8 @@ TArray<FOptimusCDIPinDefinition> UOptimusHalfEdgeDataInterface::GetPinDefinition
 	FName Triangle(UOptimusSkinnedMeshComponentSource::Domains::Triangle);
 
 	TArray<FOptimusCDIPinDefinition> Defs;
-	Defs.Add({ "EdgesPerVertex", "ReadEdge", Vertex, "ReadNumVertices" });
-	Defs.Add({ "TwinEdges", "ReadTwinEdge", Triangle, 3, "ReadNumTriangles" });
+	Defs.Add({ "EdgesPerVertex", "ReadEdge", Vertex, "ReadNumVertices", false});
+	Defs.Add({ "TwinEdges", "ReadTwinEdge", Triangle, 3, "ReadNumTriangles", false});
 	return Defs;
 }
 
diff --git a/Engine/Plugins/Animation/DeformerGraph/Source/OptimusCore/Private/DataInterfaces/OptimusDataInterfaceRawBuffer.cpp b/Engine/Plugins/Animation/DeformerGraph/Source/OptimusCore/Private/DataInterfaces/OptimusDataInterfaceRawBuffer.cpp
--- a/Engine/Plugins/Animation/DeformerGraph/Source/OptimusCore/Private/DataInterfaces/OptimusDataInterfaceRawBuffer.cpp
+++ b/Engine/Plugins/Animation/DeformerGraph/Source/OptimusCore/Private/DataInterfaces/OptimusDataInterfaceRawBuffer.cpp
@@ -68,6 +68,7 @@ TSubclassOf<UActorComponent> UOptimusRawBufferDataInterface::GetRequiredComponen
 
 void UOptimusRawBufferDataInterface::GetSupportedInputs(TArray<FShaderFunctionDefinition>& OutFunctions) const
 {
+	// Functions in order of EOptimusBufferReadType
 	OutFunctions.AddDefaulted_GetRef()
 		.SetName(TEXT("ReadNumValues"))
 		.AddReturnType(EShaderFundamentalType::Uint);
@@ -76,6 +77,11 @@ void UOptimusRawBufferDataInterface::GetSupportedInputs(TArray<FShaderFunctionDe
 		.SetName(TEXT("ReadValue"))
 		.AddReturnType(ValueType)
 		.AddParam(EShaderFundamentalType::Uint);
+
+	OutFunctions.AddDefaulted_GetRef()
+		.SetName(TEXT("ReadValueUAV"))
+		.AddReturnType(ValueType)
+		.AddParam(EShaderFundamentalType::Uint);
 }
 
 void UOptimusRawBufferDataInterface::GetSupportedOutputs(TArray<FShaderFunctionDefinition>& OutFunctions) const
@@ -105,9 +111,9 @@ void UOptimusRawBufferDataInterface::GetSupportedOutputs(TArray<FShaderFunctionD
 		.AddParam(ValueType);
 }
 
-int32 UOptimusRawBufferDataInterface::GetReadValueInputIndex()
+int32 UOptimusRawBufferDataInterface::GetReadValueInputIndex(EOptimusBufferReadType ReadType)
 {
-	return 1;
+	return (int32)ReadType;
 }
 
 int32 UOptimusRawBufferDataInterface::GetWriteValueOutputIndex(EOptimusBufferWriteType WriteType)
@@ -127,6 +133,19 @@ void UOptimusTransientBufferDataInterface::GetShaderParameters(TCHAR const* UID,
 	InOutBuilder.AddNestedStruct<FTransientBufferDataInterfaceParameters>(UID);
 }
 
+BEGIN_SHADER_PARAMETER_STRUCT(FImplicitPersistentBufferDataInterfaceParameters, )
+	SHADER_PARAMETER(uint32, StartOffset)
+	SHADER_PARAMETER(uint32, BufferSize)
+	SHADER_PARAMETER_RDG_BUFFER_SRV(StructuredBuffer<int>, BufferSRV)
+	SHADER_PARAMETER_RDG_BUFFER_UAV(RWStructuredBuffer<int>, BufferUAV)
+END_SHADER_PARAMETER_STRUCT()
+
+void UOptimusImplicitPersistentBufferDataInterface::GetShaderParameters(TCHAR const* UID,
+	FShaderParametersMetadataBuilder& InOutBuilder, FShaderParametersMetadataAllocations& InOutAllocations) const
+{
+	InOutBuilder.AddNestedStruct<FImplicitPersistentBufferDataInterfaceParameters>(UID);
+}
+
 BEGIN_SHADER_PARAMETER_STRUCT(FPersistentBufferDataInterfaceParameters, )
 	SHADER_PARAMETER(uint32, StartOffset)
 	SHADER_PARAMETER(uint32, BufferSize)
@@ -181,8 +200,20 @@ UComputeDataProvider* UOptimusTransientBufferDataInterface::CreateDataProvider(
 	) const
 {
 	UOptimusTransientBufferDataProvider *Provider = CreateProvider<UOptimusTransientBufferDataProvider>(InBinding);
-	Provider->ElementStride = ValueType->GetResourceElementSize();
-	Provider->RawStride = GetRawStride();
+	Provider->bZeroInitForAtomicWrites = bZeroInitForAtomicWrites;
+	return Provider;
+}
+
+
+FString UOptimusImplicitPersistentBufferDataInterface::GetDisplayName() const
+{
+	return TEXT("ImplicitPersistent");
+}
+
+UComputeDataProvider* UOptimusImplicitPersistentBufferDataInterface::CreateDataProvider(TObjectPtr<UObject> InBinding,
+	uint64 InInputMask, uint64 InOutputMask) const
+{
+	UOptimusImplicitPersistentBufferDataProvider *Provider = CreateProvider<UOptimusImplicitPersistentBufferDataProvider>(InBinding);
 	Provider->bZeroInitForAtomicWrites = bZeroInitForAtomicWrites;
 	return Provider;
 }
@@ -380,16 +411,24 @@ bool UOptimusRawBufferDataProvider::GetLodAndInvocationElementCounts(
 	return false;
 }
 
-bool UOptimusRawBufferDataProvider::GetInvocationElementCounts(TArray<int32>& OutInvocationElementCounts) const
+bool UOptimusRawBufferDataProvider::GetLodContextAndInvocationElementCounts(
+	TArray<int32>& OutInvocationElementCounts,
+	FOptimusDeformerInstanceComponentLodContext* OutLodContext
+	) const
 {
-	TArray<float> Values = DeformerInstance->GetConstantValuePerInvocation(DomainConstantIdentifier);
+	FOptimusConstantEvaluationResult Result = DeformerInstance->GetConstantValuePerInvocation(DomainConstantIdentifier);
 
 	// Can happen if the bound component does not have actual data, like when there is no preview mesh
-	if (Values.Num() == 0)
+	if (!Result.IsValid())
 	{
 		return false;
 	}
-	
+
+	if (OutLodContext)
+	{
+		*OutLodContext = Result.LodContext;
+	}
+	const TArray<float>& Values = Result.ValuePerInvocation;
 	OutInvocationElementCounts.Reset(Values.Num());
 	for (const float& Value : Values)
 	{
@@ -418,7 +457,7 @@ FComputeDataProviderRenderProxy* UOptimusTransientBufferDataProvider::GetRenderP
 	if (DomainConstantIdentifier.IsValid())
 	{
 		// Querying the deformer instance for the domain of this buffer
-		if (!GetInvocationElementCounts(InvocationCounts))
+		if (!GetLodContextAndInvocationElementCounts(InvocationCounts))
 		{
 			InvocationCounts.Reset();
 		}
@@ -434,6 +473,28 @@ FComputeDataProviderRenderProxy* UOptimusTransientBufferDataProvider::GetRenderP
 	return new FOptimusTransientBufferDataProviderProxy(InvocationCounts, ElementStride, RawStride, bZeroInitForAtomicWrites);
 }
 
+FComputeDataProviderRenderProxy* UOptimusImplicitPersistentBufferDataProvider::GetRenderProxy()
+{
+	FOptimusDeformerInstanceComponentLodContext LodContext;
+	TArray<int32> InvocationCounts;
+
+	// Identifier can be unassigned when there isn't a valid component bound, see UOptimusRawBufferDataInterface::CreateProvider
+	if (DomainConstantIdentifier.IsValid())
+	{
+		// Querying the deformer instance for the domain of this buffer
+		if (!GetLodContextAndInvocationElementCounts(InvocationCounts, &LodContext))
+		{
+			InvocationCounts.Reset();
+		}
+	}
+	
+	return new FOptimusImplicitPersistentBufferDataProviderProxy(
+		InvocationCounts, ElementStride, RawStride, bZeroInitForAtomicWrites,
+		BufferPool,
+		DomainConstantIdentifier,
+		LodContext);
+}
+
 
 FComputeDataProviderRenderProxy* UOptimusPersistentBufferDataProvider::GetRenderProxy()
 {
@@ -513,6 +574,80 @@ void FOptimusTransientBufferDataProviderProxy::GatherDispatchData(FDispatchData
 	}
 }
 
+FOptimusImplicitPersistentBufferDataProviderProxy::FOptimusImplicitPersistentBufferDataProviderProxy(
+	TArray<int32> InInvocationElementCounts,
+	int32 InElementStride,
+	int32 InRawStride,
+	bool bInZeroInitForAtomicWrites,
+	TSharedPtr<FOptimusPersistentBufferPool> InBufferPool,
+	const FOptimusConstantIdentifier& InDomainConstantIdentifier,
+	const FOptimusDeformerInstanceComponentLodContext& InLodContext
+	) :
+	InvocationElementCounts(InInvocationElementCounts),
+	TotalElementCount(0),
+	ElementStride(InElementStride),
+	RawStride(InRawStride),
+	bZeroInitForAtomicWrites(bInZeroInitForAtomicWrites),
+	BufferPool(InBufferPool),
+	DomainConstantIdentifier(InDomainConstantIdentifier),
+	LodContext(InLodContext)
+{
+	for (int32 NumElements : InvocationElementCounts)
+	{
+		TotalElementCount += NumElements;
+	}
+}
+
+bool FOptimusImplicitPersistentBufferDataProviderProxy::IsValid(FValidationData const& InValidationData) const
+{
+	if (InValidationData.ParameterStructSize != sizeof(FParameters))
+	{
+		return false;
+	}
+	
+	if (TotalElementCount <= 0)
+	{
+		return false;
+	}
+
+	return true;
+}
+
+void FOptimusImplicitPersistentBufferDataProviderProxy::AllocateResources(FRDGBuilder& GraphBuilder)
+{
+	TArray<int32> Count;
+	Count.Add(TotalElementCount);
+	TArray<FRDGBufferRef> Buffers;
+	bool bJustAllocated = false;
+	BufferPool->GetImplicitPersistentBuffers(GraphBuilder, DomainConstantIdentifier, LodContext, ElementStride, RawStride, Count, Buffers, bJustAllocated);
+
+	ensure(Buffers.Num() == 1);
+	Buffer = Buffers[0];
+
+	BufferSRV = GraphBuilder.CreateSRV(Buffer);
+	BufferUAV = GraphBuilder.CreateUAV(Buffer, ERDGUnorderedAccessViewFlags::SkipBarrier);
+
+	if (bZeroInitForAtomicWrites && bJustAllocated)
+	{
+		AddClearUAVPass(GraphBuilder, BufferUAV, 0);
+	}
+}
+
+void FOptimusImplicitPersistentBufferDataProviderProxy::GatherDispatchData(FDispatchData const& InDispatchData)
+{
+	TStridedView<FParameters> ParameterArray = MakeStridedParameterView<FParameters>(InDispatchData);
+	for (int32 InvocationIndex = 0, StartOffset = 0; InvocationIndex < ParameterArray.Num(); ++InvocationIndex)
+	{
+		FParameters& Parameters = ParameterArray[InvocationIndex];
+		Parameters.StartOffset = InDispatchData.bUnifiedDispatch ? 0 : StartOffset;
+		Parameters.BufferSize = InDispatchData.bUnifiedDispatch ? TotalElementCount : InvocationElementCounts[InvocationIndex];
+		Parameters.BufferSRV = BufferSRV;
+		Parameters.BufferUAV = BufferUAV;
+		
+		StartOffset += InvocationElementCounts[InvocationIndex];
+	}
+}
+
 
 FOptimusPersistentBufferDataProviderProxy::FOptimusPersistentBufferDataProviderProxy(
 	TArray<int32> InInvocationElementCounts,
@@ -556,7 +691,8 @@ void FOptimusPersistentBufferDataProviderProxy::AllocateResources(FRDGBuilder& G
 	TArray<int32> Count;
 	Count.Add(TotalElementCount);
 	TArray<FRDGBufferRef> Buffers;
-	BufferPool->GetResourceBuffers(GraphBuilder, ResourceName, LODIndex, ElementStride, RawStride, Count, Buffers);
+	bool bJustAllocated = false;
+	BufferPool->GetResourceBuffers(GraphBuilder, ResourceName, LODIndex, ElementStride, RawStride, Count, Buffers, bJustAllocated);
 
 	ensure(Buffers.Num() == 1);
 	Buffer = Buffers[0];
diff --git a/Engine/Plugins/Animation/DeformerGraph/Source/OptimusCore/Private/DataInterfaces/OptimusDataInterfaceRawBuffer.h b/Engine/Plugins/Animation/DeformerGraph/Source/OptimusCore/Private/DataInterfaces/OptimusDataInterfaceRawBuffer.h
--- a/Engine/Plugins/Animation/DeformerGraph/Source/OptimusCore/Private/DataInterfaces/OptimusDataInterfaceRawBuffer.h
+++ b/Engine/Plugins/Animation/DeformerGraph/Source/OptimusCore/Private/DataInterfaces/OptimusDataInterfaceRawBuffer.h
@@ -3,22 +3,32 @@
 #pragma once
 
 #include "IOptimusDeformerInstanceAccessor.h"
+#include "IOptimusPersistentBufferProvider.h"
 #include "ComputeFramework/ComputeDataProvider.h"
 #include "ComputeFramework/ShaderParamTypeDefinition.h"
 #include "OptimusComputeDataInterface.h"
 #include "OptimusConstant.h"
 #include "OptimusDataDomain.h"
+#include "OptimusDeformerInstance.h"
 #include "RenderGraphFwd.h"
 
 #include "OptimusDataInterfaceRawBuffer.generated.h"
 
 class FOptimusPersistentBufferPool;
-class FPersistentBufferDataInterfaceParameters;
 class FTransientBufferDataInterfaceParameters;
+class FImplicitPersistentBufferDataInterfaceParameters;
+class FPersistentBufferDataInterfaceParameters;
 class UOptimusComponentSource;
 class UOptimusComponentSourceBinding;
 class UOptimusRawBufferDataProvider;
 
+enum class EOptimusBufferReadType : uint8
+{
+	ReadSize,
+	Default,
+	ForceUAV
+};
+
 /** Write to buffer operation types. */
 UENUM()
 enum class EOptimusBufferWriteType : uint8
@@ -36,7 +46,7 @@ class OPTIMUSCORE_API UOptimusRawBufferDataInterface : public UOptimusComputeDat
 	GENERATED_BODY()
 
 public:
-	static int32 GetReadValueInputIndex();
+	static int32 GetReadValueInputIndex(EOptimusBufferReadType ReadType = EOptimusBufferReadType::Default);
 	static int32 GetWriteValueOutputIndex(EOptimusBufferWriteType WriteType);
 	
 	//~ Begin UOptimusComputeDataInterface Interface
@@ -120,6 +130,29 @@ public:
 	bool bZeroInitForAtomicWrites = false;	
 };
 
+/** Compute Framework Data Interface for a transient buffer. */
+UCLASS(Category = ComputeFramework)
+class OPTIMUSCORE_API UOptimusImplicitPersistentBufferDataInterface : public UOptimusRawBufferDataInterface
+{
+	GENERATED_BODY()
+
+public:
+	//~ Begin UOptimusComputeDataInterface Interface
+	FString GetDisplayName() const override;
+	//~ End UOptimusComputeDataInterface Interface
+
+
+	//~ Begin UComputeDataInterface Interface
+	TCHAR const* GetClassName() const override { return TEXT("ImplicitPersistentBuffer"); }
+	bool CanSupportUnifiedDispatch() const override { return true; }
+	void GetShaderParameters(TCHAR const* UID, FShaderParametersMetadataBuilder& InOutBuilder, FShaderParametersMetadataAllocations& InOutAllocations) const override;
+	UComputeDataProvider* CreateDataProvider(TObjectPtr<UObject> InBinding, uint64 InInputMask, uint64 InOutputMask) const override;
+	//~ End UComputeDataInterface Interface
+
+	UPROPERTY()
+	bool bZeroInitForAtomicWrites = false;
+};
+
 
 /** Compute Framework Data Interface for a transient buffer. */
 UCLASS(Category = ComputeFramework)
@@ -170,8 +203,9 @@ public:
 	/** Helper function to calculate the element count given a constant identifier,
 	 *	LOD is handled by the deformer instance
 	 */
-	bool GetInvocationElementCounts(
-		TArray<int32>& OutInvocationElementCounts
+	bool GetLodContextAndInvocationElementCounts(
+		TArray<int32>& OutInvocationElementCounts,
+		FOptimusDeformerInstanceComponentLodContext* OutLodContext = nullptr
 		) const;
 
 	//~ Begin IOptimusDeformerInstanceAccessor Interface
@@ -180,27 +214,20 @@ public:
 	//~ End IOptimusDeformerInstanceAccessor Interface
 	
 	/** The skinned mesh component that governs the sizing and LOD of this buffer */
-	UPROPERTY()
 	TWeakObjectPtr<const UActorComponent> Component = nullptr;
 
-	UPROPERTY()
 	TWeakObjectPtr<const UOptimusComponentSource> ComponentSource = nullptr;
 
 	/** The data domain this buffer covers */
-	UPROPERTY()
 	FOptimusDataDomain DataDomain;
 
-	UPROPERTY()
 	int32 ElementStride = 4;
 
-	UPROPERTY()
 	int32 RawStride = 0;
 
-	UPROPERTY()
 	FOptimusConstantIdentifier DomainConstantIdentifier;
 	
 private:
-	UPROPERTY()
 	TObjectPtr<UOptimusDeformerInstance> DeformerInstance = nullptr;
 };
 
@@ -222,7 +249,9 @@ public:
 
 /** Compute Framework Data Provider for a transient buffer. */
 UCLASS(BlueprintType, editinlinenew, Category = ComputeFramework)
-class OPTIMUSCORE_API UOptimusPersistentBufferDataProvider : public UOptimusRawBufferDataProvider
+class OPTIMUSCORE_API UOptimusImplicitPersistentBufferDataProvider :
+	public UOptimusRawBufferDataProvider,
+	public IOptimusPersistentBufferProvider
 {
 	GENERATED_BODY()
 
@@ -231,13 +260,43 @@ public:
 	FComputeDataProviderRenderProxy* GetRenderProxy() override;
 	//~ End UComputeDataProvider Interface
 
+	//~ Begin IOptimusPersistentBufferPoolUser Interface
+	void SetBufferPool(TSharedPtr<FOptimusPersistentBufferPool> InBufferPool) override { BufferPool = InBufferPool; };
+	//~ Begin IOptimusPersistentBufferPoolUser Interface
+	
+	bool bZeroInitForAtomicWrites = false;
+private:
 	/** The buffer pool we refer to. Set by UOptimusDeformerInstance::SetupFromDeformer after providers have been
 	 *  created
 	 */
 	TSharedPtr<FOptimusPersistentBufferPool> BufferPool;
+};
 
+/** Compute Framework Data Provider for a transient buffer. */
+UCLASS(BlueprintType, editinlinenew, Category = ComputeFramework)
+class OPTIMUSCORE_API UOptimusPersistentBufferDataProvider :
+	public UOptimusRawBufferDataProvider,
+	public IOptimusPersistentBufferProvider
+{
+	GENERATED_BODY()
+
+public:
+	//~ Begin UComputeDataProvider Interface
+	FComputeDataProviderRenderProxy* GetRenderProxy() override;
+	//~ End UComputeDataProvider Interface
+	
+	//~ Begin IOptimusPersistentBufferPoolUser Interface
+	void SetBufferPool(TSharedPtr<FOptimusPersistentBufferPool> InBufferPool) override { BufferPool = InBufferPool; };
+	//~ Begin IOptimusPersistentBufferPoolUser Interface
+	
 	/** The resource this buffer is provider to */
 	FName ResourceName;
+	
+private:
+	/** The buffer pool we refer to. Set by UOptimusDeformerInstance::SetupFromDeformer after providers have been
+	 *  created
+	 */
+	TSharedPtr<FOptimusPersistentBufferPool> BufferPool;
 };
 
 class FOptimusTransientBufferDataProviderProxy :
@@ -271,6 +330,43 @@ private:
 	FRDGBufferUAVRef BufferUAV;
 };
 
+class FOptimusImplicitPersistentBufferDataProviderProxy :
+	public FComputeDataProviderRenderProxy
+{
+public:
+	FOptimusImplicitPersistentBufferDataProviderProxy(
+		TArray<int32> InInvocationElementCounts,
+		int32 InElementStride,
+		int32 InRawStride,
+		bool bInZeroInitForAtomicWrites,
+		TSharedPtr<FOptimusPersistentBufferPool> InBufferPool,
+		const FOptimusConstantIdentifier& InDomainConstantIdentifier,
+		const FOptimusDeformerInstanceComponentLodContext& InLodContext
+		);
+
+	//~ Begin FComputeDataProviderRenderProxy Interface
+	bool IsValid(FValidationData const& InValidationData) const override;
+	void AllocateResources(FRDGBuilder& GraphBuilder) override;
+	void GatherDispatchData(FDispatchData const& InDispatchData) override;
+	//~ End FComputeDataProviderRenderProxy Interface
+
+	private:
+	using FParameters = FImplicitPersistentBufferDataInterfaceParameters;
+
+	const TArray<int32> InvocationElementCounts;
+	int32 TotalElementCount;
+	const int32 ElementStride;
+	const int32 RawStride;
+	const bool bZeroInitForAtomicWrites;
+	
+	const TSharedPtr<FOptimusPersistentBufferPool> BufferPool;
+	const FOptimusConstantIdentifier DomainConstantIdentifier;
+	const FOptimusDeformerInstanceComponentLodContext LodContext;
+
+	FRDGBufferRef Buffer;
+	FRDGBufferSRVRef BufferSRV;
+	FRDGBufferUAVRef BufferUAV;
+};
 
 class FOptimusPersistentBufferDataProviderProxy :
 	public FComputeDataProviderRenderProxy
diff --git a/Engine/Plugins/Animation/DeformerGraph/Source/OptimusCore/Private/DataInterfaces/OptimusDataInterfaceSkinnedMesh.cpp b/Engine/Plugins/Animation/DeformerGraph/Source/OptimusCore/Private/DataInterfaces/OptimusDataInterfaceSkinnedMesh.cpp
--- a/Engine/Plugins/Animation/DeformerGraph/Source/OptimusCore/Private/DataInterfaces/OptimusDataInterfaceSkinnedMesh.cpp
+++ b/Engine/Plugins/Animation/DeformerGraph/Source/OptimusCore/Private/DataInterfaces/OptimusDataInterfaceSkinnedMesh.cpp
@@ -26,15 +26,15 @@ TArray<FOptimusCDIPinDefinition> UOptimusSkinnedMeshDataInterface::GetPinDefinit
 	FName Triangle(UOptimusSkinnedMeshComponentSource::Domains::Triangle);
 	
 	TArray<FOptimusCDIPinDefinition> Defs;
-	Defs.Add({"NumVertices", "ReadNumVertices"});
-	Defs.Add({"Position", "ReadPosition", Vertex, "ReadNumVertices"});
-	Defs.Add({"TangentX", "ReadTangentX", Vertex, "ReadNumVertices"});
-	Defs.Add({"TangentZ", "ReadTangentZ", Vertex, "ReadNumVertices"});
-	Defs.Add({"NumUVChannels", "ReadNumUVChannels"});
-	Defs.Add({"UV", "ReadUV", {{Vertex, "ReadNumVertices"}, {Optimus::DomainName::UVChannel, "ReadNumUVChannels"}}});
-	Defs.Add({"Color", "ReadColor", Vertex, "ReadColor" });
-	Defs.Add({"NumTriangles", "ReadNumTriangles" });
-	Defs.Add({"IndexBuffer", "ReadIndexBuffer", Triangle, 3, "ReadNumTriangles"});
+	Defs.Add({"NumVertices", "ReadNumVertices", false});
+	Defs.Add({"Position", "ReadPosition", Vertex, "ReadNumVertices", false});
+	Defs.Add({"TangentX", "ReadTangentX", Vertex, "ReadNumVertices", false});
+	Defs.Add({"TangentZ", "ReadTangentZ", Vertex, "ReadNumVertices", false});
+	Defs.Add({"NumUVChannels", "ReadNumUVChannels", false});
+	Defs.Add({"UV", "ReadUV", {{Vertex, "ReadNumVertices"}, {Optimus::DomainName::UVChannel, "ReadNumUVChannels"}}, false});
+	Defs.Add({"Color", "ReadColor", Vertex, "ReadColor", false});
+	Defs.Add({"NumTriangles", "ReadNumTriangles", false});
+	Defs.Add({"IndexBuffer", "ReadIndexBuffer", Triangle, 3, "ReadNumTriangles", false});
 	return Defs;
 }
 
diff --git a/Engine/Plugins/Animation/DeformerGraph/Source/OptimusCore/Private/Nodes/OptimusNode_ComputeKernelBase.cpp b/Engine/Plugins/Animation/DeformerGraph/Source/OptimusCore/Private/Nodes/OptimusNode_ComputeKernelBase.cpp
--- a/Engine/Plugins/Animation/DeformerGraph/Source/OptimusCore/Private/Nodes/OptimusNode_ComputeKernelBase.cpp
+++ b/Engine/Plugins/Animation/DeformerGraph/Source/OptimusCore/Private/Nodes/OptimusNode_ComputeKernelBase.cpp
@@ -59,7 +59,6 @@ FOptimus_ComputeKernelResult UOptimusNode_ComputeKernelBase::CreateComputeKernel
 	const FOptimus_PinToDataInterfaceMap& InLinkDataInterfaceMap,
 	const TArray<const UOptimusNode *>& InValueNodes,
 	const UComputeDataInterface* InGraphDataInterface,
-	const UOptimusComponentSourceBinding* InGraphDataComponentBinding,
 	UComputeDataInterface* InOutKernelDataInterface,
 	FOptimus_InterfaceBindingMap& OutInputDataBindings,
 	FOptimus_InterfaceBindingMap& OutOutputDataBindings,
@@ -138,7 +137,7 @@ FOptimus_ComputeKernelResult UOptimusNode_ComputeKernelBase::CreateComputeKernel
 		}
 	}
 
-	TSet<UOptimusComponentSourceBinding*> PrimaryBindings = GetGroupComponentSourceBindings(GetPrimaryGroupPin());
+	TSet<UOptimusComponentSourceBinding*> PrimaryBindings = GetPrimaryGroupPin()->GetComponentSourceBindingsRecursively();
 	
 	// ValidateForCompile() should guaranteed that we don't hit this error
 	if (!ensure(PrimaryBindings.Num() == 1))
@@ -163,7 +162,7 @@ FOptimus_ComputeKernelResult UOptimusNode_ComputeKernelBase::CreateComputeKernel
 				TOptional<FText> Result = ProcessInputPinForComputeKernel(
 					InTraversalContext, SubPin, GroupIndex == 0 ? FString() : Pin->GetName(),
 					InNodeDataInterfaceMap, InLinkDataInterfaceMap, InValueNodes,
-					InGraphDataInterface, InGraphDataComponentBinding,
+					InGraphDataInterface,
 					KernelSource, GeneratedFunctions, OutInputDataBindings, OutKernelConstantContainer
 					);
 				if (Result.IsSet())
@@ -179,7 +178,7 @@ FOptimus_ComputeKernelResult UOptimusNode_ComputeKernelBase::CreateComputeKernel
 			ProcessOutputPinForComputeKernel(
 				InTraversalContext, Pin,
 				InNodeDataInterfaceMap, InLinkDataInterfaceMap,
-				KernelSource, GeneratedFunctions, OutOutputDataBindings, OutKernelConstantContainer);
+				KernelSource, GeneratedFunctions, OutInputDataBindings, OutOutputDataBindings, OutKernelConstantContainer);
 		}
 	}
 
@@ -243,6 +242,25 @@ FOptimus_ComputeKernelResult UOptimusNode_ComputeKernelBase::CreateComputeKernel
 	return FOptimus_ComputeKernelResult(TInPlaceType<UOptimusKernelSource*>(), KernelSource);
 }
 
+
+bool UOptimusNode_ComputeKernelBase::HasMutableInput() const
+{
+	for (const UOptimusNodePin* Pin: GetPins())
+	{
+		if (Pin->IsGroupingPin())
+		{
+			for (const UOptimusNodePin* SubPin: Pin->GetSubPins())
+			{
+				if (GetOwningGraph()->IsPinMutable(SubPin))
+				{
+					return true;
+				}
+			}
+		}
+	}
+	return false;
+}
+
 TOptional<FText> UOptimusNode_ComputeKernelBase::ValidateForCompile() const
 {
 	auto GetStructTypeDefFromPin = [](const UOptimusNodePin* InPin) -> TOptional<FText>
@@ -315,7 +333,7 @@ TOptional<FText> UOptimusNode_ComputeKernelBase::ValidateForCompile() const
 
 		if (Pin->GetDirection() == EOptimusNodePinDirection::Output)
 		{
-			if (GetPinSupportAtomic(Pin))
+			if (DoesOutputPinSupportAtomic(Pin) || DoesOutputPinSupportRead(Pin))
 			{
 				TArray<FOptimusRoutedNodePin> ConnectedPins = Pin->GetConnectedPinsWithRouting();
 				if (ConnectedPins.Num() > 1)
@@ -335,13 +353,13 @@ TOptional<FText> UOptimusNode_ComputeKernelBase::ValidateForCompile() const
 						}
 						else
 						{
-							return FText::Format(LOCTEXT("InvalidAtomicPinConnection", "Pin '{0}' supports Atomic writes and should be connected to either a single resource node or kernel nodes"), FText::FromName(Pin->GetUniqueName()));
+							return FText::Format(LOCTEXT("TooManyPinConnections", "Pin '{0}' supports Atomic Writes or Read Access and should be connected to either a single resource node or strictly kernel nodes"), FText::FromName(Pin->GetUniqueName()));
 						}
 					}
 
 					if (PersistentBufferCount + TransientBufferCount > 1)
 					{
-						return FText::Format(LOCTEXT("InvalidAtomicPinConnection", "Pin '{0}' supports Atomic writes and should be connected to either a single resource node or kernel nodes"), FText::FromName(Pin->GetUniqueName()));
+						return FText::Format(LOCTEXT("TooManyPinConnections", "Pin '{0}' supports Atomic Writes or Read Access and should be connected to either a single resource node or strictly kernel nodes"), FText::FromName(Pin->GetUniqueName()));
 					}
 				}
 			}
@@ -417,19 +435,7 @@ FString UOptimusNode_ComputeKernelBase::GetCookedKernelSource(
 	}
 }
 
-TSet<UOptimusComponentSourceBinding*> UOptimusNode_ComputeKernelBase::GetGroupComponentSourceBindings(const UOptimusNodePin* InGroupPin)
-{
-	TSet<UOptimusComponentSourceBinding*> Bindings;
 
-	Bindings.Append(InGroupPin->GetComponentSourceBindings());
-	
-	for (const UOptimusNodePin* Pin: InGroupPin->GetSubPins())
-	{
-		Bindings.Append(Pin->GetComponentSourceBindings());
-	}
-	
-	return Bindings;
-}
 
 FString UOptimusNode_ComputeKernelBase::GetAtomicWriteFunctionName(
 	EOptimusBufferWriteType InWriteType,
@@ -457,6 +463,11 @@ FString UOptimusNode_ComputeKernelBase::GetAtomicWriteFunctionName(
 	return FunctionName;
 }
 
+FString UOptimusNode_ComputeKernelBase::GetReadFunctionName(const FString& InBindingName)
+{
+	return FString::Printf(TEXT("Read%s"), *InBindingName);
+}
+
 
 TOptional<FText> UOptimusNode_ComputeKernelBase::ProcessInputPinForComputeKernel(
 	const FOptimusPinTraversalContext& InTraversalContext,
@@ -466,7 +477,6 @@ TOptional<FText> UOptimusNode_ComputeKernelBase::ProcessInputPinForComputeKernel
 	const FOptimus_PinToDataInterfaceMap& InLinkDataInterfaceMap,
 	const TArray<const UOptimusNode*>& InValueNodes,
 	const UComputeDataInterface* InGraphDataInterface,
-	const UOptimusComponentSourceBinding* InGraphDataComponentBinding,
 	UOptimusKernelSource* InKernelSource,
 	TArray<FString>& OutGeneratedFunctions,
 	FOptimus_InterfaceBindingMap& OutInputDataBindings,
@@ -501,18 +511,11 @@ TOptional<FText> UOptimusNode_ComputeKernelBase::ProcessInputPinForComputeKernel
 			// For transient buffers we need the function index as given by the
 			// ReadValue function. 
 			DataInterface = InLinkDataInterfaceMap[OutputPin];
-			DataInterfaceFuncIndex = UOptimusTransientBufferDataInterface::GetReadValueInputIndex();
+			DataInterfaceFuncIndex = UOptimusRawBufferDataInterface::GetReadValueInputIndex(EOptimusBufferReadType::Default);
 
 			TArray<FShaderFunctionDefinition> ReadFunctions;
 			DataInterface->GetSupportedInputs(ReadFunctions);
-
-			// Get the component binding from the upstream connection.
-			TArray<UOptimusComponentSourceBinding*> ComponentBindings = GetOwningGraph()->GetComponentSourceBindingsForPin(OutputPin).Array();
-			if (ensure(ComponentBindings.Num() == 1))
-			{
-				ComponentBinding = ComponentBindings[0];
-			}
-
+			
 			DataFunctionName = ReadFunctions[DataInterfaceFuncIndex].Name;
 		}
 		else if(InNodeDataInterfaceMap.Contains(OutputNode))
@@ -522,7 +525,6 @@ TOptional<FText> UOptimusNode_ComputeKernelBase::ProcessInputPinForComputeKernel
 			// FIXME: Sub-pin read support.
 			UOptimusComputeDataInterface const* OptimusDataInterface = InNodeDataInterfaceMap[OutputNode];
 			DataInterface = OptimusDataInterface;
-			ComponentBinding = InterfaceProvider->GetComponentBinding();
 
 			DataInterfaceFuncIndex = InterfaceProvider->GetDataFunctionIndexFromPin(OutputPin);
 			
@@ -535,14 +537,13 @@ TOptional<FText> UOptimusNode_ComputeKernelBase::ProcessInputPinForComputeKernel
 		{
 			// Value nodes bind the single graph data interface.
 			DataInterface = InGraphDataInterface;
-			ComponentBinding = InGraphDataComponentBinding;
 			DataInterfaceFuncIndex = InValueNodes.Find(OutputNode);
 			check(DataInterfaceFuncIndex != INDEX_NONE);
 			DataFunctionName = ValueProvider->GetValueName();
 		}
 
 		// If we are connected from a data interface, set the input binding up now.
-		if (DataInterface && ensure(ComponentBinding))
+		if (DataInterface)
 		{
 			if (!DataInterface->CanSupportUnifiedDispatch() && !InGroupName.IsEmpty())
 			{
@@ -570,7 +571,6 @@ TOptional<FText> UOptimusNode_ComputeKernelBase::ProcessInputPinForComputeKernel
 
 			FOptimus_InterfaceBinding InterfaceBinding;
 			InterfaceBinding.DataInterface = DataInterface;
-			InterfaceBinding.ComponentBinding = ComponentBinding;
 			InterfaceBinding.DataInterfaceBindingIndex = DataInterfaceFuncIndex;
 			InterfaceBinding.BindingFunctionName = FString::Printf(TEXT("Read%s"), *InInputPin->GetName());
 			InterfaceBinding.BindingFunctionNamespace = InGroupName;
@@ -630,7 +630,7 @@ TOptional<FText> UOptimusNode_ComputeKernelBase::ProcessInputPinForComputeKernel
 		
 		if (const UOptimusNodePin* GroupPin = InInputPin->GetParentPin())
 		{
-			TArray<UOptimusComponentSourceBinding*> ComponentBindings = GetGroupComponentSourceBindings(GroupPin).Array();
+			TArray<UOptimusComponentSourceBinding*> ComponentBindings = GroupPin->GetComponentSourceBindingsRecursively().Array();
 			if (ensure(ComponentBindings.Num() == 1))
 			{
 				ComponentBindingIndex = ComponentBindings[0]->GetIndex();
@@ -682,6 +682,7 @@ void UOptimusNode_ComputeKernelBase::ProcessOutputPinForComputeKernel(
 	const FOptimus_PinToDataInterfaceMap& InLinkDataInterfaceMap,
 	UOptimusKernelSource* InKernelSource,
 	TArray<FString>& OutGeneratedFunctions,
+	FOptimus_InterfaceBindingMap& OutInputDataBindings,
 	FOptimus_InterfaceBindingMap& OutOutputDataBindings,
 	FOptimusKernelConstantContainer& OutKernelConstantContainer
 	) const
@@ -794,7 +795,6 @@ void UOptimusNode_ComputeKernelBase::ProcessOutputPinForComputeKernel(
 			
 			FOptimus_InterfaceBinding InterfaceBinding;
 			InterfaceBinding.DataInterface = WriteConnectionDef.DataInterface;
-			InterfaceBinding.ComponentBinding = ComponentBinding;
 			InterfaceBinding.DataInterfaceBindingIndex = DataInterfaceFuncIndex;
 			InterfaceBinding.BindingFunctionName = WrapFunctionName;
 			
@@ -816,7 +816,7 @@ void UOptimusNode_ComputeKernelBase::ProcessOutputPinForComputeKernel(
 					*InOutputPin->GetName(), *FString::Join(IndexParamNames, TEXT(", ")), *ValueType->ToString(), *FString::Join(WrapFunctionNameCalls, TEXT(";\n"))));
 		}	
 		
-		if (GetPinSupportAtomic(InOutputPin))
+		if (DoesOutputPinSupportAtomic(InOutputPin) || DoesOutputPinSupportRead(InOutputPin))
 		{
 			UOptimusRawBufferDataInterface* RawBufferDataInterface = nullptr;
 			if (InLinkDataInterfaceMap.Contains(InOutputPin))
@@ -836,30 +836,53 @@ void UOptimusNode_ComputeKernelBase::ProcessOutputPinForComputeKernel(
 
 			if (ensure(RawBufferDataInterface))
 			{
-				TArray<FShaderFunctionDefinition> WriteFunctions;
-				RawBufferDataInterface->GetSupportedOutputs(WriteFunctions);
-				
-				for (uint8 WriteType = uint8(EOptimusBufferWriteType::Write) + 1; WriteType < uint8(EOptimusBufferWriteType::Count); WriteType++)
+				if (DoesOutputPinSupportAtomic(InOutputPin))
 				{
+					TArray<FShaderFunctionDefinition> WriteFunctions;
+					RawBufferDataInterface->GetSupportedOutputs(WriteFunctions);
+				
+					for (uint8 WriteType = uint8(EOptimusBufferWriteType::Write) + 1; WriteType < uint8(EOptimusBufferWriteType::Count); WriteType++)
+					{
 
-					FString WrapFunctionName = GetAtomicWriteFunctionName((EOptimusBufferWriteType)WriteType, InOutputPin->GetName());
+						FString WrapFunctionName = GetAtomicWriteFunctionName((EOptimusBufferWriteType)WriteType, InOutputPin->GetName());
 
-					const int32 OutputIndex = UOptimusTransientBufferDataInterface::GetWriteValueOutputIndex((EOptimusBufferWriteType)WriteType);
+						const int32 OutputIndex = UOptimusRawBufferDataInterface::GetWriteValueOutputIndex((EOptimusBufferWriteType)WriteType);
 					
-					FShaderFunctionDefinition FuncDef = WriteFunctions[OutputIndex];
+						FShaderFunctionDefinition FuncDef = WriteFunctions[OutputIndex];
+						for (int32 Index = 0; Index < FuncDef.ParamTypes.Num(); Index++)
+						{
+							FuncDef.ParamTypes[Index].ResetTypeDeclaration();
+						}
+				
+						FOptimus_InterfaceBinding InterfaceBinding;
+						InterfaceBinding.DataInterface = RawBufferDataInterface;
+						InterfaceBinding.DataInterfaceBindingIndex = OutputIndex;
+						InterfaceBinding.BindingFunctionName = WrapFunctionName;
+				
+						OutOutputDataBindings.Add(InKernelSource->ExternalOutputs.Num(), InterfaceBinding);
+						InKernelSource->ExternalOutputs.Emplace(FuncDef);	
+					}
+				}
+
+				if (DoesOutputPinSupportRead(InOutputPin))
+				{
+					TArray<FShaderFunctionDefinition> ReadFunctions;
+					RawBufferDataInterface->GetSupportedInputs(ReadFunctions);
+					int32 InputIndex = UOptimusRawBufferDataInterface::GetReadValueInputIndex(EOptimusBufferReadType::ForceUAV);
+
+					FShaderFunctionDefinition FuncDef = ReadFunctions[InputIndex];
 					for (int32 Index = 0; Index < FuncDef.ParamTypes.Num(); Index++)
 					{
 						FuncDef.ParamTypes[Index].ResetTypeDeclaration();
 					}
 				
 					FOptimus_InterfaceBinding InterfaceBinding;
-					InterfaceBinding.ComponentBinding = ComponentBinding;
 					InterfaceBinding.DataInterface = RawBufferDataInterface;
-					InterfaceBinding.DataInterfaceBindingIndex = OutputIndex;
-					InterfaceBinding.BindingFunctionName = WrapFunctionName;
+					InterfaceBinding.DataInterfaceBindingIndex = InputIndex;
+					InterfaceBinding.BindingFunctionName = GetReadFunctionName(InOutputPin->GetName());
 				
-					OutOutputDataBindings.Add(InKernelSource->ExternalOutputs.Num(), InterfaceBinding);
-					InKernelSource->ExternalOutputs.Emplace(FuncDef);	
+					OutInputDataBindings.Add(InKernelSource->ExternalInputs.Num(), InterfaceBinding);
+					InKernelSource->ExternalInputs.Emplace(FuncDef);		
 				}
 			}
 		}
@@ -939,7 +962,6 @@ void UOptimusNode_ComputeKernelBase::BindKernelDataInterfaceForComputeKernel(
 
 		FOptimus_InterfaceBinding InterfaceBinding;
 		InterfaceBinding.DataInterface = InOutKernelDataInterface;
-		InterfaceBinding.ComponentBinding = InKernelPrimaryComponentSourceBinding;
 		InterfaceBinding.DataInterfaceBindingIndex = FuncIndex;
 		InterfaceBinding.BindingFunctionName = FuncDef.Name;
 		InterfaceBinding.BindingFunctionNamespace = FString();
diff --git a/Engine/Plugins/Animation/DeformerGraph/Source/OptimusCore/Private/Nodes/OptimusNode_ComputeKernelBase.h b/Engine/Plugins/Animation/DeformerGraph/Source/OptimusCore/Private/Nodes/OptimusNode_ComputeKernelBase.h
--- a/Engine/Plugins/Animation/DeformerGraph/Source/OptimusCore/Private/Nodes/OptimusNode_ComputeKernelBase.h
+++ b/Engine/Plugins/Animation/DeformerGraph/Source/OptimusCore/Private/Nodes/OptimusNode_ComputeKernelBase.h
@@ -43,7 +43,6 @@ public:
 		const FOptimus_PinToDataInterfaceMap& InLinkDataInterfaceMap,
 		const TArray<const UOptimusNode*>& InValueNodes,
 		const UComputeDataInterface* InGraphDataInterface,
-		const UOptimusComponentSourceBinding* InGraphDataComponentBinding,
 		UComputeDataInterface* InOutKernelDataInterface,
 		FOptimus_InterfaceBindingMap& OutInputDataBindings,
 		FOptimus_InterfaceBindingMap& OutOutputDataBindings,
@@ -53,7 +52,9 @@ public:
 	FOptimusExecutionDomain GetExecutionDomain() const override PURE_VIRTUAL(UOptimusNode_ComputeKernelBase::GetExecutionDomain, return {}; );
 	const UOptimusNodePin* GetPrimaryGroupPin() const override PURE_VIRTUAL(UOptimusNode_ComputeKernelBase::GetPrimaryGroupPin, return {}; );  
 	UComputeDataInterface* MakeKernelDataInterface(UObject* InOuter) const override PURE_VIRTUAL(UOptimusNode_ComputeKernelBase::MakeKernelDataInterface, return {}; );
-	bool GetPinSupportAtomic(const UOptimusNodePin* InPin) const override PURE_VIRTUAL(UOptimusNode_ComputeKernelBase::GetPinSupportAtomic, return {}; );
+	bool DoesOutputPinSupportAtomic(const UOptimusNodePin* InPin) const override PURE_VIRTUAL(UOptimusNode_ComputeKernelBase::DoesOutputPinSupportAtomic, return false; );
+	bool DoesOutputPinSupportRead(const UOptimusNodePin* InPin) const override PURE_VIRTUAL(UOptimusNode_ComputeKernelBase::DoesOutputPinSupportRead, return false; );
+	bool HasMutableInput() const override;
 	
 	// -- UOptimusNode overrides
 	TOptional<FText> ValidateForCompile() const override;
@@ -102,9 +103,8 @@ protected:
 		FIntVector InGroupSize
 		);
 
-	static TSet<UOptimusComponentSourceBinding*> GetGroupComponentSourceBindings(const UOptimusNodePin* InGroupPin);
-
 	static FString GetAtomicWriteFunctionName(EOptimusBufferWriteType InWriteType, const FString& InBindingName);
+	static FString GetReadFunctionName(const FString& InBindingName);
 	
 private:
 	TOptional<FText> ProcessInputPinForComputeKernel(
@@ -115,7 +115,6 @@ private:
 		const FOptimus_PinToDataInterfaceMap& InLinkDataInterfaceMap,
 		const TArray<const UOptimusNode*>& InValueNodes,
 		const UComputeDataInterface* InGraphDataInterface,
-		const UOptimusComponentSourceBinding* InGraphDataComponentBinding,
 		UOptimusKernelSource* InKernelSource,
 		TArray<FString>& OutGeneratedFunctions,
 		FOptimus_InterfaceBindingMap& OutInputDataBindings,
@@ -129,6 +128,7 @@ private:
 		const FOptimus_PinToDataInterfaceMap& InLinkDataInterfaceMap,
 		UOptimusKernelSource* InKernelSource,
 		TArray<FString>& OutGeneratedFunctions,
+		FOptimus_InterfaceBindingMap& OutInputDataBindings,
 		FOptimus_InterfaceBindingMap& OutOutputDataBindings,
 		FOptimusKernelConstantContainer& OutKernelConstantContainer
 		) const;
diff --git a/Engine/Plugins/Animation/DeformerGraph/Source/OptimusCore/Private/Nodes/OptimusNode_ComputeKernelFunction.h b/Engine/Plugins/Animation/DeformerGraph/Source/OptimusCore/Private/Nodes/OptimusNode_ComputeKernelFunction.h
--- a/Engine/Plugins/Animation/DeformerGraph/Source/OptimusCore/Private/Nodes/OptimusNode_ComputeKernelFunction.h
+++ b/Engine/Plugins/Animation/DeformerGraph/Source/OptimusCore/Private/Nodes/OptimusNode_ComputeKernelFunction.h
@@ -79,8 +79,9 @@ public:
 	FOptimusExecutionDomain GetExecutionDomain() const override;
 	const UOptimusNodePin* GetPrimaryGroupPin() const override { return {}; }
 	UComputeDataInterface* MakeKernelDataInterface(UObject* InOuter) const override { return nullptr; };
-	bool GetPinSupportAtomic(const UOptimusNodePin* InPin) const override {return false;};
-
+	bool DoesOutputPinSupportAtomic(const UOptimusNodePin* InPin) const override {return false;};
+	bool DoesOutputPinSupportRead(const UOptimusNodePin* InPin) const override {return false;};
+	bool HasMutableInput() const override {return false;};
 private:
 	UOptimusNode_ComputeKernelFunctionGeneratorClass *GetGeneratorClass() const;
 };
diff --git a/Engine/Plugins/Animation/DeformerGraph/Source/OptimusCore/Private/Nodes/OptimusNode_ConstantValue.h b/Engine/Plugins/Animation/DeformerGraph/Source/OptimusCore/Private/Nodes/OptimusNode_ConstantValue.h
--- a/Engine/Plugins/Animation/DeformerGraph/Source/OptimusCore/Private/Nodes/OptimusNode_ConstantValue.h
+++ b/Engine/Plugins/Animation/DeformerGraph/Source/OptimusCore/Private/Nodes/OptimusNode_ConstantValue.h
@@ -3,6 +3,7 @@
 #pragma once
 
 #include "IOptimusGeneratedClassDefiner.h"
+#include "IOptimusPinMutabilityDefiner.h"
 #include "IOptimusValueProvider.h"
 #include "OptimusNode.h"
 #include "OptimusDataType.h"
@@ -43,7 +44,8 @@ UCLASS(Hidden)
 class UOptimusNode_ConstantValue :
 	public UOptimusNode,
 	public IOptimusValueProvider,
-	public IOptimusGeneratedClassDefiner
+	public IOptimusGeneratedClassDefiner,
+	public IOptimusPinMutabilityDefiner
 {
 	GENERATED_BODY()
 
@@ -74,6 +76,9 @@ public:
 		UPackage* InPackage,
 		const TCHAR* InCreationString
 		) const override;
+
+	// IOptimusPinMutabilityDefiner implementation
+	bool IsOutputPinMutable(const UOptimusNodePin* InPin) const override { return false; };
 	
 	
 protected:
diff --git a/Engine/Plugins/Animation/DeformerGraph/Source/OptimusCore/Private/Nodes/OptimusNode_CustomComputeKernel.cpp b/Engine/Plugins/Animation/DeformerGraph/Source/OptimusCore/Private/Nodes/OptimusNode_CustomComputeKernel.cpp
--- a/Engine/Plugins/Animation/DeformerGraph/Source/OptimusCore/Private/Nodes/OptimusNode_CustomComputeKernel.cpp
+++ b/Engine/Plugins/Animation/DeformerGraph/Source/OptimusCore/Private/Nodes/OptimusNode_CustomComputeKernel.cpp
@@ -122,7 +122,7 @@ UComputeDataInterface* UOptimusNode_CustomComputeKernel::MakeKernelDataInterface
 	return KernelDataInterface;
 }
 
-bool UOptimusNode_CustomComputeKernel::GetPinSupportAtomic(const UOptimusNodePin* InPin) const
+bool UOptimusNode_CustomComputeKernel::DoesOutputPinSupportAtomic(const UOptimusNodePin* InPin) const
 {
 	if (ensure(InPin->GetDirection() == EOptimusNodePinDirection::Output))
 	{
@@ -139,6 +139,23 @@ bool UOptimusNode_CustomComputeKernel::GetPinSupportAtomic(const UOptimusNodePin
 	return false;
 }
 
+bool UOptimusNode_CustomComputeKernel::DoesOutputPinSupportRead(const UOptimusNodePin* InPin) const
+{
+	if (ensure(InPin->GetDirection() == EOptimusNodePinDirection::Output))
+	{
+		if (const FOptimusParameterBinding* Binding = OutputBindingArray.FindByPredicate(
+			[InPin](const FOptimusParameterBinding& InBinding)
+				{
+					return InBinding.Name == InPin->GetFName(); 
+				}))
+		{
+			return Binding->bSupportRead;
+		}
+	}
+
+	return false;
+}
+
 #if WITH_EDITOR
 
 FString UOptimusNode_CustomComputeKernel::GetNameForShaderTextEditor() const
@@ -203,7 +220,7 @@ FString UOptimusNode_CustomComputeKernel::GetBindingDeclaration(
 	return FString();
 }
 
-bool UOptimusNode_CustomComputeKernel::GetBindingAtomicSupportCheckBoxVisibility(FName BindingName) const
+bool UOptimusNode_CustomComputeKernel::GetBindingSupportAtomicCheckBoxVisibility(FName BindingName) const
 {
 	auto ParameterBindingPredicate = [BindingName](const FOptimusParameterBinding& InBinding)
 	{
@@ -227,6 +244,26 @@ bool UOptimusNode_CustomComputeKernel::GetBindingAtomicSupportCheckBoxVisibility
 	return false;
 }
 
+bool UOptimusNode_CustomComputeKernel::GetBindingSupportReadCheckBoxVisibility(FName BindingName) const
+{
+	auto ParameterBindingPredicate = [BindingName](const FOptimusParameterBinding& InBinding)
+	{
+		if (InBinding.Name == BindingName)
+		{
+			return true;	
+		}
+			
+		return false;
+	};
+
+	if (const FOptimusParameterBinding* Binding = OutputBindingArray.FindByPredicate(ParameterBindingPredicate))
+	{
+		return true;
+	}
+
+	return false;
+}
+
 
 TArray<IOptimusNodeAdderPinProvider::FAdderPinAction> UOptimusNode_CustomComputeKernel::GetAvailableAdderPinActions(
 	const UOptimusNodePin* InSourcePin,
@@ -283,7 +320,7 @@ TArray<IOptimusNodeAdderPinProvider::FAdderPinAction> UOptimusNode_CustomCompute
 		if (InSourcePin->GetDirection() != InNewPinDirection)
 		{
 			const UOptimusNodePin* PrimaryGroupPin = GetPrimaryGroupPin();
-			TSet<UOptimusComponentSourceBinding*> PrimaryGroupComponentBindings = GetGroupComponentSourceBindings(PrimaryGroupPin);
+			TSet<UOptimusComponentSourceBinding*> PrimaryGroupComponentBindings = PrimaryGroupPin->GetComponentSourceBindingsRecursively();
 			TSet<UOptimusComponentSourceBinding*> MatchingBindings = PrimaryGroupComponentBindings.Intersect(SourceComponentBindings);
 		
 			if (MatchingBindings.Num() > 0 ||
@@ -333,7 +370,7 @@ TArray<IOptimusNodeAdderPinProvider::FAdderPinAction> UOptimusNode_CustomCompute
 				Action.DisplayName = GroupPin->GetFName();
 				Action.Key = GroupPin->GetFName();
 					
-				TSet<UOptimusComponentSourceBinding*> ComponentBindingsForGroup = GetGroupComponentSourceBindings(GroupPin);
+				TSet<UOptimusComponentSourceBinding*> ComponentBindingsForGroup = GroupPin->GetComponentSourceBindingsRecursively();
 				TSet<UOptimusComponentSourceBinding*> MatchingBindings = ComponentBindingsForGroup.Intersect(SourceComponentBindings);
 	
 				if (MatchingBindings.Num() > 0 ||
@@ -607,7 +644,7 @@ TArray<FName> UOptimusNode_CustomComputeKernel::GetExecutionDomains() const
 {
 	// Find all component sources for the primary pins. If we end up with any other number
 	// than one, then something's gone wrong and we can't determine the execution domains.
-	TSet<UOptimusComponentSourceBinding*> PrimaryBindings = GetGroupComponentSourceBindings(GetPrimaryGroupPin());
+	TSet<UOptimusComponentSourceBinding*> PrimaryBindings = GetPrimaryGroupPin()->GetComponentSourceBindingsRecursively();
 
 	if (PrimaryBindings.Num() == 1)
 	{
@@ -901,6 +938,11 @@ void UOptimusNode_CustomComputeKernel::PropertyValueChanged(
 		UpdatePreamble();
 		return;
 	}
+	else if (InPropertyChangedEvent.GetPropertyName() == GET_MEMBER_NAME_STRING_CHECKED(FOptimusParameterBinding, bSupportRead))
+	{
+		UpdatePreamble();
+		return;
+	}
 	else if ( InPropertyChangedEvent.GetPropertyName() == GET_MEMBER_NAME_STRING_CHECKED(FOptimusParameterBindingArray, InnerArray))
 	{
 		// Pasting into a specific binding
@@ -1911,6 +1953,12 @@ FString UOptimusNode_CustomComputeKernel::GetDeclarationForBinding(const FOptimu
 			}
 		}
 
+		if (Binding.bSupportRead)
+		{
+			Declarations.Add(FString::Printf(TEXT("%s Read%s(%s);"), 
+				*GetShaderValueTypeFriendlyName(Binding.DataType), *Binding.Name.ToString(), *FString::Join(Indexes, TEXT(", "))));
+		}
+
 		return FString::Join(Declarations, TEXT("\n"));
 	}
 
diff --git a/Engine/Plugins/Animation/DeformerGraph/Source/OptimusCore/Private/Nodes/OptimusNode_CustomComputeKernel.h b/Engine/Plugins/Animation/DeformerGraph/Source/OptimusCore/Private/Nodes/OptimusNode_CustomComputeKernel.h
--- a/Engine/Plugins/Animation/DeformerGraph/Source/OptimusCore/Private/Nodes/OptimusNode_CustomComputeKernel.h
+++ b/Engine/Plugins/Animation/DeformerGraph/Source/OptimusCore/Private/Nodes/OptimusNode_CustomComputeKernel.h
@@ -62,7 +62,8 @@ public:
 	FOptimusExecutionDomain GetExecutionDomain() const override;
 	const UOptimusNodePin* GetPrimaryGroupPin() const override;
 	UComputeDataInterface* MakeKernelDataInterface(UObject* InOuter) const override;
-	bool GetPinSupportAtomic(const UOptimusNodePin* InPin) const override;
+	bool DoesOutputPinSupportAtomic(const UOptimusNodePin* InPin) const override;
+	bool DoesOutputPinSupportRead(const UOptimusNodePin* InPin) const override;
 
 #if WITH_EDITOR
 	// IOptimusShaderTextProvider overrides
@@ -75,8 +76,9 @@ public:
 	
 	// IOptimusParameterBindingProvider
 	FString GetBindingDeclaration(FName BindingName) const override;
-	bool GetBindingAtomicSupportCheckBoxVisibility(FName BindingName) const override;
-
+	bool GetBindingSupportAtomicCheckBoxVisibility(FName BindingName) const override;
+	bool GetBindingSupportReadCheckBoxVisibility(FName BindingName) const override;
+	
 	// IOptimusNodeAdderPinProvider
 	TArray<FAdderPinAction> GetAvailableAdderPinActions(
 		const UOptimusNodePin* InSourcePin,
diff --git a/Engine/Plugins/Animation/DeformerGraph/Source/OptimusCore/Private/Nodes/OptimusNode_DataInterface.cpp b/Engine/Plugins/Animation/DeformerGraph/Source/OptimusCore/Private/Nodes/OptimusNode_DataInterface.cpp
--- a/Engine/Plugins/Animation/DeformerGraph/Source/OptimusCore/Private/Nodes/OptimusNode_DataInterface.cpp
+++ b/Engine/Plugins/Animation/DeformerGraph/Source/OptimusCore/Private/Nodes/OptimusNode_DataInterface.cpp
@@ -190,6 +190,27 @@ UOptimusComponentSourceBinding* UOptimusNode_DataInterface::GetComponentBinding(
 	return nullptr;
 }
 
+bool UOptimusNode_DataInterface::IsOutputPinMutable(const UOptimusNodePin* InPin) const
+{
+	const TArray<FOptimusCDIPinDefinition> PinDefinitions = DataInterfaceData->GetPinDefinitions();
+
+	int32 PinDefinitionIndex = INDEX_NONE;
+	for (int32 Index = 0 ; Index < PinDefinitions.Num(); ++Index)
+	{
+		if (PinDefinitions[Index].PinName == InPin->GetUniqueName())
+		{
+			PinDefinitionIndex = Index;
+			break;
+		}
+	}
+	if (!ensure(PinDefinitionIndex != INDEX_NONE))
+	{
+		return true;
+	}
+
+	return PinDefinitions[PinDefinitionIndex].bMutable;
+}
+
 
 void UOptimusNode_DataInterface::PostLoad() 
 {
diff --git a/Engine/Plugins/Animation/DeformerGraph/Source/OptimusCore/Private/Nodes/OptimusNode_DataInterface.h b/Engine/Plugins/Animation/DeformerGraph/Source/OptimusCore/Private/Nodes/OptimusNode_DataInterface.h
--- a/Engine/Plugins/Animation/DeformerGraph/Source/OptimusCore/Private/Nodes/OptimusNode_DataInterface.h
+++ b/Engine/Plugins/Animation/DeformerGraph/Source/OptimusCore/Private/Nodes/OptimusNode_DataInterface.h
@@ -3,6 +3,7 @@
 #pragma once
 
 #include "IOptimusDataInterfaceProvider.h"
+#include "IOptimusPinMutabilityDefiner.h"
 #include "OptimusComputeDataInterface.h"
 #include "OptimusComponentSource.h"
 
@@ -17,7 +18,8 @@
 UCLASS(Hidden)
 class UOptimusNode_DataInterface :
 	public UOptimusNode,
-	public IOptimusDataInterfaceProvider
+	public IOptimusDataInterfaceProvider,
+	public IOptimusPinMutabilityDefiner
 {
 	GENERATED_BODY()
 
@@ -40,6 +42,9 @@ public:
 	UOptimusComputeDataInterface *GetDataInterface(UObject *InOuter) const override;
 	int32 GetDataFunctionIndexFromPin(const UOptimusNodePin* InPin) const override;
 	UOptimusComponentSourceBinding* GetComponentBinding() const override;
+
+	// -- IOptimusPinMutabilityDefiner
+	bool IsOutputPinMutable(const UOptimusNodePin* InPin) const override;
 	
 protected:
 	// -- UOptimusNode overrides
diff --git a/Engine/Plugins/Animation/DeformerGraph/Source/OptimusCore/Private/Nodes/OptimusNode_GetResource.cpp b/Engine/Plugins/Animation/DeformerGraph/Source/OptimusCore/Private/Nodes/OptimusNode_GetResource.cpp
--- a/Engine/Plugins/Animation/DeformerGraph/Source/OptimusCore/Private/Nodes/OptimusNode_GetResource.cpp
+++ b/Engine/Plugins/Animation/DeformerGraph/Source/OptimusCore/Private/Nodes/OptimusNode_GetResource.cpp
@@ -20,7 +20,7 @@ int32 UOptimusNode_GetResource::GetDataFunctionIndexFromPin(const UOptimusNodePi
 		return INDEX_NONE;
 	}
 
-	return UOptimusRawBufferDataInterface::GetReadValueInputIndex();
+	return UOptimusRawBufferDataInterface::GetReadValueInputIndex(EOptimusBufferReadType::Default);
 }
 
 void UOptimusNode_GetResource::ConstructNode()
diff --git a/Engine/Plugins/Animation/DeformerGraph/Source/OptimusCore/Private/Nodes/OptimusNode_GetVariable.h b/Engine/Plugins/Animation/DeformerGraph/Source/OptimusCore/Private/Nodes/OptimusNode_GetVariable.h
--- a/Engine/Plugins/Animation/DeformerGraph/Source/OptimusCore/Private/Nodes/OptimusNode_GetVariable.h
+++ b/Engine/Plugins/Animation/DeformerGraph/Source/OptimusCore/Private/Nodes/OptimusNode_GetVariable.h
@@ -2,6 +2,7 @@
 
 #pragma once
 
+#include "IOptimusPinMutabilityDefiner.h"
 #include "IOptimusValueProvider.h"
 #include "OptimusNode.h"
 
@@ -32,7 +33,8 @@ struct FOptimusNode_GetVariable_DuplicationInfo
 UCLASS(Hidden)
 class UOptimusNode_GetVariable : 
 	public UOptimusNode,
-	public IOptimusValueProvider
+	public IOptimusValueProvider,
+	public IOptimusPinMutabilityDefiner
 {
 	GENERATED_BODY()
 
@@ -53,6 +55,9 @@ public:
 	FString GetValueName() const override;
 	FOptimusDataTypeRef GetValueType() const override;
 	FShaderValueType::FValue GetShaderValue() const override;
+
+	//IOptimusPinMutabilityDefiner overrides 
+	bool IsOutputPinMutable(const UOptimusNodePin* InPin) const override { return true; };
 	
 protected:
 	void ConstructNode() override;
diff --git a/Engine/Plugins/Animation/DeformerGraph/Source/OptimusCore/Private/Nodes/OptimusNode_Resource.cpp b/Engine/Plugins/Animation/DeformerGraph/Source/OptimusCore/Private/Nodes/OptimusNode_Resource.cpp
--- a/Engine/Plugins/Animation/DeformerGraph/Source/OptimusCore/Private/Nodes/OptimusNode_Resource.cpp
+++ b/Engine/Plugins/Animation/DeformerGraph/Source/OptimusCore/Private/Nodes/OptimusNode_Resource.cpp
@@ -21,7 +21,7 @@ int32 UOptimusNode_Resource::GetDataFunctionIndexFromPin(const UOptimusNodePin*
 	}
 
 	return InPin->GetDirection() == EOptimusNodePinDirection::Output 
-		? UOptimusRawBufferDataInterface::GetReadValueInputIndex() 
+		? UOptimusRawBufferDataInterface::GetReadValueInputIndex(EOptimusBufferReadType::Default) 
 		: UOptimusRawBufferDataInterface::GetWriteValueOutputIndex(EOptimusBufferWriteType::Write);
 }
 
diff --git a/Engine/Plugins/Animation/DeformerGraph/Source/OptimusCore/Private/Nodes/OptimusNode_ResourceAccessorBase.h b/Engine/Plugins/Animation/DeformerGraph/Source/OptimusCore/Private/Nodes/OptimusNode_ResourceAccessorBase.h
--- a/Engine/Plugins/Animation/DeformerGraph/Source/OptimusCore/Private/Nodes/OptimusNode_ResourceAccessorBase.h
+++ b/Engine/Plugins/Animation/DeformerGraph/Source/OptimusCore/Private/Nodes/OptimusNode_ResourceAccessorBase.h
@@ -5,6 +5,7 @@
 #include "DataInterfaces/OptimusDataInterfaceRawBuffer.h"
 #include "IOptimusComponentBindingProvider.h"
 #include "IOptimusDataInterfaceProvider.h"
+#include "IOptimusPinMutabilityDefiner.h"
 #include "OptimusNode.h"
 
 #include "OptimusNode_ResourceAccessorBase.generated.h"
@@ -34,7 +35,8 @@ UCLASS(Abstract)
 class UOptimusNode_ResourceAccessorBase : 
 	public UOptimusNode,
 	public IOptimusDataInterfaceProvider,
-	public IOptimusComponentBindingProvider
+	public IOptimusComponentBindingProvider,
+	public IOptimusPinMutabilityDefiner
 {
 	GENERATED_BODY()
 
@@ -63,6 +65,9 @@ public:
 	PRAGMA_DISABLE_DEPRECATION_WARNINGS
 	EOptimusBufferWriteType GetDeprecatedBufferWriteType() const { return WriteType_DEPRECATED; }
 	PRAGMA_ENABLE_DEPRECATION_WARNINGS
+
+	//IOptimusPinMutabilityDefiner overrides 
+	bool IsOutputPinMutable(const UOptimusNodePin* InPin) const override { return true; };
 	
 protected:
 	void PreDuplicateRequirementActions(const UOptimusNodeGraph* InTargetGraph, FOptimusCompoundAction* InCompoundAction) override;
diff --git a/Engine/Plugins/Animation/DeformerGraph/Source/OptimusCore/Private/OptimusDeformer.cpp b/Engine/Plugins/Animation/DeformerGraph/Source/OptimusCore/Private/OptimusDeformer.cpp
--- a/Engine/Plugins/Animation/DeformerGraph/Source/OptimusCore/Private/OptimusDeformer.cpp
+++ b/Engine/Plugins/Animation/DeformerGraph/Source/OptimusCore/Private/OptimusDeformer.cpp
@@ -1442,14 +1442,8 @@ bool UOptimusDeformer::Compile()
 
 	for (const UOptimusNodeGraph* Graph: Graphs)
 	{
-		if (UOptimusComputeGraph* ComputeGraph = CompileNodeGraphToComputeGraph(Graph, ErrorReporter))
-		{
-			FOptimusComputeGraphInfo Info;
-			Info.GraphType = Graph->GraphType;
-			Info.GraphName = Graph->GetFName();
-			Info.ComputeGraph = ComputeGraph;
-			ComputeGraphs.Add(Info);
-		}
+		TArray<FOptimusComputeGraphInfo> ComputeGraphInfos = CompileNodeGraphToComputeGraphs(Graph, ErrorReporter);
+		ComputeGraphs.Append(ComputeGraphInfos);
 	}
 	
 	CompileEndDelegate.Broadcast(this);
@@ -1505,7 +1499,7 @@ TArray<UOptimusNode*> UOptimusDeformer::GetAllNodesOfClass(UClass* InNodeClass)
 }
 
 
-UOptimusComputeGraph* UOptimusDeformer::CompileNodeGraphToComputeGraph(
+TArray<FOptimusComputeGraphInfo> UOptimusDeformer::CompileNodeGraphToComputeGraphs(
 	const UOptimusNodeGraph* InNodeGraph,
 	TFunction<void(EOptimusDiagnosticLevel, FText, const UObject*)> InErrorReporter
 	)
@@ -1527,7 +1521,7 @@ UOptimusComputeGraph* UOptimusDeformer::CompileNodeGraphToComputeGraph(
 	// No nodes in the graph, nothing to do.
 	if (InNodeGraph->GetAllNodes().IsEmpty())
 	{
-		return nullptr;
+		return {};
 	}
 
 	// Clear the error state of all nodes.
@@ -1576,7 +1570,7 @@ UOptimusComputeGraph* UOptimusDeformer::CompileNodeGraphToComputeGraph(
 	if (TerminalNodes.IsEmpty())
 	{
 		AddDiagnostic(EOptimusDiagnosticLevel::Error, LOCTEXT("NoOutputDataInterfaceFound", "No connected output data interface nodes found. Compilation aborted."));
-		return nullptr;
+		return {};
 	}
 
 	TArray<FOptimusRoutedConstNode> ConnectedNodes;
@@ -1599,11 +1593,87 @@ UOptimusComputeGraph* UOptimusDeformer::CompileNodeGraphToComputeGraph(
 	}
 	if (bValidationFailed)
 	{
-		return nullptr;
+		return {};
 	}
 
-	const FName GraphName = MakeUniqueObjectName(this, UOptimusComputeGraph::StaticClass(), InNodeGraph->GetFName());
-	UOptimusComputeGraph* ComputeGraph = NewObject<UOptimusComputeGraph>(this, GraphName);
+	TArray<FOptimusComputeGraphInfo> GraphInfos;
+	TMap<const IOptimusComputeKernelProvider*, int> KernelComputeGraphIndexMap;
+	
+	if (InNodeGraph->GraphType != EOptimusNodeGraphType::Update)
+	{
+		FOptimusComputeGraphInfo GraphInfo;
+		GraphInfo.GraphName = MakeUniqueObjectName(this, UOptimusComputeGraph::StaticClass(), InNodeGraph->GetFName());;
+		GraphInfo.GraphType = InNodeGraph->GetGraphType();
+		GraphInfo.ComputeGraph = NewObject<UOptimusComputeGraph>(this, GraphInfo.GraphName);
+		GraphInfos.Add(GraphInfo);
+
+		for (FOptimusRoutedConstNode ConnectedNode: ConnectedNodes)
+		{
+			if (const IOptimusComputeKernelProvider *KernelProvider = Cast<const IOptimusComputeKernelProvider>(ConnectedNode.Node))
+			{
+				KernelComputeGraphIndexMap.Add(KernelProvider) = 0;
+			}
+		}
+	}
+	else
+	{
+		bool bHasRunOnceKernels = false;
+		bool bHasRunAlwaysKernels = false;
+	
+		for (FOptimusRoutedConstNode ConnectedNode: ConnectedNodes)
+		{
+			if (const IOptimusComputeKernelProvider *KernelProvider = Cast<const IOptimusComputeKernelProvider>(ConnectedNode.Node))
+			{
+				if (!KernelProvider->HasMutableInput())
+				{
+					bHasRunOnceKernels = true;
+				}
+				else
+				{
+					bHasRunAlwaysKernels = true;
+				}
+			}
+		}
+
+		// Insert a setup graph if applicable
+		if (bHasRunOnceKernels)
+		{
+			FName GraphName = InNodeGraph->GetFName();
+			GraphName = *(GraphName.ToString() + TEXT("_Setup"));
+			GraphName = MakeUniqueObjectName(this, UOptimusComputeGraph::StaticClass(), GraphName);
+		
+			FOptimusComputeGraphInfo GraphInfo;
+			GraphInfo.GraphName = GraphName;
+			GraphInfo.GraphType = EOptimusNodeGraphType::Setup;
+			GraphInfo.ComputeGraph = NewObject<UOptimusComputeGraph>(this, GraphInfo.GraphName);
+			GraphInfos.Add(GraphInfo);
+		}
+
+		if (bHasRunAlwaysKernels)
+		{
+			FOptimusComputeGraphInfo GraphInfo;
+			GraphInfo.GraphName = MakeUniqueObjectName(this, UOptimusComputeGraph::StaticClass(), InNodeGraph->GetFName());;
+			GraphInfo.GraphType = InNodeGraph->GetGraphType();
+			GraphInfo.ComputeGraph = NewObject<UOptimusComputeGraph>(this, GraphInfo.GraphName);
+			GraphInfos.Add(GraphInfo);
+		}
+
+		for (FOptimusRoutedConstNode ConnectedNode: ConnectedNodes)
+		{
+			if (const IOptimusComputeKernelProvider *KernelProvider = Cast<const IOptimusComputeKernelProvider>(ConnectedNode.Node))
+			{
+				int32 ComputeGraphIndex = 0;
+				if (GraphInfos.Num() == 2)
+				{
+					ComputeGraphIndex = KernelProvider->HasMutableInput() ? 1 : 0;
+				}
+				KernelComputeGraphIndexMap.Add(KernelProvider) = ComputeGraphIndex;
+			}
+		}
+	}
+
+	check(GraphInfos.Num() <= 2);
+
 	
 	// Find all data interface nodes and create their data interfaces.
 	FOptimus_NodeToDataInterfaceMap NodeDataInterfaceMap;
@@ -1617,7 +1687,11 @@ UOptimusComputeGraph* UOptimusDeformer::CompileNodeGraphToComputeGraph(
 	FOptimus_KernelNodeToKernelDataInterfaceMap KernelDataInterfaceMap;
 
 	// Find all value nodes (constant and variable) 
-	TArray<const UOptimusNode *> ValueNodes; 
+	TArray<const UOptimusNode *> ValueNodes;
+
+	TMap<const UComputeDataInterface*, int32> DataInterfaceToBindingIndexMap;
+
+	
 
 	for (FOptimusRoutedConstNode ConnectedNode: ConnectedNodes)
 	{
@@ -1627,16 +1701,31 @@ UOptimusComputeGraph* UOptimusDeformer::CompileNodeGraphToComputeGraph(
 			if (!DataInterface)
 			{
 				AddDiagnostic(EOptimusDiagnosticLevel::Error, LOCTEXT("NoDataInterfaceOnProvider", "No data interface object returned from node. Compilation aborted."));
-				return nullptr;
+				return {};
 			}
 
 			NodeDataInterfaceMap.Add(ConnectedNode.Node, DataInterface);
+			DataInterfaceToBindingIndexMap.Add(DataInterface) = DataInterfaceNode->GetComponentBinding()->GetIndex();
 		}
 		else if (const IOptimusComputeKernelProvider* KernelProvider = Cast<const IOptimusComputeKernelProvider>(ConnectedNode.Node))
 		{
 			UComputeDataInterface* KernelDataInterface = KernelProvider->MakeKernelDataInterface(this);
+			TSet<UOptimusComponentSourceBinding*> KernelPrimaryBindings = KernelProvider->GetPrimaryGroupPin()->GetComponentSourceBindingsRecursively();
+
+			if (!ensure(KernelPrimaryBindings.Num() == 1))
+			{
+				AddDiagnostic(EOptimusDiagnosticLevel::Error,
+					FText::Format(LOCTEXT("InvalidComponentBindingForKernel", "Missing or multiple component bindings found in primary group of a kernel ({0}). Compilation aborted."),
+					ConnectedNode.Node->GetDisplayName()),
+					ConnectedNode.Node);
+				return {};
+			}
+			
+			UOptimusComponentSourceBinding* KernelPrimaryBinding = *KernelPrimaryBindings.CreateConstIterator();
+			int32 PrimaryBindingIndex = KernelPrimaryBinding->GetIndex();
 
 			KernelDataInterfaceMap.Add(ConnectedNode.Node, KernelDataInterface);
+			DataInterfaceToBindingIndexMap.Add(KernelDataInterface) = PrimaryBindingIndex;	
 			
 			for (const UOptimusNodePin* Pin: ConnectedNode.Node->GetPins())
 			{
@@ -1647,32 +1736,43 @@ UOptimusComputeGraph* UOptimusDeformer::CompileNodeGraphToComputeGraph(
 					for (const FOptimusRoutedNodePin& ConnectedPin: Pin->GetConnectedPinsWithRouting(ConnectedNode.TraversalContext))
 					{
 						// Make sure it connects to another kernel node.
-						if (Cast<const IOptimusComputeKernelProvider>(ConnectedPin.NodePin->GetOwningNode()) != nullptr &&
-							ensure(Pin->GetDataType().IsValid()))
+						if (const IOptimusComputeKernelProvider* ConnectedKernel = Cast<const IOptimusComputeKernelProvider>(ConnectedPin.NodePin->GetOwningNode());
+							ConnectedKernel && ensure(Pin->GetDataType().IsValid()))
 						{
-							UOptimusTransientBufferDataInterface* TransientBufferDI =
-								NewObject<UOptimusTransientBufferDataInterface>(this);
-
-							TSet<UOptimusComponentSourceBinding*> ComponentSourceBindings = Pin->GetComponentSourceBindings();
-							if (ComponentSourceBindings.Num() != 1)
+							if (KernelComputeGraphIndexMap[KernelProvider] == KernelComputeGraphIndexMap[ConnectedKernel])
 							{
-								AddDiagnostic(EOptimusDiagnosticLevel::Error,
-									FText::Format(LOCTEXT("InvalidComponentBindingOnKernelPin", "Missing or multiple component bindings on kernel-to-kernel pin ({0}). Compilation aborted."),
-										FText::FromName(Pin->GetUniqueName())),
-									ConnectedNode.Node);
-								return nullptr;
+								UOptimusTransientBufferDataInterface* TransientBufferDI =
+									NewObject<UOptimusTransientBufferDataInterface>(this);
+
+								TransientBufferDI->ValueType = Pin->GetDataType()->ShaderValueType;
+								TransientBufferDI->DataDomain = Pin->GetDataDomain();
+								TransientBufferDI->ComponentSourceBinding = KernelPrimaryBinding;
+								TransientBufferDI->DomainConstantIdentifier = {ConnectedNode.Node, NAME_None, Pin->GetFName()};
+								if (KernelProvider->DoesOutputPinSupportAtomic(Pin))
+								{
+									TransientBufferDI->bZeroInitForAtomicWrites = true;
+								}
+							
+								LinkDataInterfaceMap.Add(Pin, TransientBufferDI);
+								DataInterfaceToBindingIndexMap.Add(TransientBufferDI) = PrimaryBindingIndex;	
 							}
-
-							TransientBufferDI->ValueType = Pin->GetDataType()->ShaderValueType;
-							TransientBufferDI->DataDomain = Pin->GetDataDomain();
-							TransientBufferDI->ComponentSourceBinding = *ComponentSourceBindings.CreateConstIterator();
-							TransientBufferDI->DomainConstantIdentifier = {ConnectedNode.Node, NAME_None, Pin->GetFName()};
-							if (KernelProvider->GetPinSupportAtomic(Pin))
+							else
 							{
-								TransientBufferDI->bZeroInitForAtomicWrites = true;
-							}
+								UOptimusImplicitPersistentBufferDataInterface* ImplicitPersistentBufferDI =
+									NewObject<UOptimusImplicitPersistentBufferDataInterface>(this);
+
+								ImplicitPersistentBufferDI->ValueType = Pin->GetDataType()->ShaderValueType;
+								ImplicitPersistentBufferDI->DataDomain = Pin->GetDataDomain();
+								ImplicitPersistentBufferDI->ComponentSourceBinding = KernelPrimaryBinding;
+								ImplicitPersistentBufferDI->DomainConstantIdentifier = {ConnectedNode.Node, NAME_None, Pin->GetFName()};
+								if (KernelProvider->DoesOutputPinSupportAtomic(Pin))
+								{
+									ImplicitPersistentBufferDI->bZeroInitForAtomicWrites = true;
+								}
 							
-							LinkDataInterfaceMap.Add(Pin, TransientBufferDI);
+								LinkDataInterfaceMap.Add(Pin, ImplicitPersistentBufferDI);
+								DataInterfaceToBindingIndexMap.Add(ImplicitPersistentBufferDI) = PrimaryBindingIndex;	
+							}
 						}
 					}
 				}
@@ -1705,230 +1805,192 @@ UOptimusComputeGraph* UOptimusDeformer::CompileNodeGraphToComputeGraph(
 	}
 	GraphDataInterface->Init(ValueNodeDescriptions);
 
-	// Loop through all kernels, create a kernel source, and create a compute kernel for it.
-	struct FKernelWithDataBindings
-	{
-		int32 KernelNodeIndex;
-		UComputeKernel *Kernel;
-		FOptimus_InterfaceBindingMap InputDataBindings;
-		FOptimus_InterfaceBindingMap OutputDataBindings;
-	};
-
 	// The component binding for the graph data is the primary binding on the deformer.
 	UOptimusComponentSourceBinding* GraphDataComponentBinding = Bindings->Bindings[0];
-	
-	TArray<FKernelWithDataBindings> BoundKernels;
-	for (FOptimusRoutedConstNode ConnectedNode: ConnectedNodes)
+	DataInterfaceToBindingIndexMap.Add(GraphDataInterface) = GraphDataComponentBinding->GetIndex();
+
+	for (int32 ComputeGraphIndex = 0 ; ComputeGraphIndex < GraphInfos.Num(); ComputeGraphIndex++)
 	{
-		if (const IOptimusComputeKernelProvider *KernelProvider = Cast<const IOptimusComputeKernelProvider>(ConnectedNode.Node))
+		UOptimusComputeGraph* ComputeGraph = GraphInfos[ComputeGraphIndex].ComputeGraph;
+
+		// Create the binding objects.
+		for (const UOptimusComponentSourceBinding* Binding: Bindings->Bindings)
 		{
-			FKernelWithDataBindings BoundKernel;
+			ComputeGraph->Bindings.Add(Binding->GetComponentSource()->GetComponentClass());
+		}
 
-			BoundKernel.KernelNodeIndex = InNodeGraph->Nodes.IndexOfByKey(ConnectedNode.Node);
-			BoundKernel.Kernel = NewObject<UComputeKernel>(this);
+		// Now that we've collected all the pieces, time to line them up.
+		ComputeGraph->DataInterfaces.Add(GraphDataInterface);
+		ComputeGraph->DataInterfaceToBinding.Add(0);		// Graph data interface always uses the primary binding.
 
-			FOptimusKernelConstantContainer& KernelConstantContainer = ConstantContainer.AddContainerForKernel();
-			UComputeDataInterface* KernelDataInterface = KernelDataInterfaceMap[ConnectedNode.Node];
-			
-			FOptimus_ComputeKernelResult KernelSourceResult = KernelProvider->CreateComputeKernel(	
-				BoundKernel.Kernel, ConnectedNode.TraversalContext,
-				NodeDataInterfaceMap, LinkDataInterfaceMap,
-				ValueNodes,
-				GraphDataInterface, GraphDataComponentBinding,
-				KernelDataInterface,
-				BoundKernel.InputDataBindings, BoundKernel.OutputDataBindings, KernelConstantContainer
-			);
-			if (FText* ErrorMessage = KernelSourceResult.TryGet<FText>())
-			{
-				AddDiagnostic(EOptimusDiagnosticLevel::Error,
-					FText::Format(LOCTEXT("CantCreateKernelWithError", "{0}. Compilation aborted."), *ErrorMessage),
-					ConnectedNode.Node);
-				return nullptr;
-			}
+		for (TPair<const UOptimusNode*, UOptimusComputeDataInterface*>& Item : NodeDataInterfaceMap)
+		{
+			ComputeGraph->DataInterfaces.Add(Item.Value);
+			ComputeGraph->DataInterfaceToBinding.Add(DataInterfaceToBindingIndexMap[Item.Value]);
+		}
+		for (TPair<const UOptimusNodePin *, UOptimusComputeDataInterface *>&Item: LinkDataInterfaceMap)
+		{
+			ComputeGraph->DataInterfaces.Add(Item.Value);
+			ComputeGraph->DataInterfaceToBinding.Add(DataInterfaceToBindingIndexMap[Item.Value]);
+		}
+		for (TPair<const UOptimusNode *, UComputeDataInterface *>&Item: KernelDataInterfaceMap)
+		{
+			ComputeGraph->DataInterfaces.Add(Item.Value);
+			ComputeGraph->DataInterfaceToBinding.Add(DataInterfaceToBindingIndexMap[Item.Value]);
+		}
 
-			if (BoundKernel.InputDataBindings.IsEmpty() || BoundKernel.OutputDataBindings.IsEmpty())
-			{
-				AddDiagnostic(EOptimusDiagnosticLevel::Error,
-				LOCTEXT("KernelHasNoBindings", "Kernel has either no input or output bindings. Compilation aborted."),
-					ConnectedNode.Node);
-				return nullptr;
-			}
+		// Create bound kernels
+		struct FKernelWithDataBindings
+		{
+			UComputeKernel *Kernel;
+			FOptimus_InterfaceBindingMap InputDataBindings;
+			FOptimus_InterfaceBindingMap OutputDataBindings;
+		};
 
-			bool bHasExecution = false;
-			for (const TPair<int32, FOptimus_InterfaceBinding>& DataBinding: BoundKernel.InputDataBindings)
+		TArray<FKernelWithDataBindings> BoundKernels;
+		for (FOptimusRoutedConstNode ConnectedNode: ConnectedNodes)
+		{
+			if (const IOptimusComputeKernelProvider *KernelProvider = Cast<const IOptimusComputeKernelProvider>(ConnectedNode.Node))
 			{
-				const int32 KernelBindingIndex = DataBinding.Key;
-				const FOptimus_InterfaceBinding& InterfaceBinding = DataBinding.Value;
-				const UComputeDataInterface* DataInterface = InterfaceBinding.DataInterface;
-				if (DataInterface->IsExecutionInterface())
+				if (KernelComputeGraphIndexMap[KernelProvider] != ComputeGraphIndex)
 				{
-					bHasExecution = true;
-					break;
+					continue;
 				}
-			}
-			
-			if (!bHasExecution)
-			{
-				AddDiagnostic(EOptimusDiagnosticLevel::Error,
-				LOCTEXT("KernelHasNoExecutionDataInterface", "Kernel has no execution data interface connected. Compilation aborted."),
-					ConnectedNode.Node);
-				return nullptr;
-			}
-			
-			BoundKernel.Kernel->KernelSource = KernelSourceResult.Get<UOptimusKernelSource*>();
-
-			BoundKernels.Add(BoundKernel);
+				
+				FKernelWithDataBindings BoundKernel;
 
-			ComputeGraph->KernelInvocations.Add(BoundKernel.Kernel);
-			ComputeGraph->KernelToNode.Add(ConnectedNode.Node);
-		}
-	}
+				BoundKernel.Kernel = NewObject<UComputeKernel>(this);
 
-	// Create a map from the data interfaces to the component bindings.
-	// FIXME: Instead of collecting this during compilation, we should do this when we collect data interfaces instead.
-	TMap<const UComputeDataInterface*, int32> DataInterfaceToBindingIndexMap;
-	for (int32 KernelIndex = 0; KernelIndex < ComputeGraph->KernelInvocations.Num(); KernelIndex++)
-	{
-		const FKernelWithDataBindings& BoundKernel = BoundKernels[KernelIndex];
+				FOptimusKernelConstantContainer& KernelConstantContainer = ConstantContainer.AddContainerForKernel();
+				UComputeDataInterface* KernelDataInterface = KernelDataInterfaceMap[ConnectedNode.Node];
+				
+				FOptimus_ComputeKernelResult KernelSourceResult = KernelProvider->CreateComputeKernel(	
+					BoundKernel.Kernel, ConnectedNode.TraversalContext,
+					NodeDataInterfaceMap, LinkDataInterfaceMap,
+					ValueNodes,
+					GraphDataInterface,
+					KernelDataInterface,
+					BoundKernel.InputDataBindings, BoundKernel.OutputDataBindings, KernelConstantContainer
+				);
+				if (FText* ErrorMessage = KernelSourceResult.TryGet<FText>())
+				{
+					AddDiagnostic(EOptimusDiagnosticLevel::Error,
+						FText::Format(LOCTEXT("CantCreateKernelWithError", "{0}. Compilation aborted."), *ErrorMessage),
+						ConnectedNode.Node);
+					return {};
+				}
 
-		for (const TPair<int32, FOptimus_InterfaceBinding>& DataBinding: BoundKernel.InputDataBindings)
-		{
-			const FOptimus_InterfaceBinding& InterfaceBinding = DataBinding.Value;
-			const UComputeDataInterface* DataInterface = InterfaceBinding.DataInterface;
-			const UOptimusComponentSourceBinding* ComponentBinding = InterfaceBinding.ComponentBinding; 
+				if (BoundKernel.InputDataBindings.IsEmpty() || BoundKernel.OutputDataBindings.IsEmpty())
+				{
+					AddDiagnostic(EOptimusDiagnosticLevel::Error,
+					LOCTEXT("KernelHasNoBindings", "Kernel has either no input or output bindings. Compilation aborted."),
+						ConnectedNode.Node);
+					return {};
+				}
 
-			if (ensure(ComponentBinding))
-			{
-				int32 BindingIndex = Bindings->Bindings.IndexOfByKey(ComponentBinding);
-				if (ensure(BindingIndex != INDEX_NONE))
+				bool bHasExecution = false;
+				for (const TPair<int32, FOptimus_InterfaceBinding>& DataBinding: BoundKernel.InputDataBindings)
 				{
-					if (DataInterfaceToBindingIndexMap.Contains(DataInterface) &&
-						DataInterfaceToBindingIndexMap[DataInterface] != BindingIndex)
+					const int32 KernelBindingIndex = DataBinding.Key;
+					const FOptimus_InterfaceBinding& InterfaceBinding = DataBinding.Value;
+					const UComputeDataInterface* DataInterface = InterfaceBinding.DataInterface;
+					if (DataInterface->IsExecutionInterface())
 					{
-						UE_LOG(LogOptimusCore, Error, TEXT("Datainterface found with different component bindings?"));
+						bHasExecution = true;
+						break;
 					}
-					DataInterfaceToBindingIndexMap.Add(DataInterface, BindingIndex);
 				}
+				
+				if (!bHasExecution)
+				{
+					AddDiagnostic(EOptimusDiagnosticLevel::Error,
+					LOCTEXT("KernelHasNoExecutionDataInterface", "Kernel has no execution data interface connected. Compilation aborted."),
+						ConnectedNode.Node);
+					return {};
+				}
+				
+				BoundKernel.Kernel->KernelSource = KernelSourceResult.Get<UOptimusKernelSource*>();
+
+				BoundKernels.Add(BoundKernel);
+				ComputeGraph->KernelInvocations.Add(BoundKernel.Kernel);
+				ComputeGraph->KernelToNode.Add(ConnectedNode.Node);
 			}
 		}
-		for (const TPair<int32, FOptimus_InterfaceBinding>& DataBinding: BoundKernel.OutputDataBindings)
+		
+		check(ComputeGraph->KernelInvocations.Num() == BoundKernels.Num());
+
+		// Create the graph edges.
+		for (int32 KernelIndex = 0; KernelIndex < ComputeGraph->KernelInvocations.Num(); KernelIndex++)
 		{
-			const FOptimus_InterfaceBinding& InterfaceBinding = DataBinding.Value;
-			const UComputeDataInterface* DataInterface = InterfaceBinding.DataInterface;
-			const UOptimusComponentSourceBinding* ComponentBinding = InterfaceBinding.ComponentBinding; 
+			const FKernelWithDataBindings& BoundKernel = BoundKernels[KernelIndex];
+			const TArray<FShaderFunctionDefinition>& KernelInputs = BoundKernel.Kernel->KernelSource->ExternalInputs;
 
-			if (ensure(ComponentBinding))
+			// FIXME: Hoist these two loops into a helper function/lambda.
+			for (const TPair<int32, FOptimus_InterfaceBinding>& DataBinding: BoundKernel.InputDataBindings)
 			{
-				int32 BindingIndex = Bindings->Bindings.IndexOfByKey(ComponentBinding);
-				if (ensure(BindingIndex != INDEX_NONE))
+				const int32 KernelBindingIndex = DataBinding.Key;
+				const FOptimus_InterfaceBinding& InterfaceBinding = DataBinding.Value;
+				const UComputeDataInterface* DataInterface = InterfaceBinding.DataInterface;
+				const int32 DataInterfaceBindingIndex = InterfaceBinding.DataInterfaceBindingIndex;
+				const FString BindingFunctionName = InterfaceBinding.BindingFunctionName;
+				const FString BindingFunctionNamespace = InterfaceBinding.BindingFunctionNamespace;
+
+				// FIXME: Collect this beforehand.
+				TArray<FShaderFunctionDefinition> DataInterfaceFunctions;
+				DataInterface->GetSupportedInputs(DataInterfaceFunctions);
+				
+				if (ensure(KernelInputs.IsValidIndex(KernelBindingIndex)) &&
+					ensure(DataInterfaceFunctions.IsValidIndex(DataInterfaceBindingIndex)))
 				{
-					if (DataInterfaceToBindingIndexMap.Contains(DataInterface) &&
-						DataInterfaceToBindingIndexMap[DataInterface] != BindingIndex)
-					{
-						UE_LOG(LogOptimusCore, Error, TEXT("Datainterface found with different component bindings?"));
-					}
-					DataInterfaceToBindingIndexMap.Add(DataInterface, BindingIndex);
+					FComputeGraphEdge GraphEdge;
+					GraphEdge.bKernelInput = true;
+					GraphEdge.KernelIndex = KernelIndex;
+					GraphEdge.KernelBindingIndex = KernelBindingIndex;
+					GraphEdge.DataInterfaceIndex = ComputeGraph->DataInterfaces.IndexOfByKey(DataInterface);
+					GraphEdge.DataInterfaceBindingIndex = DataInterfaceBindingIndex;
+					GraphEdge.BindingFunctionNameOverride = BindingFunctionName;
+					GraphEdge.BindingFunctionNamespace = BindingFunctionNamespace;
+					ComputeGraph->GraphEdges.Add(GraphEdge);
 				}
 			}
-		}
-	}
-	
-	// Create the binding objects.
-	for (const UOptimusComponentSourceBinding* Binding: Bindings->Bindings)
-	{
-		ComputeGraph->Bindings.Add(Binding->GetComponentSource()->GetComponentClass());
-	}
 
-	// Now that we've collected all the pieces, time to line them up.
-	ComputeGraph->DataInterfaces.Add(GraphDataInterface);
-	ComputeGraph->DataInterfaceToBinding.Add(0);		// Graph data interface always uses the primary binding.
-
-	for (TPair<const UOptimusNode*, UOptimusComputeDataInterface*>& Item : NodeDataInterfaceMap)
-	{
-		ComputeGraph->DataInterfaces.Add(Item.Value);
-		ComputeGraph->DataInterfaceToBinding.Add(DataInterfaceToBindingIndexMap[Item.Value]);
-	}
-	for (TPair<const UOptimusNodePin *, UOptimusComputeDataInterface *>&Item: LinkDataInterfaceMap)
-	{
-		ComputeGraph->DataInterfaces.Add(Item.Value);
-		ComputeGraph->DataInterfaceToBinding.Add(DataInterfaceToBindingIndexMap[Item.Value]);
-	}
-	for (TPair<const UOptimusNode *, UComputeDataInterface *>&Item: KernelDataInterfaceMap)
-	{
-		ComputeGraph->DataInterfaces.Add(Item.Value);
-		ComputeGraph->DataInterfaceToBinding.Add(DataInterfaceToBindingIndexMap[Item.Value]);
-	}
-
-	// Create the graph edges.
-	for (int32 KernelIndex = 0; KernelIndex < ComputeGraph->KernelInvocations.Num(); KernelIndex++)
-	{
-		const FKernelWithDataBindings& BoundKernel = BoundKernels[KernelIndex];
-		const TArray<FShaderFunctionDefinition>& KernelInputs = BoundKernel.Kernel->KernelSource->ExternalInputs;
-
-		// FIXME: Hoist these two loops into a helper function/lambda.
-		for (const TPair<int32, FOptimus_InterfaceBinding>& DataBinding: BoundKernel.InputDataBindings)
-		{
-			const int32 KernelBindingIndex = DataBinding.Key;
-			const FOptimus_InterfaceBinding& InterfaceBinding = DataBinding.Value;
-			const UComputeDataInterface* DataInterface = InterfaceBinding.DataInterface;
-			const int32 DataInterfaceBindingIndex = InterfaceBinding.DataInterfaceBindingIndex;
-			const FString BindingFunctionName = InterfaceBinding.BindingFunctionName;
-			const FString BindingFunctionNamespace = InterfaceBinding.BindingFunctionNamespace;
-
-			// FIXME: Collect this beforehand.
-			TArray<FShaderFunctionDefinition> DataInterfaceFunctions;
-			DataInterface->GetSupportedInputs(DataInterfaceFunctions);
-			
-			if (ensure(KernelInputs.IsValidIndex(KernelBindingIndex)) &&
-				ensure(DataInterfaceFunctions.IsValidIndex(DataInterfaceBindingIndex)))
+			const TArray<FShaderFunctionDefinition>& KernelOutputs = BoundKernels[KernelIndex].Kernel->KernelSource->ExternalOutputs;
+			for (const TPair<int32, FOptimus_InterfaceBinding>& DataBinding: BoundKernel.OutputDataBindings)
 			{
-				FComputeGraphEdge GraphEdge;
-				GraphEdge.bKernelInput = true;
-				GraphEdge.KernelIndex = KernelIndex;
-				GraphEdge.KernelBindingIndex = KernelBindingIndex;
-				GraphEdge.DataInterfaceIndex = ComputeGraph->DataInterfaces.IndexOfByKey(DataInterface);
-				GraphEdge.DataInterfaceBindingIndex = DataInterfaceBindingIndex;
-				GraphEdge.BindingFunctionNameOverride = BindingFunctionName;
-				GraphEdge.BindingFunctionNamespace = BindingFunctionNamespace;
-				ComputeGraph->GraphEdges.Add(GraphEdge);
-			}
-		}
+				const int32 KernelBindingIndex = DataBinding.Key;
+				const FOptimus_InterfaceBinding& InterfaceBinding = DataBinding.Value;
+				const UComputeDataInterface* DataInterface = InterfaceBinding.DataInterface;
+				const int32 DataInterfaceBindingIndex = InterfaceBinding.DataInterfaceBindingIndex;
+				const FString BindingFunctionName = InterfaceBinding.BindingFunctionName;
+				const FString BindingFunctionNamespace = InterfaceBinding.BindingFunctionNamespace;
 
-		const TArray<FShaderFunctionDefinition>& KernelOutputs = BoundKernels[KernelIndex].Kernel->KernelSource->ExternalOutputs;
-		for (const TPair<int32, FOptimus_InterfaceBinding>& DataBinding: BoundKernel.OutputDataBindings)
-		{
-			const int32 KernelBindingIndex = DataBinding.Key;
-			const FOptimus_InterfaceBinding& InterfaceBinding = DataBinding.Value;
-			const UComputeDataInterface* DataInterface = InterfaceBinding.DataInterface;
-			const int32 DataInterfaceBindingIndex = InterfaceBinding.DataInterfaceBindingIndex;
-			const FString BindingFunctionName = InterfaceBinding.BindingFunctionName;
-			const FString BindingFunctionNamespace = InterfaceBinding.BindingFunctionNamespace;
-
-			// FIXME: Collect this beforehand.
-			TArray<FShaderFunctionDefinition> DataInterfaceFunctions;
-			DataInterface->GetSupportedOutputs(DataInterfaceFunctions);
-			
-			if (ensure(KernelOutputs.IsValidIndex(KernelBindingIndex)) &&
-				ensure(DataInterfaceFunctions.IsValidIndex(DataInterfaceBindingIndex)))
-			{
-				FComputeGraphEdge GraphEdge;
-				GraphEdge.bKernelInput = false;
-				GraphEdge.KernelIndex = KernelIndex;
-				GraphEdge.KernelBindingIndex = KernelBindingIndex;
-				GraphEdge.DataInterfaceIndex = ComputeGraph->DataInterfaces.IndexOfByKey(DataInterface);
-				GraphEdge.DataInterfaceBindingIndex = DataInterfaceBindingIndex;
-				GraphEdge.BindingFunctionNameOverride = BindingFunctionName;
-				GraphEdge.BindingFunctionNamespace = BindingFunctionNamespace;
-				ComputeGraph->GraphEdges.Add(GraphEdge);
+				// FIXME: Collect this beforehand.
+				TArray<FShaderFunctionDefinition> DataInterfaceFunctions;
+				DataInterface->GetSupportedOutputs(DataInterfaceFunctions);
+				
+				if (ensure(KernelOutputs.IsValidIndex(KernelBindingIndex)) &&
+					ensure(DataInterfaceFunctions.IsValidIndex(DataInterfaceBindingIndex)))
+				{
+					FComputeGraphEdge GraphEdge;
+					GraphEdge.bKernelInput = false;
+					GraphEdge.KernelIndex = KernelIndex;
+					GraphEdge.KernelBindingIndex = KernelBindingIndex;
+					GraphEdge.DataInterfaceIndex = ComputeGraph->DataInterfaces.IndexOfByKey(DataInterface);
+					GraphEdge.DataInterfaceBindingIndex = DataInterfaceBindingIndex;
+					GraphEdge.BindingFunctionNameOverride = BindingFunctionName;
+					GraphEdge.BindingFunctionNamespace = BindingFunctionNamespace;
+					ComputeGraph->GraphEdges.Add(GraphEdge);
+				}
 			}
-		}
+		}	
 	}
+	
 
 #if PRINT_COMPILED_OUTPUT
 	
 #endif
 
-	return ComputeGraph;
+	return GraphInfos;
 }
 
 void UOptimusDeformer::OnDataTypeChanged(FName InTypeName)
diff --git a/Engine/Plugins/Animation/DeformerGraph/Source/OptimusCore/Private/OptimusDeformerInstance.cpp b/Engine/Plugins/Animation/DeformerGraph/Source/OptimusCore/Private/OptimusDeformerInstance.cpp
--- a/Engine/Plugins/Animation/DeformerGraph/Source/OptimusCore/Private/OptimusDeformerInstance.cpp
+++ b/Engine/Plugins/Animation/DeformerGraph/Source/OptimusCore/Private/OptimusDeformerInstance.cpp
@@ -5,6 +5,7 @@
 #include "Components/MeshComponent.h"
 #include "ComputeFramework/ComputeFramework.h"
 #include "ComputeWorkerInterface.h"
+#include "IOptimusPersistentBufferProvider.h"
 #include "DataInterfaces/OptimusDataInterfaceGraph.h"
 #include "DataInterfaces/OptimusDataInterfaceRawBuffer.h"
 #include "OptimusComputeGraph.h"
@@ -33,9 +34,11 @@ void FOptimusPersistentBufferPool::GetResourceBuffers(
 	int32 InElementStride,
 	int32 InRawStride,
 	TArray<int32> const& InElementCounts,
-	TArray<FRDGBufferRef>& OutBuffers)
+	TArray<FRDGBufferRef>& OutBuffers,
+	bool& bOutJustAllocated)
 {
 	OutBuffers.Reset();
+	bOutJustAllocated = false;
 
 	TMap<int32, TArray<FOptimusPersistentStructuredBuffer>>& LODResources = ResourceBuffersMap.FindOrAdd(InResourceName);  
 	TArray<FOptimusPersistentStructuredBuffer>* ResourceBuffersPtr = LODResources.Find(InLODIndex);
@@ -43,52 +46,104 @@ void FOptimusPersistentBufferPool::GetResourceBuffers(
 	{
 		// Create pooled buffers and store.
 		TArray<FOptimusPersistentStructuredBuffer> ResourceBuffers;
-		ResourceBuffers.Reserve(InElementCounts.Num());
-
-		// If we are using a raw type alias for the buffer then we need to adjust stride and count.
-		check(InRawStride == 0 || InElementStride % InRawStride == 0);
-		const int32 Stride = InRawStride ? InRawStride : InElementStride;
-		const int32 ElementStrideMultiplier = InRawStride ? InElementStride / InRawStride : 1;
+		AllocateBuffers(GraphBuilder, InElementStride, InRawStride, InElementCounts, ResourceBuffers, OutBuffers);
+		LODResources.Add(InLODIndex, MoveTemp(ResourceBuffers));
+		bOutJustAllocated = true;
+	}
+	else
+	{
+		ValidateAndGetBuffers(GraphBuilder,InElementStride, InElementCounts, *ResourceBuffersPtr, OutBuffers);
+	}
+}
 
-		for (int32 Index = 0; Index < InElementCounts.Num(); Index++)
-		{
-			FRDGBufferDesc BufferDesc = FRDGBufferDesc::CreateStructuredDesc(Stride, InElementCounts[Index] * ElementStrideMultiplier);
-			FRDGBufferRef Buffer = GraphBuilder.CreateBuffer(BufferDesc, TEXT("FOptimusPersistentBuffer"), ERDGBufferFlags::None);
-			OutBuffers.Add(Buffer);
-
-			FOptimusPersistentStructuredBuffer& PersistentBuffer = ResourceBuffers.AddDefaulted_GetRef();
-			PersistentBuffer.ElementStride = InElementStride;
-			PersistentBuffer.ElementCount = InElementCounts[Index];
-			PersistentBuffer.PooledBuffer = GraphBuilder.ConvertToExternalBuffer(Buffer);
-		}
+void FOptimusPersistentBufferPool::GetImplicitPersistentBuffers(
+	FRDGBuilder& GraphBuilder,
+	const FOptimusConstantIdentifier& InIdentifier, 
+	const FOptimusDeformerInstanceComponentLodContext& InLodContext,
+	int32 InElementStride, 
+	int32 InRawStride,
+	TArray<int32> const& InElementCounts,
+	TArray<FRDGBuffer*>& OutBuffers,
+	bool& bOutJustAllocated)
+{
+	OutBuffers.Reset();
+	bOutJustAllocated = false;
 
-		LODResources.Add(InLODIndex, MoveTemp(ResourceBuffers));
+	TMap<FOptimusConstantIdentifier, TArray<FOptimusPersistentStructuredBuffer>>& LODResources = ImplicitBuffersMap.FindOrAdd(InLodContext);  
+	TArray<FOptimusPersistentStructuredBuffer>& ResourceBuffers = LODResources.FindOrAdd(InIdentifier);
+	if (ResourceBuffers.IsEmpty())
+	{
+		// Create pooled buffers and store.
+		AllocateBuffers(GraphBuilder, InElementStride, InRawStride, InElementCounts, ResourceBuffers, OutBuffers);
+		bOutJustAllocated = true;
 	}
 	else
 	{
-		// Verify that the buffers are correct based on the incoming information. 
-		// If there's a mismatch, then something has gone wrong upstream.
-		// Maybe either duplicated names, missing resource clearing on recompile, or something else.
-		if (!ensure(ResourceBuffersPtr->Num() == InElementCounts.Num()))
-		{
-			return;
-		}
+		ValidateAndGetBuffers(GraphBuilder,InElementStride, InElementCounts, ResourceBuffers, OutBuffers);
+	}
+}
+
+
+
+void FOptimusPersistentBufferPool::AllocateBuffers(
+	FRDGBuilder& GraphBuilder,
+	int32 InElementStride,
+	int32 InRawStride,
+	TArray<int32> const& InElementCounts,
+	TArray<FOptimusPersistentStructuredBuffer>& OutResourceBuffers,
+	TArray<FRDGBuffer*>& OutBuffers
+	)
+{
+	OutResourceBuffers.Reserve(InElementCounts.Num());
 
-		for (int32 Index = 0; Index < ResourceBuffersPtr->Num(); Index++)
+	// If we are using a raw type alias for the buffer then we need to adjust stride and count.
+	check(InRawStride == 0 || InElementStride % InRawStride == 0);
+	const int32 Stride = InRawStride ? InRawStride : InElementStride;
+	const int32 ElementStrideMultiplier = InRawStride ? InElementStride / InRawStride : 1;
+
+	for (int32 Index = 0; Index < InElementCounts.Num(); Index++)
+	{
+		FRDGBufferDesc BufferDesc = FRDGBufferDesc::CreateStructuredDesc(Stride, InElementCounts[Index] * ElementStrideMultiplier);
+		FRDGBufferRef Buffer = GraphBuilder.CreateBuffer(BufferDesc, TEXT("FOptimusPersistentBuffer"), ERDGBufferFlags::None);
+		OutBuffers.Add(Buffer);
+
+		FOptimusPersistentStructuredBuffer& PersistentBuffer = OutResourceBuffers.AddDefaulted_GetRef();
+		PersistentBuffer.ElementStride = InElementStride;
+		PersistentBuffer.ElementCount = InElementCounts[Index];
+		PersistentBuffer.PooledBuffer = GraphBuilder.ConvertToExternalBuffer(Buffer);
+	}
+}
+
+void FOptimusPersistentBufferPool::ValidateAndGetBuffers(
+	FRDGBuilder& GraphBuilder,
+	int32 InElementStride,
+	TArray<int32> const& InElementCounts,
+	const TArray<FOptimusPersistentStructuredBuffer>& InResourceBuffers,
+	TArray<FRDGBuffer*>& OutBuffers
+	) const
+{
+	// Verify that the buffers are correct based on the incoming information. 
+	// If there's a mismatch, then something has gone wrong upstream.
+	// Maybe either duplicated names, missing resource clearing on recompile, or something else.
+	if (!ensure(InResourceBuffers.Num() == InElementCounts.Num()))
+	{
+		return;
+	}
+
+	for (int32 Index = 0; Index < InResourceBuffers.Num(); Index++)
+	{
+		const FOptimusPersistentStructuredBuffer& PersistentBuffer = InResourceBuffers[Index];
+		if (!ensure(PersistentBuffer.PooledBuffer.IsValid()) ||
+			!ensure(PersistentBuffer.ElementStride == InElementStride) ||
+			!ensure(PersistentBuffer.ElementCount == InElementCounts[Index]))
 		{
-			FOptimusPersistentStructuredBuffer& PersistentBuffer = (*ResourceBuffersPtr)[Index];
-			if (!ensure(PersistentBuffer.PooledBuffer.IsValid()) ||
-				!ensure(PersistentBuffer.ElementStride == InElementStride) ||
-				!ensure(PersistentBuffer.ElementCount == InElementCounts[Index]))
-			{
-				OutBuffers.Reset();
-				return;
-			}	
+			OutBuffers.Reset();
+			return;
+		}	
 
-			// Register buffer back into the graph and return it.
-			FRDGBufferRef Buffer = GraphBuilder.RegisterExternalBuffer(PersistentBuffer.PooledBuffer);
-			OutBuffers.Add(Buffer);
-		}
+		// Register buffer back into the graph and return it.
+		FRDGBufferRef Buffer = GraphBuilder.RegisterExternalBuffer(PersistentBuffer.PooledBuffer);
+		OutBuffers.Add(Buffer);
 	}
 }
 
@@ -96,6 +151,7 @@ void FOptimusPersistentBufferPool::ReleaseResources()
 {
 	check(IsInRenderingThread());
 	ResourceBuffersMap.Reset();
+	ImplicitBuffersMap.Reset();
 }
 
 FOptimusDeformerInstanceExecInfo::FOptimusDeformerInstanceExecInfo()
@@ -360,10 +416,9 @@ void UOptimusDeformerInstance::SetupFromDeformer(UOptimusDeformer* InDeformer)
 		for(TObjectPtr<UComputeDataProvider> DataProvider: Info.ComputeGraphInstance.GetDataProviders())
 		{
 			// Make the persistent buffer data provider aware of the buffer pool and current LOD index.
-			// TBD: Interface-based.
-			if (UOptimusPersistentBufferDataProvider* PersistentBufferProvider = Cast<UOptimusPersistentBufferDataProvider>(DataProvider))
+			if (IOptimusPersistentBufferProvider* PersistentBufferProvider = Cast<IOptimusPersistentBufferProvider>(DataProvider))
 			{
-				PersistentBufferProvider->BufferPool = BufferPool;
+				PersistentBufferProvider->SetBufferPool(BufferPool);
 			}
 
 			// Set this instance on the graph data provider so that it can query variables.
@@ -417,15 +472,11 @@ void UOptimusDeformerInstance::SetCanBeActive(bool bInCanBeActive)
 	bCanBeActive = bInCanBeActive;
 }
 
-TArray<float> UOptimusDeformerInstance::GetConstantValuePerInvocation(const FOptimusConstantIdentifier& InIdentifier)
+FOptimusConstantEvaluationResult UOptimusDeformerInstance::GetConstantValuePerInvocation(
+	const FOptimusConstantIdentifier& InIdentifier)
 {
-	if (!InIdentifier.IsValid())
-	{
-		return {};
-	}
-	
-	FOptimusDeformerInstanceComponentContext ComponentContext;
-	
+	FOptimusConstantEvaluationResult Result;
+
 	check(WeakBoundComponents.Num() == WeakComponentSources.Num());
 	int32 NumBindings = WeakBoundComponents.Num();
 	for (int32 BindingIndex = 0; BindingIndex < NumBindings; BindingIndex++)
@@ -442,18 +493,24 @@ TArray<float> UOptimusDeformerInstance::GetConstantValuePerInvocation(const FOpt
 
 		if (LodIndex != INDEX_NONE)
 		{
-			ComponentContext.LodIndexPerComponent.Add(LodIndex);
+			Result.LodContext.LodIndexPerComponent.Add(LodIndex);
 		}
 		else
 		{
 			return {};
 		}
 	}
+	
+	
+	if (!InIdentifier.IsValid())
+	{
+		return Result;
+	}
 
-
-	if (const FOptimusConstantContainerInstance* ConstantValues = ConstantValuesPerContext.Find(ComponentContext))
+	if (const FOptimusConstantContainerInstance* ConstantValues = ConstantValuesPerContext.Find(Result.LodContext))
 	{
-		return ConstantValues->GetConstantValuePerInvocation(InIdentifier);
+		Result.ValuePerInvocation = ConstantValues->GetConstantValuePerInvocation(InIdentifier);
+		return Result;
 	}
 
 	TMap<int32, TMap<FName, TArray<float>>> BindingIndexToConstantValues;
@@ -501,15 +558,15 @@ TArray<float> UOptimusDeformerInstance::GetConstantValuePerInvocation(const FOpt
 
 	if (CollectAllBindingConstants())
 	{
-		FOptimusConstantContainerInstance& NewConstantValues = ConstantValuesPerContext.Add(ComponentContext);
+		FOptimusConstantContainerInstance& NewConstantValues = ConstantValuesPerContext.Add(Result.LodContext);
 		if (NewConstantValues.Initialize(ConstantContainer, BindingIndexToConstantValues))
 		{
-			return NewConstantValues.GetConstantValuePerInvocation(InIdentifier);		
+			Result.ValuePerInvocation = NewConstantValues.GetConstantValuePerInvocation(InIdentifier);
+			return Result;
 		}
 	}
 
-	return{};
-	
+	return Result;
 }
 
 void UOptimusDeformerInstance::AllocateResources()
diff --git a/Engine/Plugins/Animation/DeformerGraph/Source/OptimusCore/Private/OptimusNodeGraph.cpp b/Engine/Plugins/Animation/DeformerGraph/Source/OptimusCore/Private/OptimusNodeGraph.cpp
--- a/Engine/Plugins/Animation/DeformerGraph/Source/OptimusCore/Private/OptimusNodeGraph.cpp
+++ b/Engine/Plugins/Animation/DeformerGraph/Source/OptimusCore/Private/OptimusNodeGraph.cpp
@@ -27,6 +27,7 @@
 
 #include <limits>
 
+#include "IOptimusPinMutabilityDefiner.h"
 #include "OptimusNodeSubGraph.h"
 #include "Nodes/OptimusNode_CustomComputeKernel.h"
 #include "Nodes/OptimusNode_FunctionReference.h"
@@ -1493,6 +1494,81 @@ TSet<UOptimusComponentSourceBinding*> UOptimusNodeGraph::GetComponentSourceBindi
 }
 
 
+bool UOptimusNodeGraph::IsPinMutable(const UOptimusNodePin* InNodePin) const
+{
+	TSet<FOptimusRoutedConstNode> VisitedNodes;
+	TQueue<FOptimusRoutedConstNodePin> WorkingSet;
+
+	// If given an input pin, find the other side, since output pin provides mutability definition
+	const UOptimusNodePin* StartNodePin;
+	if (InNodePin->GetDirection() == EOptimusNodePinDirection::Input)
+	{
+		TArray<FOptimusRoutedNodePin> RoutedPins = GetConnectedPinsWithRouting(InNodePin, {});
+		if (RoutedPins.IsEmpty())
+		{
+			// Nothing is connected, pin data will never change
+			return false;
+		}
+		
+		check(RoutedPins.Num() == 1);
+		StartNodePin = RoutedPins[0].NodePin;
+	}
+	else
+	{
+		StartNodePin = InNodePin;
+	}
+
+	WorkingSet.Enqueue({StartNodePin, FOptimusPinTraversalContext{}});
+	
+	FOptimusRoutedConstNodePin WorkItem;
+	while (WorkingSet.Dequeue(WorkItem))
+	{
+		TArray<const UOptimusNodePin*> InputPins; 
+		const UOptimusNodePin* NodePin = WorkItem.NodePin;
+		const UOptimusNode* Node = NodePin->GetOwningNode();
+
+		if (const IOptimusPinMutabilityDefiner* PinMutabilityDefiner = Cast<const IOptimusPinMutabilityDefiner>(Node))
+		{
+			if (PinMutabilityDefiner->IsOutputPinMutable(NodePin))
+			{
+				return true;
+			}
+		}
+		
+		// Grab all inputs.
+		for (const UOptimusNodePin* Pin: Node->GetPins())
+		{
+			if (Pin->GetDirection() == EOptimusNodePinDirection::Input)
+			{
+				InputPins.Append(Pin->GetSubPinsRecursively(true));
+			}
+		}
+		
+		// Traverse in the direction of inputs to outputs (up the graph).
+		for (const UOptimusNodePin* Pin: InputPins)
+		{
+			for (const FOptimusRoutedNodePin& ConnectedPin: Pin->GetConnectedPinsWithRouting(WorkItem.TraversalContext))
+			{
+				if (ensure(ConnectedPin.NodePin != nullptr))
+				{
+					FOptimusRoutedConstNodePin CollectedNodePin{ConnectedPin.NodePin, ConnectedPin.TraversalContext};
+					const UOptimusNode *NextNode = ConnectedPin.NodePin->GetOwningNode();
+					FOptimusRoutedConstNode CollectedNode{NextNode, ConnectedPin.TraversalContext};
+
+					if (!VisitedNodes.Contains(CollectedNode))
+					{
+						VisitedNodes.Add(CollectedNode);
+						WorkingSet.Enqueue(CollectedNodePin);
+					}
+				}
+			}
+		}
+	}
+
+	return false;
+}
+
+
 TArray<const UOptimusNodeLink*> UOptimusNodeGraph::GetPinLinks(
 	const UOptimusNodePin* InNodePin
 	) const
diff --git a/Engine/Plugins/Animation/DeformerGraph/Source/OptimusCore/Private/OptimusNodePin.cpp b/Engine/Plugins/Animation/DeformerGraph/Source/OptimusCore/Private/OptimusNodePin.cpp
--- a/Engine/Plugins/Animation/DeformerGraph/Source/OptimusCore/Private/OptimusNodePin.cpp
+++ b/Engine/Plugins/Animation/DeformerGraph/Source/OptimusCore/Private/OptimusNodePin.cpp
@@ -187,6 +187,25 @@ TSet<UOptimusComponentSourceBinding*> UOptimusNodePin::GetComponentSourceBinding
 	return GetOwningNode()->GetOwningGraph()->GetComponentSourceBindingsForPin(this);
 }
 
+TSet<UOptimusComponentSourceBinding*> UOptimusNodePin::GetComponentSourceBindingsRecursively() const
+{
+	TArray<UOptimusNodePin*> PinsToConsider = GetSubPinsRecursively(true);
+
+	TSet<UOptimusComponentSourceBinding*> Bindings;
+
+	for (const UOptimusNodePin* Pin : PinsToConsider)
+	{
+		Bindings.Append(Pin->GetComponentSourceBindings());
+	}
+	
+	return Bindings;
+}
+
+bool UOptimusNodePin::IsMutable() const
+{
+	return GetOwningNode()->GetOwningGraph()->IsPinMutable(this);
+}
+
 
 FProperty* UOptimusNodePin::GetPropertyFromPin() const
 {
diff --git a/Engine/Plugins/Animation/DeformerGraph/Source/OptimusCore/Public/IOptimusComputeKernelProvider.h b/Engine/Plugins/Animation/DeformerGraph/Source/OptimusCore/Public/IOptimusComputeKernelProvider.h
--- a/Engine/Plugins/Animation/DeformerGraph/Source/OptimusCore/Public/IOptimusComputeKernelProvider.h
+++ b/Engine/Plugins/Animation/DeformerGraph/Source/OptimusCore/Public/IOptimusComputeKernelProvider.h
@@ -23,7 +23,6 @@ struct FOptimusExecutionDomain;
 struct FOptimus_InterfaceBinding
 {
 	const UComputeDataInterface* DataInterface;
-	const UOptimusComponentSourceBinding* ComponentBinding;
 	int32 DataInterfaceBindingIndex;
 	FString BindingFunctionName;
 	FString BindingFunctionNamespace;
@@ -80,7 +79,6 @@ public:
 	 * @param InLinkDataInterfaceMap
 	 * @param InValueNodes
 	 * @param InGraphDataInterface
-	 * @param InGraphDataComponentBinding
 	 * @param InOutKernelDataInterface
 	 * @param OutInputDataBindings
 	 * @param OutOutputDataBindings
@@ -93,7 +91,6 @@ public:
 		const FOptimus_PinToDataInterfaceMap& InLinkDataInterfaceMap,
 		const TArray<const UOptimusNode*>& InValueNodes,
 		const UComputeDataInterface* InGraphDataInterface,
-		const UOptimusComponentSourceBinding* InGraphDataComponentBinding,
 		UComputeDataInterface* InOutKernelDataInterface,
 		FOptimus_InterfaceBindingMap& OutInputDataBindings,
 		FOptimus_InterfaceBindingMap& OutOutputDataBindings,
@@ -110,5 +107,10 @@ public:
 	virtual UComputeDataInterface* MakeKernelDataInterface(UObject* InOuter) const = 0;
 
 	/** Check if a specific pin needs to support atomic operation */
-	virtual bool GetPinSupportAtomic(const UOptimusNodePin* InPin) const = 0;
+	virtual bool DoesOutputPinSupportAtomic(const UOptimusNodePin* InPin) const = 0;
+	
+	/** Check if a specific pin needs to support atomic operation */
+	virtual bool DoesOutputPinSupportRead(const UOptimusNodePin* InPin) const = 0;
+
+	virtual bool HasMutableInput() const = 0;
 };
diff --git a/Engine/Plugins/Animation/DeformerGraph/Source/OptimusCore/Public/IOptimusParameterBindingProvider.h b/Engine/Plugins/Animation/DeformerGraph/Source/OptimusCore/Public/IOptimusParameterBindingProvider.h
--- a/Engine/Plugins/Animation/DeformerGraph/Source/OptimusCore/Public/IOptimusParameterBindingProvider.h
+++ b/Engine/Plugins/Animation/DeformerGraph/Source/OptimusCore/Public/IOptimusParameterBindingProvider.h
@@ -23,5 +23,6 @@ class OPTIMUSCORE_API IOptimusParameterBindingProvider
 
 public:
 	virtual FString GetBindingDeclaration(FName BindingName) const = 0;
-	virtual bool GetBindingAtomicSupportCheckBoxVisibility(FName BindingName) const = 0;
+	virtual bool GetBindingSupportAtomicCheckBoxVisibility(FName BindingName) const = 0;
+	virtual bool GetBindingSupportReadCheckBoxVisibility(FName BindingName) const = 0;
 };
diff --git a/Engine/Plugins/Animation/DeformerGraph/Source/OptimusCore/Public/IOptimusPersistentBufferProvider.h b/Engine/Plugins/Animation/DeformerGraph/Source/OptimusCore/Public/IOptimusPersistentBufferProvider.h
new file mode 100644
--- /dev/null
+++ b/Engine/Plugins/Animation/DeformerGraph/Source/OptimusCore/Public/IOptimusPersistentBufferProvider.h
@@ -0,0 +1,26 @@
+﻿// Copyright Epic Games, Inc. All Rights Reserved.
+
+#pragma once
+
+#include "UObject/Interface.h"
+
+#include "IOptimusPersistentBufferProvider.generated.h"
+
+
+class FOptimusPersistentBufferPool;
+
+UINTERFACE()
+class OPTIMUSCORE_API UOptimusPersistentBufferProvider :
+	public UInterface
+{
+	GENERATED_BODY()
+};
+
+
+class IOptimusPersistentBufferProvider
+{
+	GENERATED_BODY()
+
+public:
+	virtual void SetBufferPool(TSharedPtr<FOptimusPersistentBufferPool> InBufferPool) = 0;
+};
diff --git a/Engine/Plugins/Animation/DeformerGraph/Source/OptimusCore/Public/IOptimusPinMutabilityDefiner.h b/Engine/Plugins/Animation/DeformerGraph/Source/OptimusCore/Public/IOptimusPinMutabilityDefiner.h
new file mode 100644
--- /dev/null
+++ b/Engine/Plugins/Animation/DeformerGraph/Source/OptimusCore/Public/IOptimusPinMutabilityDefiner.h
@@ -0,0 +1,26 @@
+﻿// Copyright Epic Games, Inc. All Rights Reserved.
+
+#pragma once
+
+#include "UObject/Interface.h"
+
+#include "IOptimusPinMutabilityDefiner.generated.h"
+
+
+class UOptimusNodePin;
+UINTERFACE()
+class OPTIMUSCORE_API UOptimusPinMutabilityDefiner :
+	public UInterface
+{
+	GENERATED_BODY()
+};
+
+
+class IOptimusPinMutabilityDefiner
+{
+	GENERATED_BODY()
+
+public:
+	/** Returns the component binding for the node that this interface is implemented on */
+	virtual bool IsOutputPinMutable(const UOptimusNodePin* InPin) const = 0;
+};
diff --git a/Engine/Plugins/Animation/DeformerGraph/Source/OptimusCore/Public/OptimusBindingTypes.h b/Engine/Plugins/Animation/DeformerGraph/Source/OptimusCore/Public/OptimusBindingTypes.h
--- a/Engine/Plugins/Animation/DeformerGraph/Source/OptimusCore/Public/OptimusBindingTypes.h
+++ b/Engine/Plugins/Animation/DeformerGraph/Source/OptimusCore/Public/OptimusBindingTypes.h
@@ -49,6 +49,10 @@ struct FOptimusParameterBinding
 	// Int type resource can optionally support atomic writes, memory is zero-initialized
 	UPROPERTY(EditAnywhere, DisplayName="Support Atomic", Category = Binding)
 	bool bSupportAtomicIfCompatibleDataType = false;
+
+	// Optionally support both read and write
+	UPROPERTY(EditAnywhere, DisplayName="Support Read", Category = Binding)
+	bool bSupportRead = false;
 };
 
 
diff --git a/Engine/Plugins/Animation/DeformerGraph/Source/OptimusCore/Public/OptimusComputeDataInterface.h b/Engine/Plugins/Animation/DeformerGraph/Source/OptimusCore/Public/OptimusComputeDataInterface.h
--- a/Engine/Plugins/Animation/DeformerGraph/Source/OptimusCore/Public/OptimusComputeDataInterface.h
+++ b/Engine/Plugins/Animation/DeformerGraph/Source/OptimusCore/Public/OptimusComputeDataInterface.h
@@ -19,10 +19,12 @@ struct FOptimusCDIPinDefinition
 	// Singleton value read/write. The context name is implied as Optimus::ContextName::Singleton.
 	FOptimusCDIPinDefinition(
 		const FName InPinName,
-		const FString InDataFunctionName
+		const FString InDataFunctionName,
+		const bool bInMutable = true
 		) :
 		PinName(InPinName),
-		DataFunctionName(InDataFunctionName)
+		DataFunctionName(InDataFunctionName),
+		bMutable(bInMutable)
 	{ }
 
 	// A single level context lookup.
@@ -30,11 +32,13 @@ struct FOptimusCDIPinDefinition
 		const FName InPinName,
 		const FString InDataFunctionName,
 		const FName InContextName,
-		const FString InCountFunctionName
+		const FString InCountFunctionName,
+		const bool bInMutable = true
 		) :
 		PinName(InPinName),
 		DataFunctionName(InDataFunctionName),
-		DataDimensions{{InContextName, InCountFunctionName}}
+		DataDimensions{{InContextName, InCountFunctionName}},
+		bMutable(bInMutable)
 	{ }
 
 	FOptimusCDIPinDefinition(
@@ -42,21 +46,25 @@ struct FOptimusCDIPinDefinition
 		const FString InDataFunctionName,
 		const FName InContextName,
 		const int32 InMultiplier,
-		const FString InCountFunctionName
+		const FString InCountFunctionName,
+		const bool bInMutable = true
 		) :
 		PinName(InPinName),
 		DataFunctionName(InDataFunctionName),
 		DataDimensions{{InContextName, InCountFunctionName}},
-		DomainMultiplier(FMath::Max(1, InMultiplier))
+		DomainMultiplier(FMath::Max(1, InMultiplier)),
+		bMutable(bInMutable)
 	{ }
 	
 	FOptimusCDIPinDefinition(
 		const FName InPinName,
 		const FString InDataFunctionName,
-		const std::initializer_list<FDimensionInfo> InContexts
+		const std::initializer_list<FDimensionInfo> InContexts,
+		const bool bInMutable = true
 		) :
 		PinName(InPinName),
-		DataFunctionName(InDataFunctionName)
+		DataFunctionName(InDataFunctionName),
+		bMutable(bInMutable)
 	{
 		for (FDimensionInfo ContextInfo: InContexts)
 		{
@@ -110,6 +118,8 @@ struct FOptimusCDIPinDefinition
 
 	// For single-level domains, how many values per element of that dimension's range. 
 	int32 DomainMultiplier = 1;
+
+	bool bMutable = true;
 };
 
 
diff --git a/Engine/Plugins/Animation/DeformerGraph/Source/OptimusCore/Public/OptimusDeformer.h b/Engine/Plugins/Animation/DeformerGraph/Source/OptimusCore/Public/OptimusDeformer.h
--- a/Engine/Plugins/Animation/DeformerGraph/Source/OptimusCore/Public/OptimusDeformer.h
+++ b/Engine/Plugins/Animation/DeformerGraph/Source/OptimusCore/Public/OptimusDeformer.h
@@ -483,8 +483,8 @@ private:
 
 	TArray<UOptimusNode*> GetAllNodesOfClass(UClass* InNodeClass) const;
 	
-	/// Compile a node graph to a compute graph. Returns a complete compute graph if compilation succeeded. 
-	UOptimusComputeGraph* CompileNodeGraphToComputeGraph(
+	/// Compile a node graph to a compute graph. Returns one or two complete compute graphs if compilation succeeded. 
+	TArray<FOptimusComputeGraphInfo> CompileNodeGraphToComputeGraphs(
 		const UOptimusNodeGraph *InNodeGraph,
 		TFunction<void(EOptimusDiagnosticLevel, FText, const UObject*)> InErrorReporter
 		);
diff --git a/Engine/Plugins/Animation/DeformerGraph/Source/OptimusCore/Public/OptimusDeformerInstance.h b/Engine/Plugins/Animation/DeformerGraph/Source/OptimusCore/Public/OptimusDeformerInstance.h
--- a/Engine/Plugins/Animation/DeformerGraph/Source/OptimusCore/Public/OptimusDeformerInstance.h
+++ b/Engine/Plugins/Animation/DeformerGraph/Source/OptimusCore/Public/OptimusDeformerInstance.h
@@ -22,15 +22,12 @@ class UOptimusVariableContainer;
 class UOptimusVariableDescription;
 class UOptimusComponentSourceBinding;
 
-USTRUCT()
-struct FOptimusDeformerInstanceComponentContext
-{
-	GENERATED_BODY()
 
-	UPROPERTY()
+struct FOptimusDeformerInstanceComponentLodContext
+{
 	TArray<int32> LodIndexPerComponent;
 
-	friend uint32 GetTypeHash(const FOptimusDeformerInstanceComponentContext& InContext)
+	friend uint32 GetTypeHash(const FOptimusDeformerInstanceComponentLodContext& InContext)
 	{
 		uint32 Hash = GetTypeHash(InContext.LodIndexPerComponent.Num());
 		
@@ -41,12 +38,23 @@ struct FOptimusDeformerInstanceComponentContext
 		return Hash;
 	}
 
-	bool operator==(const FOptimusDeformerInstanceComponentContext& InOther) const
+	bool operator==(const FOptimusDeformerInstanceComponentLodContext& InOther) const
 	{
 		return LodIndexPerComponent == InOther.LodIndexPerComponent;
 	}
 };
 
+struct FOptimusConstantEvaluationResult
+{
+	TArray<float> ValuePerInvocation;
+	FOptimusDeformerInstanceComponentLodContext LodContext;
+
+	bool IsValid() const
+	{
+		return ValuePerInvocation.Num() != 0;
+	};
+};
+
 class FOptimusPersistentBufferPool
 {
 public:
@@ -63,13 +71,42 @@ public:
 		int32 InElementStride,
 		int32 InRawStride,
 		TArray<int32> const& InElementCounts,
-		TArray<FRDGBuffer*>& OutBuffers );
+		TArray<FRDGBuffer*>& OutBuffers,
+		bool& bOutJustAllocated);
+
+	void GetImplicitPersistentBuffers(
+		FRDGBuilder& GraphBuilder,
+		const FOptimusConstantIdentifier& InIdentifier,
+		const FOptimusDeformerInstanceComponentLodContext& InLodContext,
+		int32 InElementStride,
+		int32 InRawStride,
+		TArray<int32> const& InElementCounts,
+		TArray<FRDGBuffer*>& OutBuffers,
+		bool& bOutJustAllocated);
 
 	/** Release _all_ resources allocated by this pool */
 	void ReleaseResources();
 	
 private:
+	 void AllocateBuffers(
+		FRDGBuilder& GraphBuilder,
+		int32 InElementStride,
+		int32 InRawStride,
+		TArray<int32> const& InElementCounts,
+		TArray<FOptimusPersistentStructuredBuffer>& OutResourceBuffers,
+		TArray<FRDGBuffer*>& OutBuffers
+		);
+
+	void ValidateAndGetBuffers(
+		FRDGBuilder& GraphBuilder,
+		int32 InElementStride,
+		TArray<int32> const& InElementCounts,
+		const TArray<FOptimusPersistentStructuredBuffer>& InResourceBuffers,
+		TArray<FRDGBuffer*>& OutBuffers
+		) const;
+	
 	TMap<FName, TMap<int32, TArray<FOptimusPersistentStructuredBuffer>>> ResourceBuffersMap;
+	TMap<FOptimusDeformerInstanceComponentLodContext, TMap<FOptimusConstantIdentifier, TArray<FOptimusPersistentStructuredBuffer>>> ImplicitBuffersMap;
 };
 using FOptimusPersistentBufferPoolPtr = TSharedPtr<FOptimusPersistentBufferPool>;
 
@@ -228,7 +265,7 @@ public:
 
 	void SetCanBeActive(bool bInCanBeActive);
 
-	TArray<float> GetConstantValuePerInvocation(const FOptimusConstantIdentifier& InIdentifier);
+	FOptimusConstantEvaluationResult GetConstantValuePerInvocation(const FOptimusConstantIdentifier& InIdentifier);
 
 protected:
 	/** Implementation of UMeshDeformerInstance. */
@@ -262,7 +299,7 @@ private:
 	UPROPERTY()
 	FOptimusConstantContainer ConstantContainer;
 	
-	TMap<FOptimusDeformerInstanceComponentContext, FOptimusConstantContainerInstance> ConstantValuesPerContext;
+	TMap<FOptimusDeformerInstanceComponentLodContext, FOptimusConstantContainerInstance> ConstantValuesPerContext;
 
 	// List of graphs that should be run on the next tick. 
 	TSet<FName> GraphsToRunOnNextTick;
diff --git a/Engine/Plugins/Animation/DeformerGraph/Source/OptimusCore/Public/OptimusNodeGraph.h b/Engine/Plugins/Animation/DeformerGraph/Source/OptimusCore/Public/OptimusNodeGraph.h
--- a/Engine/Plugins/Animation/DeformerGraph/Source/OptimusCore/Public/OptimusNodeGraph.h
+++ b/Engine/Plugins/Animation/DeformerGraph/Source/OptimusCore/Public/OptimusNodeGraph.h
@@ -286,6 +286,11 @@ public:
 	TSet<UOptimusComponentSourceBinding*> GetComponentSourceBindingsForPin(
 		const UOptimusNodePin* InNodePin
 		) const;
+
+	/** Check if a pin is represents time varying data */
+	bool IsPinMutable(
+		const UOptimusNodePin* InNodePin
+		) const;
 		
 	TArray<const UOptimusNodeLink *> GetPinLinks(const UOptimusNodePin* InNodePin) const;
 
diff --git a/Engine/Plugins/Animation/DeformerGraph/Source/OptimusCore/Public/OptimusNodePin.h b/Engine/Plugins/Animation/DeformerGraph/Source/OptimusCore/Public/OptimusNodePin.h
--- a/Engine/Plugins/Animation/DeformerGraph/Source/OptimusCore/Public/OptimusNodePin.h
+++ b/Engine/Plugins/Animation/DeformerGraph/Source/OptimusCore/Public/OptimusNodePin.h
@@ -89,11 +89,18 @@ public:
 	/** Returns the data domain that this pin is expected to cover */
 	const FOptimusDataDomain& GetDataDomain() const { return DataDomain; }
 	
-	/** Return all component source bindings that flow into this input pin. If the pin is an
-	 *  output pin, it always returns an empty array.
+	/** Return all component source bindings that flow into pin. 
 	 */
 	TSet<UOptimusComponentSourceBinding*> GetComponentSourceBindings() const;
 
+	/** Return all component source bindings that flow into this input pin and its sub-pins.
+	 */
+	TSet<UOptimusComponentSourceBinding*> GetComponentSourceBindingsRecursively() const;
+
+	/** Whether the data presented by the pin can change over time
+	 */
+	bool IsMutable() const;
+	
 	/** Returns the FProperty object for this pin. This can be used to directly address the
 	  * node data represented by this pin. Not all pins have an underlying resource so this can
 	  * return nullptr.
diff --git a/Engine/Plugins/Animation/DeformerGraph/Source/OptimusEditor/Private/OptimusDetailsCustomization.cpp b/Engine/Plugins/Animation/DeformerGraph/Source/OptimusEditor/Private/OptimusDetailsCustomization.cpp
--- a/Engine/Plugins/Animation/DeformerGraph/Source/OptimusEditor/Private/OptimusDetailsCustomization.cpp
+++ b/Engine/Plugins/Animation/DeformerGraph/Source/OptimusEditor/Private/OptimusDetailsCustomization.cpp
@@ -1101,17 +1101,9 @@ public:
 		
 		TSharedPtr<IPropertyHandle> DataTypeProperty = BindingPropertyHandle->GetChildHandle(GET_MEMBER_NAME_CHECKED(FOptimusParameterBinding, DataType));
 		const TSharedPtr<IPropertyHandle> DataDomainProperty = BindingPropertyHandle->GetChildHandle(GET_MEMBER_NAME_CHECKED(FOptimusParameterBinding, DataDomain));
-		const TSharedPtr<IPropertyHandle> SupportAtomicProperty = BindingPropertyHandle->GetChildHandle(GET_MEMBER_NAME_CHECKED(FOptimusParameterBinding, bSupportAtomicIfCompatibleDataType));
 
 		const TArray<TWeakObjectPtr<UObject>>& SelectedObjects = InCustomizationUtils.GetPropertyUtilities()->GetSelectedObjects();
 
-		auto IsAtomicCheckBoxVisible = [SelectedObjects, InBindingPropertyHandle]()
-		{
-			return FOptimusParameterBindingCustomization::IsAtomicCheckBoxVisible(SelectedObjects, InBindingPropertyHandle);
-		};
-
-		TSharedRef<SWidget> AtomicToggleWidget = SupportAtomicProperty->CreatePropertyValueWidget(false);
-		
 		FDetailWidgetRow DataTypeHeaderRow;
 		DataTypeRefCustomizationInstance = FOptimusDataTypeRefCustomization::MakeInstance();
 		DataTypeRefCustomizationInstance->CustomizeHeader(DataTypeProperty.ToSharedRef(), DataTypeHeaderRow, InCustomizationUtils);
@@ -1163,17 +1155,6 @@ public:
 						[
 							DataTypeHeaderRow.ValueContent().Widget
 						]
-						+SHorizontalBox::Slot()
-						.AutoWidth()
-						[
-							SNew(SHorizontalBox)
-							.Visibility_Lambda(IsAtomicCheckBoxVisible)
-							+SHorizontalBox::Slot()
-							.Padding(5,0,0,0)
-							[
-								AtomicToggleWidget
-							]
-						]
 					]
 				]
 				+SSplitter::Slot()
@@ -1277,7 +1258,7 @@ EVisibility FOptimusParameterBindingCustomization::IsAtomicCheckBoxVisible(const
 				// During drag & reorder, we can have invalid bindings in the property
 				if (Binding && Binding->Name != NAME_None)
 				{
-					return BindingProvider->GetBindingAtomicSupportCheckBoxVisibility(Binding->Name) ? EVisibility::Visible : EVisibility::Collapsed;
+					return BindingProvider->GetBindingSupportAtomicCheckBoxVisibility(Binding->Name) ? EVisibility::Visible : EVisibility::Collapsed;
 				}
 			}
 			break;
@@ -1287,6 +1268,32 @@ EVisibility FOptimusParameterBindingCustomization::IsAtomicCheckBoxVisible(const
 	return EVisibility::Collapsed;
 }
 
+EVisibility FOptimusParameterBindingCustomization::IsSupportReadCheckBoxVisible(
+	const TArray<TWeakObjectPtr<UObject>>& InSelectedObjects, TSharedRef<IPropertyHandle> InPropertyHandle)
+{
+	for (TWeakObjectPtr<UObject> Object : InSelectedObjects)
+    	{
+    		if (IOptimusParameterBindingProvider* BindingProvider = Cast<IOptimusParameterBindingProvider>(Object))
+    		{
+    			TArray<const void *> RawData;
+    
+    			InPropertyHandle->AccessRawData(RawData);
+    			if (RawData.Num() > 0)
+    			{
+    				const FOptimusParameterBinding* Binding = static_cast<const FOptimusParameterBinding*>(RawData[0]);
+    				// During drag & reorder, we can have invalid bindings in the property
+    				if (Binding && Binding->Name != NAME_None)
+    				{
+    					return BindingProvider->GetBindingSupportReadCheckBoxVisibility(Binding->Name) ? EVisibility::Visible : EVisibility::Collapsed;
+    				}
+    			}
+    			break;
+    		}
+    	}
+    		
+    	return EVisibility::Collapsed;
+}
+
 FOptimusParameterBindingCustomization::FOptimusParameterBindingCustomization()
 {
 	
@@ -1416,9 +1423,19 @@ void FOptimusParameterBindingCustomization::CustomizeChildren(
 		return IsAtomicCheckBoxVisible(SelectedObjects, InPropertyHandle);
 	};
 	
+
 	TSharedPtr<IPropertyHandle> SupportAtomicHandle = InPropertyHandle->GetChildHandle(GET_MEMBER_NAME_CHECKED(FOptimusParameterBinding, bSupportAtomicIfCompatibleDataType), false);
 	IDetailPropertyRow& AtomicRow = InChildBuilder.AddProperty(SupportAtomicHandle.ToSharedRef());
 	AtomicRow.Visibility(TAttribute<EVisibility>::CreateLambda(IsAtomicRowVisible));
+
+	auto IsSupportReadRowVisible = [SelectedObjects, InPropertyHandle]()
+	{
+		return IsSupportReadCheckBoxVisible(SelectedObjects, InPropertyHandle);
+	};
+	
+	TSharedPtr<IPropertyHandle> SupportReadHandle = InPropertyHandle->GetChildHandle(GET_MEMBER_NAME_CHECKED(FOptimusParameterBinding, bSupportRead), false);
+	IDetailPropertyRow& SupportReadRow = InChildBuilder.AddProperty(SupportReadHandle.ToSharedRef());
+	SupportReadRow.Visibility(TAttribute<EVisibility>::CreateLambda(IsSupportReadRowVisible));
 }
 
 TSharedRef<FOptimusParameterBindingArrayBuilder> FOptimusParameterBindingArrayBuilder::MakeInstance(
diff --git a/Engine/Plugins/Animation/DeformerGraph/Source/OptimusEditor/Private/OptimusDetailsCustomization.h b/Engine/Plugins/Animation/DeformerGraph/Source/OptimusEditor/Private/OptimusDetailsCustomization.h
--- a/Engine/Plugins/Animation/DeformerGraph/Source/OptimusEditor/Private/OptimusDetailsCustomization.h
+++ b/Engine/Plugins/Animation/DeformerGraph/Source/OptimusEditor/Private/OptimusDetailsCustomization.h
@@ -238,6 +238,7 @@ public:
 	
 	static TSharedRef<IPropertyTypeCustomization> MakeInstance();
 	static EVisibility IsAtomicCheckBoxVisible(const TArray<TWeakObjectPtr<UObject>>& InSelectedObjects,  TSharedRef<IPropertyHandle> InPropertyHandle);
+	static EVisibility IsSupportReadCheckBoxVisible(const TArray<TWeakObjectPtr<UObject>>& InSelectedObjects,  TSharedRef<IPropertyHandle> InPropertyHandle);
 
 	FOptimusParameterBindingCustomization();
 
diff --git a/Engine/Plugins/Animation/DeformerGraph/Source/OptimusEditor/Private/OptimusEditorGraphNode.cpp b/Engine/Plugins/Animation/DeformerGraph/Source/OptimusEditor/Private/OptimusEditorGraphNode.cpp
--- a/Engine/Plugins/Animation/DeformerGraph/Source/OptimusEditor/Private/OptimusEditorGraphNode.cpp
+++ b/Engine/Plugins/Animation/DeformerGraph/Source/OptimusEditor/Private/OptimusEditorGraphNode.cpp
@@ -2,6 +2,7 @@
 
 #include "OptimusEditorGraphNode.h"
 
+#include "IOptimusComputeKernelProvider.h"
 #include "OptimusEditorHelpers.h"
 #include "OptimusEditorGraph.h"
 #include "OptimusEditorGraphSchema.h"
@@ -12,6 +13,8 @@
 #include "OptimusNode.h"
 #include "OptimusNodePin.h"
 #include "IOptimusNodeAdderPinProvider.h"
+#include "IOptimusPinMutabilityDefiner.h"
+#include "OptimusNodeGraph.h"
 
 #include "Framework/Commands/GenericCommands.h"
 #include "Logging/TokenizedMessage.h"
@@ -309,6 +312,39 @@ void UOptimusEditorGraphNode::GetNodeContextMenuActions(
 	}
 }
 
+FLinearColor UOptimusEditorGraphNode::GetNodeTitleColor() const
+{
+	const FLinearColor ImmutableColor = FLinearColor::Green;
+	if (IOptimusComputeKernelProvider* KernelProvider = Cast<IOptimusComputeKernelProvider>(ModelNode))
+	{
+		if (ModelNode->GetOwningGraph()->GetGraphType() == EOptimusNodeGraphType::Update)
+		{
+			if (!KernelProvider->HasMutableInput())
+			{
+				return ImmutableColor;
+			}	
+		}
+	}
+
+	if (IOptimusPinMutabilityDefiner* PinMutabilityDefiner = Cast<IOptimusPinMutabilityDefiner>(ModelNode))
+	{
+		for (UOptimusNodePin* Pin : ModelNode->GetPins())
+		{
+			if (Pin->GetDirection() == EOptimusNodePinDirection::Output)
+			{
+				if (PinMutabilityDefiner->IsOutputPinMutable(Pin))
+				{
+					return Super::GetNodeTitleColor();
+				}
+			}
+		}
+
+		return ImmutableColor;
+	}
+	
+	return Super::GetNodeTitleColor();
+}
+
 
 bool UOptimusEditorGraphNode::CreateGraphPinFromModelPin(
 	const UOptimusNodePin* InModelPin,
diff --git a/Engine/Plugins/Animation/DeformerGraph/Source/OptimusEditor/Private/OptimusEditorGraphNode.h b/Engine/Plugins/Animation/DeformerGraph/Source/OptimusEditor/Private/OptimusEditorGraphNode.h
--- a/Engine/Plugins/Animation/DeformerGraph/Source/OptimusEditor/Private/OptimusEditorGraphNode.h
+++ b/Engine/Plugins/Animation/DeformerGraph/Source/OptimusEditor/Private/OptimusEditorGraphNode.h
@@ -47,7 +47,7 @@ public:
 	bool CanUserDeleteNode() const override;
 	FText GetNodeTitle(ENodeTitleType::Type InTitleType) const override;
 	void GetNodeContextMenuActions(UToolMenu* InMenu, UGraphNodeContextMenuContext* InContext) const override;
-
+	FLinearColor GetNodeTitleColor() const override;
 	
 	// FIXME: Move to private and add accessor function.
 	UPROPERTY()
diff --git a/Engine/Plugins/Animation/DeformerGraph/Source/OptimusEditor/Private/Widgets/SOptimusEditorGraphNode.cpp b/Engine/Plugins/Animation/DeformerGraph/Source/OptimusEditor/Private/Widgets/SOptimusEditorGraphNode.cpp
--- a/Engine/Plugins/Animation/DeformerGraph/Source/OptimusEditor/Private/Widgets/SOptimusEditorGraphNode.cpp
+++ b/Engine/Plugins/Animation/DeformerGraph/Source/OptimusEditor/Private/Widgets/SOptimusEditorGraphNode.cpp
@@ -186,6 +186,7 @@ public:
 		AddEntry(LOCTEXT("PinType", "Data Type"), {this, &SOptimusEditorGraphPinToolTipWidget::GetPinDataType});
 		AddEntry(LOCTEXT("PinDomain", "Data Domain"), {this, &SOptimusEditorGraphPinToolTipWidget::GetPinDataDomain});
 		AddEntry(LOCTEXT("PinComponentSource", "Component Source"), {this, &SOptimusEditorGraphPinToolTipWidget::GetPinComponentSource});
+		AddEntry(LOCTEXT("PinMutability", "Mutability"), {this, &SOptimusEditorGraphPinToolTipWidget::GetPinMutability});
 
 		ChildSlot
 		[
@@ -294,6 +295,28 @@ private:
 		}
 		return FText::GetEmpty();
 	}
+
+	FText GetPinMutability() const
+	{
+		if (TObjectPtr<UOptimusNodePin> ModelPin = ModelPinPtr.Get())
+		{
+			TSet<UOptimusComponentSourceBinding*> ComponentSourceBindings;
+
+			// Make sure it doesn't belong to a node that's just been deleted, since Slate updates are usually a frame
+			// behind.
+			if (ModelPin->GetPackage() != GetTransientPackage())
+			{
+				if (ModelPin->IsMutable())
+				{
+					return FText::FromName(TEXT("Mutable"));
+				}
+
+				return FText::FromName(TEXT("Immutable"));
+			}
+		}
+		return FText::GetEmpty();
+	}
+	
 	
 	TWeakObjectPtr<UOptimusNodePin> ModelPinPtr;
 };
