Optimus: Ensure data type objects stored in the type registry aren't garbage collected.

#rnx

diff --git a/Engine/Plugins/Animation/DeformerGraph/Source/OptimusCore/Private/OptimusDataTypeRegistry.cpp b/Engine/Plugins/Animation/DeformerGraph/Source/OptimusCore/Private/OptimusDataTypeRegistry.cpp
--- a/Engine/Plugins/Animation/DeformerGraph/Source/OptimusCore/Private/OptimusDataTypeRegistry.cpp
+++ b/Engine/Plugins/Animation/DeformerGraph/Source/OptimusCore/Private/OptimusDataTypeRegistry.cpp
@@ -1240,7 +1240,7 @@ TArray<FOptimusDataTypeHandle> FOptimusDataTypeRegistry::GetAllTypes() const
 	TArray<FOptimusDataTypeHandle> Result;
 	for (const FName& TypeName : RegistrationOrder)
 	{
-		Result.Add(RegisteredTypes[TypeName].Handle);
+		Result.Add(RegisteredTypes[TypeName].DataType);
 	}
 	return Result;
 }
@@ -1298,7 +1298,7 @@ FOptimusDataTypeHandle FOptimusDataTypeRegistry::FindType(const UClass& InClassT
 FOptimusDataTypeHandle FOptimusDataTypeRegistry::FindType(FName InTypeName) const
 {
 	const FTypeInfo* InfoPtr = RegisteredTypes.Find(InTypeName);
-	return InfoPtr ? InfoPtr->Handle : FOptimusDataTypeHandle();
+	return InfoPtr ? InfoPtr->DataType : FOptimusDataTypeHandle();
 }
 
 
@@ -1306,7 +1306,7 @@ FOptimusDataTypeHandle FOptimusDataTypeRegistry::FindType(FShaderValueTypeHandle
 {
 	for (const FName& TypeName : RegistrationOrder)
 	{
-		const FOptimusDataTypeHandle Handle = RegisteredTypes[TypeName].Handle;
+		const FOptimusDataTypeHandle Handle = RegisteredTypes[TypeName].DataType;
 		if (Handle->ShaderValueType == InValueType)
 		{
 			return Handle;
@@ -1478,9 +1478,9 @@ UScriptStruct* FOptimusDataTypeRegistry::FindAttributeType(FName InTypeName) con
 		return nullptr;
 	}
 
-	if (ensure(EnumHasAnyFlags(InfoPtr->Handle->UsageFlags, EOptimusDataTypeUsageFlags::AnimAttributes)))
+	if (ensure(EnumHasAnyFlags(InfoPtr->DataType->UsageFlags, EOptimusDataTypeUsageFlags::AnimAttributes)))
 	{
-		return Cast<UScriptStruct>(InfoPtr->Handle->TypeObject);
+		return Cast<UScriptStruct>(InfoPtr->DataType->TypeObject);
 	}
 
 	return nullptr;
@@ -1491,6 +1491,20 @@ FOptimusDataTypeRegistry::FOnDataTypeChanged& FOptimusDataTypeRegistry::GetOnDat
 	return OnDataTypeChanged;
 }
 
+void FOptimusDataTypeRegistry::AddReferencedObjects(FReferenceCollector& Collector)
+{
+	for (const TTuple<FName, FTypeInfo>& TypeItem: RegisteredTypes)
+	{
+		TSharedPtr<FOptimusDataType> DataType = TypeItem.Value.DataType;
+
+		if (DataType->TypeObject.IsValid())
+		{
+			Collector.AddReferencedObject(DataType->TypeObject);
+		}
+		
+	}
+}
+
 void FOptimusDataTypeRegistry::UnregisterType(FName InTypeName)
 {
 	RegisteredTypes.Remove(InTypeName);
diff --git a/Engine/Plugins/Animation/DeformerGraph/Source/OptimusCore/Public/OptimusDataTypeRegistry.h b/Engine/Plugins/Animation/DeformerGraph/Source/OptimusCore/Public/OptimusDataTypeRegistry.h
--- a/Engine/Plugins/Animation/DeformerGraph/Source/OptimusCore/Public/OptimusDataTypeRegistry.h
+++ b/Engine/Plugins/Animation/DeformerGraph/Source/OptimusCore/Public/OptimusDataTypeRegistry.h
@@ -7,13 +7,17 @@
 #include "Containers/Array.h"
 #include "Containers/Map.h"
 #include "Misc/Optional.h"
+#include "UObject/GCObject.h"
+
 
 class FFieldClass;
 class UScriptStruct;
 struct FShaderValueTypeHandle;
 class UUserDefinedStruct;
 
-class FOptimusDataTypeRegistry
+
+class FOptimusDataTypeRegistry :
+	public FGCObject
 {
 public:
 	DECLARE_EVENT_OneParam(FOptimusDataTypeRegistry, FOnDataTypeChanged, FName /* InTypeName */);
@@ -33,7 +37,7 @@ public:
 		int32 ShaderValueOffset;
 	};
 	
-	~FOptimusDataTypeRegistry();
+	virtual ~FOptimusDataTypeRegistry() override;
 
 	/** Get the singleton registry object */
 	OPTIMUSCORE_API static FOptimusDataTypeRegistry &Get();
@@ -172,6 +176,10 @@ public:
 
 	OPTIMUSCORE_API	FOnDataTypeChanged& GetOnDataTypeChanged();
 
+	// -- FGCObject overrides
+	FString GetReferencerName() const override { return TEXT("FOptimusDataTypeRegistry"); }
+	void AddReferencedObjects(FReferenceCollector& InCollector) override;
+	
 protected:
 	friend class FOptimusCoreModule;
 	friend struct FOptimusDataType;
@@ -219,7 +227,7 @@ private:
 
 	struct FTypeInfo
 	{
-		FOptimusDataTypeHandle Handle;
+		TSharedPtr<FOptimusDataType> DataType;
 		PropertyCreateFuncT PropertyCreateFunc;
 		PropertyValueConvertFuncT PropertyValueConvertFunc;
 		TArray<FArrayMetadata> ArrayMetadata;
