Optimus: Default to primary binding if binding pins are not connected.

#rnx

diff --git a/Engine/Plugins/Animation/DeformerGraph/Source/OptimusCore/Private/Nodes/OptimusNode_DataInterface.cpp b/Engine/Plugins/Animation/DeformerGraph/Source/OptimusCore/Private/Nodes/OptimusNode_DataInterface.cpp
--- a/Engine/Plugins/Animation/DeformerGraph/Source/OptimusCore/Private/Nodes/OptimusNode_DataInterface.cpp
+++ b/Engine/Plugins/Animation/DeformerGraph/Source/OptimusCore/Private/Nodes/OptimusNode_DataInterface.cpp
@@ -5,6 +5,7 @@
 #include "OptimusCoreModule.h"
 #include "OptimusNodePin.h"
 #include "OptimusDataTypeRegistry.h"
+#include "OptimusDeformer.h"
 #include "OptimusNodeGraph.h"
 #include "OptimusNode_ComponentSource.h"
 #include "OptimusObjectVersion.h"
@@ -100,7 +101,7 @@ bool UOptimusNode_DataInterface::IsComponentSourceCompatible(const UOptimusCompo
 
 void UOptimusNode_DataInterface::Serialize(FArchive& Ar)
 {
-	Super::Serialize(Ar);
+ 	Super::Serialize(Ar);
 	Ar.UsingCustomVersion(FOptimusObjectVersion::GUID);
 }
 
@@ -186,6 +187,12 @@ UOptimusComponentSourceBinding* UOptimusNode_DataInterface::GetComponentBinding(
 	{
 		return Bindings.Array()[0];
 	}
+
+	// Default to the primary binding, but only if we're at the top-most level of the graph.
+	if (const UOptimusDeformer* Deformer = Cast<UOptimusDeformer>(Graph->GetCollectionOwner()))
+	{
+		return Deformer->GetPrimaryComponentBinding();
+	}
 	
 	return nullptr;
 }
diff --git a/Engine/Plugins/Animation/DeformerGraph/Source/OptimusCore/Private/Nodes/OptimusNode_DataInterface.h b/Engine/Plugins/Animation/DeformerGraph/Source/OptimusCore/Private/Nodes/OptimusNode_DataInterface.h
--- a/Engine/Plugins/Animation/DeformerGraph/Source/OptimusCore/Private/Nodes/OptimusNode_DataInterface.h
+++ b/Engine/Plugins/Animation/DeformerGraph/Source/OptimusCore/Private/Nodes/OptimusNode_DataInterface.h
@@ -2,6 +2,7 @@
 
 #pragma once
 
+#include "IOptimusComponentBindingProvider.h"
 #include "IOptimusDataInterfaceProvider.h"
 #include "IOptimusPinMutabilityDefiner.h"
 #include "OptimusComputeDataInterface.h"
@@ -19,6 +20,7 @@ UCLASS(Hidden)
 class UOptimusNode_DataInterface :
 	public UOptimusNode,
 	public IOptimusDataInterfaceProvider,
+	public IOptimusComponentBindingProvider,
 	public IOptimusPinMutabilityDefiner
 {
 	GENERATED_BODY()
@@ -41,6 +43,8 @@ public:
 	// -- IOptimusDataInterfaceProvider implementations
 	UOptimusComputeDataInterface *GetDataInterface(UObject *InOuter) const override;
 	int32 GetDataFunctionIndexFromPin(const UOptimusNodePin* InPin) const override;
+	
+	// -- IOptimusDataInterfaceProvider & IOptimusComponentBindingProvider
 	UOptimusComponentSourceBinding* GetComponentBinding() const override;
 
 	// -- IOptimusPinMutabilityDefiner
diff --git a/Engine/Plugins/Animation/DeformerGraph/Source/OptimusCore/Private/OptimusDeformer.cpp b/Engine/Plugins/Animation/DeformerGraph/Source/OptimusCore/Private/OptimusDeformer.cpp
--- a/Engine/Plugins/Animation/DeformerGraph/Source/OptimusCore/Private/OptimusDeformer.cpp
+++ b/Engine/Plugins/Animation/DeformerGraph/Source/OptimusCore/Private/OptimusDeformer.cpp
@@ -1309,6 +1309,19 @@ void UOptimusDeformer::SetStatusFromDiagnostic(EOptimusDiagnosticLevel InDiagnos
 }
 
 
+UOptimusComponentSourceBinding* UOptimusDeformer::GetPrimaryComponentBinding() const
+{
+	for (UOptimusComponentSourceBinding* Binding : GetComponentBindings())
+	{
+		if (Binding->bIsPrimaryBinding)
+		{
+			return Binding;
+		}
+	}
+	return nullptr;
+}
+
+
 UOptimusComponentSourceBinding* UOptimusDeformer::ResolveComponentBinding(
 	FName InBindingName
 	) const
diff --git a/Engine/Plugins/Animation/DeformerGraph/Source/OptimusCore/Public/OptimusDeformer.h b/Engine/Plugins/Animation/DeformerGraph/Source/OptimusCore/Public/OptimusDeformer.h
--- a/Engine/Plugins/Animation/DeformerGraph/Source/OptimusCore/Public/OptimusDeformer.h
+++ b/Engine/Plugins/Animation/DeformerGraph/Source/OptimusCore/Public/OptimusDeformer.h
@@ -259,6 +259,9 @@ public:
 	UFUNCTION(BlueprintGetter)
 	const TArray<UOptimusComponentSourceBinding*>& GetComponentBindings() const { return Bindings->Bindings; }
 
+	UFUNCTION(BlueprintGetter)
+	UOptimusComponentSourceBinding* GetPrimaryComponentBinding() const;
+
 	UOptimusComponentSourceBinding* ResolveComponentBinding(
 		FName InBindingName
 		) const override;
