Fix missing extra influences on Skel. mesh deformer interface.

diff --git a/Engine/Plugins/Animation/DeformerGraph/Shaders/Private/DataInterfaceSkeleton.ush b/Engine/Plugins/Animation/DeformerGraph/Shaders/Private/DataInterfaceSkeleton.ush
--- a/Engine/Plugins/Animation/DeformerGraph/Shaders/Private/DataInterfaceSkeleton.ush
+++ b/Engine/Plugins/Animation/DeformerGraph/Shaders/Private/DataInterfaceSkeleton.ush
@@ -127,13 +127,20 @@ float3x4 ReadWeightedBoneMatrix_{DataInterfaceName}(uint VertexIndex)
 #else
 
 	const uint StreamOffset = {DataInterfaceName}_InputWeightStart + (VertexIndex * ({DataInterfaceName}_InputWeightStride / 4));
+	const uint NumInfluences = {DataInterfaceName}_NumBoneInfluences;
 
 	FGPUSkinIndexAndWeight IndicesAndWeights = (FGPUSkinIndexAndWeight)0;
-	IndicesAndWeights.BlendIndices  = GetBlendIndices_{DataInterfaceName}(StreamOffset, 0);
-	IndicesAndWeights.BlendIndices2 = GetBlendIndices_{DataInterfaceName}(StreamOffset, 1);
-	IndicesAndWeights.BlendWeights  = GetBlendWeights_{DataInterfaceName}(StreamOffset, 0);
-	IndicesAndWeights.BlendWeights2 = GetBlendWeights_{DataInterfaceName}(StreamOffset, 1);
-	return ComputeBoneMatrixWithLimitedInfluences({DataInterfaceName}_BoneMatrices, IndicesAndWeights);
+	if (NumInfluences > 0)
+	{
+		IndicesAndWeights.BlendIndices  = GetBlendIndices_{DataInterfaceName}(StreamOffset, 0);
+		IndicesAndWeights.BlendWeights  = GetBlendWeights_{DataInterfaceName}(StreamOffset, 0);
+	}
+	if (NumInfluences > 4)
+	{
+		IndicesAndWeights.BlendIndices2 = GetBlendIndices_{DataInterfaceName}(StreamOffset, 1);
+		IndicesAndWeights.BlendWeights2 = GetBlendWeights_{DataInterfaceName}(StreamOffset, 1);
+	}
+	return ComputeBoneMatrixWithLimitedInfluences({DataInterfaceName}_BoneMatrices, IndicesAndWeights, NumInfluences > 4);
 
 #endif
 }
