Optimus: Added support for loop terminal

+fix pin reorder issue with custom compute kernel
+fix new data interface node losing data when spawn via duplicate
+deprecating OptimusConstant
+temp fix for setup graph not running after compile

diff --git a/Engine/Plugins/Animation/DeformerGraph/Shaders/Private/DataInterfaceLoopTerminal.ush b/Engine/Plugins/Animation/DeformerGraph/Shaders/Private/DataInterfaceLoopTerminal.ush
new file mode 100644
--- /dev/null
+++ b/Engine/Plugins/Animation/DeformerGraph/Shaders/Private/DataInterfaceLoopTerminal.ush
@@ -0,0 +1,15 @@
+// Copyright Epic Games, Inc. All Rights Reserved.
+
+uint {DataInterfaceName}_Index;
+uint {DataInterfaceName}_Count;
+
+uint ReadIndex_{DataInterfaceName}()
+{
+	return {DataInterfaceName}_Index;
+}
+
+uint ReadCount_{DataInterfaceName}()
+{
+	return {DataInterfaceName}_Count;
+}
+
diff --git a/Engine/Plugins/Animation/DeformerGraph/Source/OptimusCore/Private/Actions/OptimusAction.cpp b/Engine/Plugins/Animation/DeformerGraph/Source/OptimusCore/Private/Actions/OptimusAction.cpp
--- a/Engine/Plugins/Animation/DeformerGraph/Source/OptimusCore/Private/Actions/OptimusAction.cpp
+++ b/Engine/Plugins/Animation/DeformerGraph/Source/OptimusCore/Private/Actions/OptimusAction.cpp
@@ -19,9 +19,12 @@ FOptimusAction::~FOptimusAction()
 
 
 
-void FOptimusCompoundAction::AddSubAction(FOptimusAction* InAction)
+TWeakPtr<FOptimusAction> FOptimusCompoundAction::AddSubAction(FOptimusAction* InAction)
 {
 	AddSubAction(TSharedPtr<FOptimusAction>(InAction));
+
+	return SubActions.Last().ToWeakPtr();
+	
 }
 
 
diff --git a/Engine/Plugins/Animation/DeformerGraph/Source/OptimusCore/Private/Actions/OptimusAction.h b/Engine/Plugins/Animation/DeformerGraph/Source/OptimusCore/Private/Actions/OptimusAction.h
--- a/Engine/Plugins/Animation/DeformerGraph/Source/OptimusCore/Private/Actions/OptimusAction.h
+++ b/Engine/Plugins/Animation/DeformerGraph/Source/OptimusCore/Private/Actions/OptimusAction.h
@@ -85,13 +85,16 @@ struct FOptimusCompoundAction :
 
 	/// Add a sub-action from a heap-constructed action. This takes ownership of the pointer.
 	/// @param InAction The action to add.
-	void AddSubAction(FOptimusAction* InAction);
+	/// Returns a weakptr for subsequent actions reference the result of this action
+	TWeakPtr<FOptimusAction> AddSubAction(FOptimusAction* InAction);
 
 	template<typename T, typename... ArgTypes>
-	typename TEnableIf<TPointerIsConvertibleFromTo<T, FOptimusAction>::Value, void>::Type 
+	typename TEnableIf<TPointerIsConvertibleFromTo<T, FOptimusAction>::Value, TWeakPtr<T>>::Type 
 	AddSubAction(ArgTypes&& ...Args)
 	{
-		SubActions.Add(MakeShared<T>(Forward<ArgTypes>(Args)...));
+		TSharedPtr<T> Action = MakeShared<T>(Forward<ArgTypes>(Args)...);
+		SubActions.Add(Action);
+		return Action.ToWeakPtr();
 	}
 
 	bool HasSubActions() const
diff --git a/Engine/Plugins/Animation/DeformerGraph/Source/OptimusCore/Private/Actions/OptimusNodeGraphActions.cpp b/Engine/Plugins/Animation/DeformerGraph/Source/OptimusCore/Private/Actions/OptimusNodeGraphActions.cpp
--- a/Engine/Plugins/Animation/DeformerGraph/Source/OptimusCore/Private/Actions/OptimusNodeGraphActions.cpp
+++ b/Engine/Plugins/Animation/DeformerGraph/Source/OptimusCore/Private/Actions/OptimusNodeGraphActions.cpp
@@ -3,6 +3,7 @@
 #include "OptimusNodeGraphActions.h"
 
 #include "IOptimusPathResolver.h"
+#include "IOptimusUnnamedNodePinProvider.h"
 #include "Nodes/OptimusNode_ComputeKernelFunction.h"
 #include "Nodes/OptimusNode_CustomComputeKernel.h"
 #include "Nodes/OptimusNode_ConstantValue.h"
@@ -10,13 +11,12 @@
 #include "OptimusNode.h"
 #include "OptimusNodeGraph.h"
 #include "OptimusNodeLink.h"
+#include "OptimusNodePair.h"
 #include "OptimusNodePin.h"
 
 
 #include "Serialization/MemoryReader.h"
 #include "Serialization/MemoryWriter.h"
-#include "Serialization/ObjectAndNameAsStringProxyArchive.h"
-#include "Serialization/ObjectWriter.h"
 #include "UObject/UObjectGlobals.h"
 
 #include UE_INLINE_GENERATED_CPP_BY_NAME(OptimusNodeGraphActions)
@@ -450,14 +450,8 @@ bool FOptimusNodeGraphAction_RemoveNode::Do(
 	}
 
 	// Take a copy of the node's contents but not sub-data (like pins).
-	{
-		FMemoryWriter NodeArchive(NodeData);
-		// This fella does the heavy lifting of serializing object references. 
-		// FMemoryWriter and fam do not handle UObject* serialization on their own.
-		FObjectAndNameAsStringProxyArchive NodeProxyArchive(
-				NodeArchive, /* bInLoadIfFindFails=*/ false);
-		Node->SerializeScriptProperties(NodeProxyArchive);
-	}
+	FMemoryWriter NodeArchive(NodeData);
+	Node->SaveState(NodeArchive);
 
 	return Graph->RemoveNodeDirect(Node);
 }
@@ -478,21 +472,103 @@ bool FOptimusNodeGraphAction_RemoveNode::Undo(
 		return false;
 	}
 
-	UOptimusNode* Node = NewObject<UOptimusNode>(Graph, NodeClass, NodeName);
-
+	UOptimusNode* Node = Graph->CreateNodeDirect(NodeClass, NodeName, [this](UOptimusNode* InNode)
 	{
 		FMemoryReader NodeArchive(NodeData);
-		FObjectAndNameAsStringProxyArchive NodeProxyArchive(
-				NodeArchive, /* bInLoadIfFindFails=*/true);
-		Node->SerializeScriptProperties(NodeProxyArchive);
+		InNode->RestoreState(NodeArchive);
+		return true;
+	});
+
+	if (!Node)
+	{
+		return false;
+	}
+	
+	return true;
+}
+
+
+FOptimusNodeGraphAction_AddRemoveNodePair::FOptimusNodeGraphAction_AddRemoveNodePair(const FString& InFirstNodePath, const FString& InSecondNodePath)
+{
+	FirstNodePath = InFirstNodePath;
+	SecondNodePath = InSecondNodePath;
+}
+
+bool FOptimusNodeGraphAction_AddRemoveNodePair::AddNodePair(IOptimusPathResolver* InRoot)
+{
+	UOptimusNode* FirstNode = InRoot->ResolveNodePath(FirstNodePath);
+	UOptimusNode* SecondNode = InRoot->ResolveNodePath(SecondNodePath);
+
+	UOptimusNodeGraph* Graph = FirstNode->GetOwningGraph();
+	return Graph->AddNodePairDirect(FirstNode, SecondNode);
+}
+
+bool FOptimusNodeGraphAction_AddRemoveNodePair::RemoveNodePair(IOptimusPathResolver* InRoot)
+{
+	UOptimusNode* FirstNode = InRoot->ResolveNodePath(FirstNodePath);
+	UOptimusNode* SecondNode = InRoot->ResolveNodePath(SecondNodePath);
+
+	UOptimusNodeGraph* Graph = FirstNode->GetOwningGraph();
+	return Graph->RemoveNodePairDirect(FirstNode, SecondNode);
+}
+
+FOptimusNodeGraphAction_AddNodePair::FOptimusNodeGraphAction_AddNodePair(const FString& InFirstNodePath, const FString& InSecondNodePath) :
+	FOptimusNodeGraphAction_AddRemoveNodePair(InFirstNodePath, InSecondNodePath)
+{
+	SetTitlef(TEXT("Add Node Pair"));
+}
+
+FOptimusNodeGraphAction_AddNodePair::FOptimusNodeGraphAction_AddNodePair(TWeakPtr<FOptimusNodeGraphAction_AddNode> InAddFirstNodeAction, TWeakPtr<FOptimusNodeGraphAction_AddNode> InAddSecondNodeAction)
+{
+	AddFirstNodeAction = InAddFirstNodeAction;
+	AddSecondNodeAction = InAddSecondNodeAction;
+
+	SetTitlef(TEXT("Add Node Pair"));
+}
+
+bool FOptimusNodeGraphAction_AddNodePair::Do(IOptimusPathResolver* InRoot)
+{
+	if (!FirstNodePath.IsEmpty() && !SecondNodePath.IsEmpty())
+	{
+		return AddNodePair(InRoot);	
+	}
+
+	// In case the node paths are not assigned, try to extract them from previous add node actions
+	if (!AddFirstNodeAction.IsValid() || !AddSecondNodeAction.IsValid())
+	{
+		return false;
+	}
+
+	const UOptimusNode* FirstNode = AddFirstNodeAction.Pin()->GetNode(InRoot);
+	const UOptimusNode* SecondNode = AddSecondNodeAction.Pin()->GetNode(InRoot);
+
+	if (!FirstNode || !SecondNode)
+	{
+		return false;	
 	}
 
-	// Create the pins.
-	Node->PostCreateNode();
+	FirstNodePath = FirstNode->GetNodePath();
+	SecondNodePath = SecondNode->GetNodePath();
+	
+	return AddNodePair(InRoot);	
+}
 
-	return Graph->AddNodeDirect(Node);
+bool FOptimusNodeGraphAction_AddNodePair::Undo(IOptimusPathResolver* InRoot)
+{
+	// Upon undo, the node paths should have been assigned during Do()
+	if (ensure(!FirstNodePath.IsEmpty() && !SecondNodePath.IsEmpty()))
+	{
+		return RemoveNodePair(InRoot);	
+	}
+
+	return false;
 }
 
+FOptimusNodeGraphAction_RemoveNodePair::FOptimusNodeGraphAction_RemoveNodePair(const UOptimusNodePair* InNodePair) :
+	FOptimusNodeGraphAction_AddRemoveNodePair(InNodePair->GetFirst()->GetNodePath(), InNodePair->GetSecond()->GetNodePath())
+{
+	SetTitlef(TEXT("Remove Node Pair"));
+}
 
 // ---- Add/remove link base
 
@@ -656,7 +732,13 @@ bool FOptimusNodeGraphAction_ConnectAdderPin::Do(IOptimusPathResolver* InRoot)
 
 	UOptimusNodePin* SourcePin = InRoot->ResolvePinPath(SourcePinPath);
 
-	TArray<UOptimusNodePin*> AddedPins = AdderPinProvider->TryAddPinFromPin(Action, SourcePin);
+	FName Name = SourcePin->GetFName();
+	if (IOptimusUnnamedNodePinProvider* UnnamedPinProvider = Cast<IOptimusUnnamedNodePinProvider>(SourcePin->GetOwningNode()))
+	{
+		Name = UnnamedPinProvider->GetNameForAdderPin(SourcePin);
+	}
+	
+	TArray<UOptimusNodePin*> AddedPins = AdderPinProvider->TryAddPinFromPin(Action, SourcePin, Name);
 
 	if (AddedPins.Num() == 0)
 	{
diff --git a/Engine/Plugins/Animation/DeformerGraph/Source/OptimusCore/Private/Actions/OptimusNodeGraphActions.h b/Engine/Plugins/Animation/DeformerGraph/Source/OptimusCore/Private/Actions/OptimusNodeGraphActions.h
--- a/Engine/Plugins/Animation/DeformerGraph/Source/OptimusCore/Private/Actions/OptimusNodeGraphActions.h
+++ b/Engine/Plugins/Animation/DeformerGraph/Source/OptimusCore/Private/Actions/OptimusNodeGraphActions.h
@@ -18,6 +18,7 @@ class UOptimusNode_CustomComputeKernel;
 enum class EOptimusNodeGraphType;
 class IOptimusNodeGraphCollectionOwner;
 class UOptimusNode;
+class UOptimusNodePair;
 class UOptimusNodeGraph;
 class UOptimusNodeLink;
 class UOptimusNodePin;
@@ -259,6 +260,89 @@ private:
 	TArray<uint8> NodeData;
 };
 
+// A base action for adding/removing node pairs.
+USTRUCT()
+struct FOptimusNodeGraphAction_AddRemoveNodePair:
+	public FOptimusAction
+{
+	GENERATED_BODY()
+
+	FOptimusNodeGraphAction_AddRemoveNodePair() = default;
+
+	FOptimusNodeGraphAction_AddRemoveNodePair(
+		const FString& InFirstNodePath,
+		const FString& InSecondNodePath
+		);
+
+protected:
+	bool AddNodePair(IOptimusPathResolver* InRoot);
+	bool RemoveNodePair(IOptimusPathResolver* InRoot);
+
+	FString FirstNodePath;
+
+	FString SecondNodePath;
+};
+
+// Mark FOptimusNodeGraphAction_AddRemoveNodePair as pure virtual, so that the UObject machinery
+// won't attempt to instantiate it.
+template<>
+struct TStructOpsTypeTraits<FOptimusNodeGraphAction_AddRemoveNodePair> :
+	TStructOpsTypeTraitsBase2<FOptimusNodeGraphAction_AddRemoveNodePair>
+{
+	enum
+	{
+		WithPureVirtual = true,
+    };
+};
+
+
+
+USTRUCT()
+struct FOptimusNodeGraphAction_AddNodePair :
+	public FOptimusNodeGraphAction_AddRemoveNodePair
+{
+	GENERATED_BODY()
+
+	FOptimusNodeGraphAction_AddNodePair() = default;
+
+	// Use when the node paths are unknown after add node actions are performed
+	FOptimusNodeGraphAction_AddNodePair(
+		TWeakPtr<FOptimusNodeGraphAction_AddNode> InAddFirstNodeAction,
+		TWeakPtr<FOptimusNodeGraphAction_AddNode> InAddSecondNodeAction
+		);
+	
+	// Use when the node paths are available already
+	FOptimusNodeGraphAction_AddNodePair(
+		const FString& InFirstNodePath,
+		const FString& InSecondNodePath
+		);
+
+protected:
+	bool Do(IOptimusPathResolver* InRoot) override; 
+	bool Undo(IOptimusPathResolver* InRoot) override;
+
+	
+	TWeakPtr<FOptimusNodeGraphAction_AddNode> AddFirstNodeAction;
+	TWeakPtr<FOptimusNodeGraphAction_AddNode> AddSecondNodeAction;
+};
+
+USTRUCT()
+struct FOptimusNodeGraphAction_RemoveNodePair :
+	public FOptimusNodeGraphAction_AddRemoveNodePair
+{
+	GENERATED_BODY()
+
+	FOptimusNodeGraphAction_RemoveNodePair() = default;
+
+	FOptimusNodeGraphAction_RemoveNodePair(
+		const UOptimusNodePair* InNodePair
+	);
+
+protected:
+	bool Do(IOptimusPathResolver* InRoot) override { return RemoveNodePair(InRoot); }
+	bool Undo(IOptimusPathResolver* InRoot) override { return AddNodePair(InRoot); }
+};
+
 
 // A base action for adding/removing nodes.
 USTRUCT()
diff --git a/Engine/Plugins/Animation/DeformerGraph/Source/OptimusCore/Private/DataInterfaces/OptimusDataInterfaceCustomComputeKernel.cpp b/Engine/Plugins/Animation/DeformerGraph/Source/OptimusCore/Private/DataInterfaces/OptimusDataInterfaceCustomComputeKernel.cpp
--- a/Engine/Plugins/Animation/DeformerGraph/Source/OptimusCore/Private/DataInterfaces/OptimusDataInterfaceCustomComputeKernel.cpp
+++ b/Engine/Plugins/Animation/DeformerGraph/Source/OptimusCore/Private/DataInterfaces/OptimusDataInterfaceCustomComputeKernel.cpp
@@ -77,10 +77,14 @@ UComputeDataProvider* UOptimusCustomComputeKernelDataInterface::CreateDataProvid
 	return Provider;
 }
 
-void UOptimusCustomComputeKernelDataInterface::SetExecutionDomainConstant(
-	const FOptimusConstantIdentifier& InExecutionDomainConstantIdentifier)
+void UOptimusCustomComputeKernelDataInterface::SetExecutionDomain(const FString& InExecutionDomain)
 {
-	ExecutionDomainConstantIdentifier = InExecutionDomainConstantIdentifier;
+	NumThreadsExpression = InExecutionDomain;
+}
+
+void UOptimusCustomComputeKernelDataInterface::SetComponentBinding(const UOptimusComponentSourceBinding* InBinding)
+{
+	ComponentSourceBinding = InBinding;
 }
 
 void UOptimusCustomComputeKernelDataProvider::InitFromDataInterface(const UOptimusCustomComputeKernelDataInterface* InDataInterface, const UObject* InBinding)
@@ -103,55 +107,7 @@ FComputeDataProviderRenderProxy* UOptimusCustomComputeKernelDataProvider::GetRen
 	return Proxy;
 }
 
-void UOptimusCustomComputeKernelDataProvider::SetDeformerInstance(UOptimusDeformerInstance* InInstance)
-{
-	DeformerInstance = InInstance;
-}
-
-UOptimusDeformerInstance* UOptimusCustomComputeKernelDataProvider::GetDeformerInstance() const
-{
-	return DeformerInstance;
-}
-
-bool UOptimusCustomComputeKernelDataProvider::GetInvocationThreadCounts(TArray<int32>& OutInvocationThreadCount,
-                                                                        int32& OutTotalThreadCount) const
-{
-	if (!WeakDataInterface.IsValid())
-	{
-		return false;
-	}
-
-	const FOptimusConstantIdentifier& ExecutionDomainIdentifier = WeakDataInterface->ExecutionDomainConstantIdentifier;
-	
-	if (!ExecutionDomainIdentifier.IsValid())
-	{
-		return GetInvocationThreadCounts_DEPRECATED(OutInvocationThreadCount, OutTotalThreadCount);
-	}
-
-	const FOptimusConstantEvaluationResult Result = DeformerInstance->GetConstantValuePerInvocation(ExecutionDomainIdentifier);
-
-	const TArray<float>& Values = Result.ValuePerInvocation;
-	
-	// Can happen if the bound component does not have actual data, like when there is no preview mesh
-	if (Values.Num() == 0)
-	{
-		return false;
-	}
-
-
-	OutTotalThreadCount = 0;
-	OutInvocationThreadCount.Reset(Values.Num());
-	for (const float& Value : Values)
-	{
-		int32 Count = static_cast<int32>(Value);
-		OutInvocationThreadCount.Add(Count);
-		OutTotalThreadCount += Count;
-	}
-
-	return true;
-}
-
-bool UOptimusCustomComputeKernelDataProvider::GetInvocationThreadCounts_DEPRECATED(
+bool UOptimusCustomComputeKernelDataProvider::GetInvocationThreadCounts(
 	TArray<int32>& OutInvocationThreadCount,
 	int32& OutTotalThreadCount
 	) const
@@ -161,15 +117,13 @@ bool UOptimusCustomComputeKernelDataProvider::GetInvocationThreadCounts_DEPRECAT
 		return false;
 	}
 
-PRAGMA_DISABLE_DEPRECATION_WARNINGS
-	if (!WeakDataInterface->ComponentSourceBinding_DEPRECATED.IsValid())
+	if (!WeakDataInterface->ComponentSourceBinding.IsValid())
 	{
 		return false;
 	}
 	
-	const UOptimusComponentSource* ComponentSource = WeakDataInterface->ComponentSourceBinding_DEPRECATED->GetComponentSource();
-	const FString& NumThreadsExpression = WeakDataInterface->NumThreadsExpression_DEPRECATED;
-PRAGMA_ENABLE_DEPRECATION_WARNINGS
+	const UOptimusComponentSource* ComponentSource = WeakDataInterface->ComponentSourceBinding->GetComponentSource();
+	const FString& NumThreadsExpression = WeakDataInterface->NumThreadsExpression;
 	
 	if (!WeakComponent.IsValid() || !ComponentSource)
 	{
diff --git a/Engine/Plugins/Animation/DeformerGraph/Source/OptimusCore/Private/DataInterfaces/OptimusDataInterfaceCustomComputeKernel.h b/Engine/Plugins/Animation/DeformerGraph/Source/OptimusCore/Private/DataInterfaces/OptimusDataInterfaceCustomComputeKernel.h
--- a/Engine/Plugins/Animation/DeformerGraph/Source/OptimusCore/Private/DataInterfaces/OptimusDataInterfaceCustomComputeKernel.h
+++ b/Engine/Plugins/Animation/DeformerGraph/Source/OptimusCore/Private/DataInterfaces/OptimusDataInterfaceCustomComputeKernel.h
@@ -3,7 +3,7 @@
 #pragma once
 
 #include "IOptimusComputeKernelDataInterface.h"
-#include "IOptimusDeformerInstanceAccessor.h"
+#include "OptimusComponentSource.h"
 #include "OptimusComputeDataInterface.h"
 #include "OptimusConstant.h"
 #include "OptimusDataType.h"
@@ -37,26 +37,24 @@ public:
 	//~ End UComputeDataInterface Interface
 
 	//~ Begin IOptimusComputeKernelDataInterface Interface
-	void SetExecutionDomainConstant(const FOptimusConstantIdentifier& InExecutionDomainConstantIdentifier) override;
+	void SetExecutionDomain(const FString& InExecutionDomain) override;
+	void SetComponentBinding(const UOptimusComponentSourceBinding* InBinding) override;
 	//~ End IOptimusComputeKernelDataInterface Interface
 	
 	UPROPERTY()
-	TWeakObjectPtr<const UOptimusComponentSourceBinding> ComponentSourceBinding_DEPRECATED;
+	TWeakObjectPtr<const UOptimusComponentSourceBinding> ComponentSourceBinding;
 	
 	UPROPERTY()
-	FString NumThreadsExpression_DEPRECATED;
+	FString NumThreadsExpression;
 
 	UPROPERTY()
-	FOptimusConstantIdentifier ExecutionDomainConstantIdentifier;
-
-protected:
+	FOptimusConstantIdentifier ExecutionDomainConstantIdentifier_DEPRECATED;
 };
 
 /** Compute Framework Data Provider for each custom compute kernel. */
 UCLASS()
 class UOptimusCustomComputeKernelDataProvider :
-	public UComputeDataProvider,
-	public IOptimusDeformerInstanceAccessor
+	public UComputeDataProvider
 {
 	GENERATED_BODY()
 
@@ -67,11 +65,6 @@ public:
 	FComputeDataProviderRenderProxy* GetRenderProxy() override;
 	//~ End UComputeDataProvider Interface
 
-	//~ Begin IOptimusDeformerInstanceAccessor Interface
-	void SetDeformerInstance(UOptimusDeformerInstance* InInstance) override;
-	UOptimusDeformerInstance* GetDeformerInstance() const override;
-	//~ End IOptimusDeformerInstanceAccessor Interface
-	
 	TWeakObjectPtr<const UActorComponent> WeakComponent = nullptr;
 
 	TWeakObjectPtr<const UOptimusCustomComputeKernelDataInterface> WeakDataInterface = nullptr;
@@ -81,16 +74,6 @@ protected:
 		TArray<int32>& OutInvocationThreadCount,
 		int32& OutTotalThreadCount
 		) const;
-	
-	bool GetInvocationThreadCounts_DEPRECATED(
-		TArray<int32>& OutInvocationThreadCount,
-		int32& OutTotalThreadCount
-		) const;
-
-		
-private:
-	UPROPERTY()
-	TObjectPtr<UOptimusDeformerInstance> DeformerInstance = nullptr;
 };
 
 class FOptimusCustomComputeKernelDataProviderProxy : public FComputeDataProviderRenderProxy
diff --git a/Engine/Plugins/Animation/DeformerGraph/Source/OptimusCore/Private/DataInterfaces/OptimusDataInterfaceLoopTerminal.cpp b/Engine/Plugins/Animation/DeformerGraph/Source/OptimusCore/Private/DataInterfaces/OptimusDataInterfaceLoopTerminal.cpp
new file mode 100644
--- /dev/null
+++ b/Engine/Plugins/Animation/DeformerGraph/Source/OptimusCore/Private/DataInterfaces/OptimusDataInterfaceLoopTerminal.cpp
@@ -0,0 +1,114 @@
+// Copyright Epic Games, Inc. All Rights Reserved.
+
+#include "OptimusDataInterfaceLoopTerminal.h"
+
+#include "Components/SceneComponent.h"
+#include "ComputeFramework/ShaderParamTypeDefinition.h"
+#include "ShaderCompilerCore.h"
+#include "ShaderParameterMetadataBuilder.h"
+
+#include UE_INLINE_GENERATED_CPP_BY_NAME(OptimusDataInterfaceLoopTerminal)
+
+
+TArray<FOptimusCDIPinDefinition> UOptimusLoopTerminalDataInterface::GetPinDefinitions() const
+{
+	TArray<FOptimusCDIPinDefinition> Defs;
+	Defs.Add({ "Index", "ReadIndex", false});
+	Defs.Add({ "Count", "ReadCount", false});
+	return Defs;	
+}
+
+TSubclassOf<UActorComponent> UOptimusLoopTerminalDataInterface::GetRequiredComponentClass() const
+{
+	return USceneComponent::StaticClass();
+}
+
+void UOptimusLoopTerminalDataInterface::GetSupportedInputs(TArray<FShaderFunctionDefinition>& OutFunctions) const
+{
+	OutFunctions.AddDefaulted_GetRef()
+		.SetName(TEXT("ReadIndex"))
+		.AddReturnType(EShaderFundamentalType::Uint);
+
+	OutFunctions.AddDefaulted_GetRef()
+		.SetName(TEXT("ReadCount"))
+		.AddReturnType(EShaderFundamentalType::Uint);
+}
+
+BEGIN_SHADER_PARAMETER_STRUCT(FLoopTerminalDataInterfaceParameters, )
+	SHADER_PARAMETER(uint32, Index)
+	SHADER_PARAMETER(uint32, Count)
+END_SHADER_PARAMETER_STRUCT()
+
+void UOptimusLoopTerminalDataInterface::GetShaderParameters(TCHAR const* UID, FShaderParametersMetadataBuilder& InOutBuilder, FShaderParametersMetadataAllocations& InOutAllocations) const
+{
+	InOutBuilder.AddNestedStruct<FLoopTerminalDataInterfaceParameters>(UID);
+}
+
+TCHAR const* UOptimusLoopTerminalDataInterface::TemplateFilePath = TEXT("/Plugin/Optimus/Private/DataInterfaceLoopTerminal.ush");
+
+TCHAR const* UOptimusLoopTerminalDataInterface::GetShaderVirtualPath() const
+{
+	return TemplateFilePath;
+}
+
+void UOptimusLoopTerminalDataInterface::GetShaderHash(FString& InOutKey) const
+{
+	GetShaderFileHash(TemplateFilePath, EShaderPlatform::SP_PCD3D_SM5).AppendString(InOutKey);
+}
+
+void UOptimusLoopTerminalDataInterface::GetHLSL(FString& OutHLSL, FString const& InDataInterfaceName) const
+{
+ 	TMap<FString, FStringFormatArg> TemplateArgs =
+	{
+		{ TEXT("DataInterfaceName"), InDataInterfaceName },
+	};
+
+	FString TemplateFile;
+	LoadShaderSourceFile(TemplateFilePath, EShaderPlatform::SP_PCD3D_SM5, &TemplateFile, nullptr);
+	OutHLSL += FString::Format(*TemplateFile, TemplateArgs);
+}
+
+UComputeDataProvider* UOptimusLoopTerminalDataInterface::CreateDataProvider(TObjectPtr<UObject> InBinding, uint64 InInputMask, uint64 InOutputMask) const
+{
+	UOptimusLoopTerminalDataProvider* Provider = NewObject<UOptimusLoopTerminalDataProvider>();
+	Provider->Index = Index;
+	Provider->Count = Count;
+	return Provider;
+}
+
+
+FComputeDataProviderRenderProxy* UOptimusLoopTerminalDataProvider::GetRenderProxy()
+{
+	return new FOptimusLoopTerminalDataProviderProxy(Index, Count);
+}
+
+FOptimusLoopTerminalDataProviderProxy::FOptimusLoopTerminalDataProviderProxy(
+	uint32 InIndex,
+	uint32 InCount
+	) :
+	Index(InIndex),
+	Count(InCount)
+{
+}
+
+
+bool FOptimusLoopTerminalDataProviderProxy::IsValid(FValidationData const& InValidationData) const
+{
+	if (InValidationData.ParameterStructSize != sizeof(FParameters))
+	{
+		return false;
+	}
+
+	return true;
+}
+
+void FOptimusLoopTerminalDataProviderProxy::GatherDispatchData(FDispatchData const& InDispatchData)
+{
+	const TStridedView<FParameters> ParameterArray = MakeStridedParameterView<FParameters>(InDispatchData);
+	for (int32 InvocationIndex = 0; InvocationIndex < ParameterArray.Num(); ++InvocationIndex)
+	{
+		FParameters& Parameters = ParameterArray[InvocationIndex];
+		Parameters.Index = Index;
+		Parameters.Count = Count;
+	}
+}
diff --git a/Engine/Plugins/Animation/DeformerGraph/Source/OptimusCore/Private/DataInterfaces/OptimusDataInterfaceLoopTerminal.h b/Engine/Plugins/Animation/DeformerGraph/Source/OptimusCore/Private/DataInterfaces/OptimusDataInterfaceLoopTerminal.h
new file mode 100644
--- /dev/null
+++ b/Engine/Plugins/Animation/DeformerGraph/Source/OptimusCore/Private/DataInterfaces/OptimusDataInterfaceLoopTerminal.h
@@ -0,0 +1,80 @@
+// Copyright Epic Games, Inc. All Rights Reserved.
+
+#pragma once
+
+#include "OptimusComputeDataInterface.h"
+#include "ComputeFramework/ComputeDataProvider.h"
+#include "OptimusDataInterfaceLoopTerminal.generated.h"
+
+class FLoopTerminalDataInterfaceParameters;
+class USceneComponent;
+
+/** Compute Framework Data Interface for reading loop data. */
+UCLASS(Category = ComputeFramework)
+class UOptimusLoopTerminalDataInterface : public UOptimusComputeDataInterface
+{
+	GENERATED_BODY()
+
+public:
+	//~ Begin UOptimusComputeDataInterface Interface
+	FString GetDisplayName() const override { return TEXT("LoopTerminal"); }
+	bool IsVisible() const override { return false; }
+	TArray<FOptimusCDIPinDefinition> GetPinDefinitions() const override;
+	TSubclassOf<UActorComponent> GetRequiredComponentClass() const override;
+	//~ End UOptimusComputeDataInterface Interface
+
+	//~ Begin UComputeDataInterface Interface
+	TCHAR const* GetClassName() const override { return TEXT("LoopTerminal"); }
+	bool CanSupportUnifiedDispatch() const override { return true; }
+	void GetSupportedInputs(TArray<FShaderFunctionDefinition>& OutFunctions) const override;
+	void GetShaderParameters(TCHAR const* UID, FShaderParametersMetadataBuilder& InOutBuilder, FShaderParametersMetadataAllocations& InOutAllocations) const override;
+	TCHAR const* GetShaderVirtualPath() const override;
+	void GetShaderHash(FString& InOutKey) const override;
+	void GetHLSL(FString& OutHLSL, FString const& InDataInterfaceName) const override;
+	UComputeDataProvider* CreateDataProvider(TObjectPtr<UObject> InBinding, uint64 InInputMask, uint64 InOutputMask) const override;
+	//~ End UComputeDataInterface Interface
+
+private:
+	friend class UOptimusDeformer;
+	
+	static TCHAR const* TemplateFilePath;
+
+	UPROPERTY()
+	uint32 Index = 0;
+
+	UPROPERTY()
+	uint32 Count = 0;
+};
+
+/** Compute Framework Data Provider for reading loop data. */
+UCLASS(BlueprintType, editinlinenew, Category = ComputeFramework)
+class UOptimusLoopTerminalDataProvider : public UComputeDataProvider
+{
+	GENERATED_BODY()
+
+public:
+	//~ Begin UComputeDataProvider Interface
+	FComputeDataProviderRenderProxy* GetRenderProxy() override;
+	//~ End UComputeDataProvider Interface
+
+	uint32 Index = 0;
+	
+	uint32 Count = 0;
+};
+
+class FOptimusLoopTerminalDataProviderProxy : public FComputeDataProviderRenderProxy
+{
+public:
+	FOptimusLoopTerminalDataProviderProxy(uint32 InIndex, uint32 InCount);
+
+	//~ Begin FComputeDataProviderRenderProxy Interface
+	bool IsValid(FValidationData const& InValidationData) const override;
+	void GatherDispatchData(FDispatchData const& InDispatchData) override;
+	//~ End FComputeDataProviderRenderProxy Interface
+
+private:
+	using FParameters = FLoopTerminalDataInterfaceParameters;
+
+	uint32 Index = 0;
+	uint32 Count = 0;
+};
diff --git a/Engine/Plugins/Animation/DeformerGraph/Source/OptimusCore/Private/DataInterfaces/OptimusDataInterfaceRawBuffer.cpp b/Engine/Plugins/Animation/DeformerGraph/Source/OptimusCore/Private/DataInterfaces/OptimusDataInterfaceRawBuffer.cpp
--- a/Engine/Plugins/Animation/DeformerGraph/Source/OptimusCore/Private/DataInterfaces/OptimusDataInterfaceRawBuffer.cpp
+++ b/Engine/Plugins/Animation/DeformerGraph/Source/OptimusCore/Private/DataInterfaces/OptimusDataInterfaceRawBuffer.cpp
@@ -214,6 +214,7 @@ UComputeDataProvider* UOptimusImplicitPersistentBufferDataInterface::CreateDataP
 	uint64 InInputMask, uint64 InOutputMask) const
 {
 	UOptimusImplicitPersistentBufferDataProvider *Provider = CreateProvider<UOptimusImplicitPersistentBufferDataProvider>(InBinding);
+	Provider->DataInterfaceName = this->GetFName();
 	Provider->bZeroInitForAtomicWrites = bZeroInitForAtomicWrites;
 	return Provider;
 }
@@ -411,63 +412,14 @@ bool UOptimusRawBufferDataProvider::GetLodAndInvocationElementCounts(
 	return false;
 }
 
-bool UOptimusRawBufferDataProvider::GetLodContextAndInvocationElementCounts(
-	TArray<int32>& OutInvocationElementCounts,
-	FOptimusDeformerInstanceComponentLodContext* OutLodContext
-	) const
-{
-	FOptimusConstantEvaluationResult Result = DeformerInstance->GetConstantValuePerInvocation(DomainConstantIdentifier);
-
-	// Can happen if the bound component does not have actual data, like when there is no preview mesh
-	if (!Result.IsValid())
-	{
-		return false;
-	}
-
-	if (OutLodContext)
-	{
-		*OutLodContext = Result.LodContext;
-	}
-	const TArray<float>& Values = Result.ValuePerInvocation;
-	OutInvocationElementCounts.Reset(Values.Num());
-	for (const float& Value : Values)
-	{
-		OutInvocationElementCounts.Add(static_cast<int32>(Value));
-	}
-
-	return true;
-}
-
-void UOptimusRawBufferDataProvider::SetDeformerInstance(UOptimusDeformerInstance* InInstance)
-{
-	DeformerInstance = InInstance;
-}
-
-UOptimusDeformerInstance* UOptimusRawBufferDataProvider::GetDeformerInstance() const
-{
-	return DeformerInstance;
-}
-
-
 FComputeDataProviderRenderProxy* UOptimusTransientBufferDataProvider::GetRenderProxy()
 {
 	int32 LodIndex;
 	TArray<int32> InvocationCounts;
 	
-	if (DomainConstantIdentifier.IsValid())
-	{
-		// Querying the deformer instance for the domain of this buffer
-		if (!GetLodContextAndInvocationElementCounts(InvocationCounts))
-		{
-			InvocationCounts.Reset();
-		}
-	}
-	else
+	if (!GetLodAndInvocationElementCounts(LodIndex, InvocationCounts))
 	{
-		if (!GetLodAndInvocationElementCounts(LodIndex, InvocationCounts))
-		{
-			InvocationCounts.Reset();
-		}
+		InvocationCounts.Reset();
 	}
 	
 	return new FOptimusTransientBufferDataProviderProxy(InvocationCounts, ElementStride, RawStride, bZeroInitForAtomicWrites);
@@ -475,24 +427,19 @@ FComputeDataProviderRenderProxy* UOptimusTransientBufferDataProvider::GetRenderP
 
 FComputeDataProviderRenderProxy* UOptimusImplicitPersistentBufferDataProvider::GetRenderProxy()
 {
-	FOptimusDeformerInstanceComponentLodContext LodContext;
+	int32 LodIndex;
 	TArray<int32> InvocationCounts;
 
-	// Identifier can be unassigned when there isn't a valid component bound, see UOptimusRawBufferDataInterface::CreateProvider
-	if (DomainConstantIdentifier.IsValid())
+	if (!GetLodAndInvocationElementCounts(LodIndex, InvocationCounts))
 	{
-		// Querying the deformer instance for the domain of this buffer
-		if (!GetLodContextAndInvocationElementCounts(InvocationCounts, &LodContext))
-		{
-			InvocationCounts.Reset();
-		}
+		InvocationCounts.Reset();
 	}
-	
+
 	return new FOptimusImplicitPersistentBufferDataProviderProxy(
 		InvocationCounts, ElementStride, RawStride, bZeroInitForAtomicWrites,
 		BufferPool,
-		DomainConstantIdentifier,
-		LodContext);
+		DataInterfaceName,
+		LodIndex);
 }
 
 
@@ -580,8 +527,8 @@ FOptimusImplicitPersistentBufferDataProviderProxy::FOptimusImplicitPersistentBuf
 	int32 InRawStride,
 	bool bInZeroInitForAtomicWrites,
 	TSharedPtr<FOptimusPersistentBufferPool> InBufferPool,
-	const FOptimusConstantIdentifier& InDomainConstantIdentifier,
-	const FOptimusDeformerInstanceComponentLodContext& InLodContext
+	FName InDataInterfaceName,
+	int32 InLODIndex
 	) :
 	InvocationElementCounts(InInvocationElementCounts),
 	TotalElementCount(0),
@@ -589,8 +536,8 @@ FOptimusImplicitPersistentBufferDataProviderProxy::FOptimusImplicitPersistentBuf
 	RawStride(InRawStride),
 	bZeroInitForAtomicWrites(bInZeroInitForAtomicWrites),
 	BufferPool(InBufferPool),
-	DomainConstantIdentifier(InDomainConstantIdentifier),
-	LodContext(InLodContext)
+	DataInterfaceName(InDataInterfaceName),
+	LODIndex(InLODIndex)
 {
 	for (int32 NumElements : InvocationElementCounts)
 	{
@@ -610,6 +557,11 @@ bool FOptimusImplicitPersistentBufferDataProviderProxy::IsValid(FValidationData
 		return false;
 	}
 
+	if (DataInterfaceName == NAME_None)
+	{
+		return false;
+	}
+
 	return true;
 }
 
@@ -619,7 +571,7 @@ void FOptimusImplicitPersistentBufferDataProviderProxy::AllocateResources(FRDGBu
 	Count.Add(TotalElementCount);
 	TArray<FRDGBufferRef> Buffers;
 	bool bJustAllocated = false;
-	BufferPool->GetImplicitPersistentBuffers(GraphBuilder, DomainConstantIdentifier, LodContext, ElementStride, RawStride, Count, Buffers, bJustAllocated);
+	BufferPool->GetImplicitPersistentBuffers(GraphBuilder, DataInterfaceName, LODIndex, ElementStride, RawStride, Count, Buffers, bJustAllocated);
 
 	ensure(Buffers.Num() == 1);
 	Buffer = Buffers[0];
diff --git a/Engine/Plugins/Animation/DeformerGraph/Source/OptimusCore/Private/DataInterfaces/OptimusDataInterfaceRawBuffer.h b/Engine/Plugins/Animation/DeformerGraph/Source/OptimusCore/Private/DataInterfaces/OptimusDataInterfaceRawBuffer.h
--- a/Engine/Plugins/Animation/DeformerGraph/Source/OptimusCore/Private/DataInterfaces/OptimusDataInterfaceRawBuffer.h
+++ b/Engine/Plugins/Animation/DeformerGraph/Source/OptimusCore/Private/DataInterfaces/OptimusDataInterfaceRawBuffer.h
@@ -2,7 +2,6 @@
 
 #pragma once
 
-#include "IOptimusDeformerInstanceAccessor.h"
 #include "IOptimusPersistentBufferProvider.h"
 #include "ComputeFramework/ComputeDataProvider.h"
 #include "ComputeFramework/ShaderParamTypeDefinition.h"
@@ -76,7 +75,7 @@ public:
 	TWeakObjectPtr<UOptimusComponentSourceBinding> ComponentSourceBinding;
 
 	UPROPERTY()
-	FOptimusConstantIdentifier DomainConstantIdentifier;
+	FOptimusConstantIdentifier DomainConstantIdentifier_DEPRECATED;
 	
 protected:
 	virtual bool UseSplitBuffers() const { return true; }
@@ -94,7 +93,6 @@ protected:
 			Provider->DataDomain = DataDomain;
 			Provider->ElementStride = ValueType->GetResourceElementSize();
 			Provider->RawStride = GetRawStride();
-			Provider->DomainConstantIdentifier = DomainConstantIdentifier;
 		}
 		return Provider;
 	}
@@ -130,7 +128,7 @@ public:
 	bool bZeroInitForAtomicWrites = false;	
 };
 
-/** Compute Framework Data Interface for a transient buffer. */
+/** Compute Framework Data Interface for a implicit persistent buffer. */
 UCLASS(Category = ComputeFramework)
 class OPTIMUSCORE_API UOptimusImplicitPersistentBufferDataInterface : public UOptimusRawBufferDataInterface
 {
@@ -185,8 +183,7 @@ protected:
 /** Compute Framework Data Provider for a transient buffer. */
 UCLASS(Abstract)
 class OPTIMUSCORE_API UOptimusRawBufferDataProvider :
-	public UComputeDataProvider,
-	public IOptimusDeformerInstanceAccessor
+	public UComputeDataProvider
 {
 	GENERATED_BODY()
 
@@ -200,19 +197,6 @@ public:
 		TArray<int32>& OutInvocationElementCounts
 		) const;
 
-	/** Helper function to calculate the element count given a constant identifier,
-	 *	LOD is handled by the deformer instance
-	 */
-	bool GetLodContextAndInvocationElementCounts(
-		TArray<int32>& OutInvocationElementCounts,
-		FOptimusDeformerInstanceComponentLodContext* OutLodContext = nullptr
-		) const;
-
-	//~ Begin IOptimusDeformerInstanceAccessor Interface
-	void SetDeformerInstance(UOptimusDeformerInstance* InInstance) override;
-	UOptimusDeformerInstance* GetDeformerInstance() const override;
-	//~ End IOptimusDeformerInstanceAccessor Interface
-	
 	/** The skinned mesh component that governs the sizing and LOD of this buffer */
 	TWeakObjectPtr<const UActorComponent> Component = nullptr;
 
@@ -224,11 +208,6 @@ public:
 	int32 ElementStride = 4;
 
 	int32 RawStride = 0;
-
-	FOptimusConstantIdentifier DomainConstantIdentifier;
-	
-private:
-	TObjectPtr<UOptimusDeformerInstance> DeformerInstance = nullptr;
 };
 
 
@@ -265,6 +244,8 @@ public:
 	//~ Begin IOptimusPersistentBufferPoolUser Interface
 	
 	bool bZeroInitForAtomicWrites = false;
+	
+	FName DataInterfaceName = NAME_None;
 private:
 	/** The buffer pool we refer to. Set by UOptimusDeformerInstance::SetupFromDeformer after providers have been
 	 *  created
@@ -340,8 +321,8 @@ public:
 		int32 InRawStride,
 		bool bInZeroInitForAtomicWrites,
 		TSharedPtr<FOptimusPersistentBufferPool> InBufferPool,
-		const FOptimusConstantIdentifier& InDomainConstantIdentifier,
-		const FOptimusDeformerInstanceComponentLodContext& InLodContext
+		FName InDataInterfaceName,
+		int32 InLODIndex
 		);
 
 	//~ Begin FComputeDataProviderRenderProxy Interface
@@ -360,8 +341,8 @@ public:
 	const bool bZeroInitForAtomicWrites;
 	
 	const TSharedPtr<FOptimusPersistentBufferPool> BufferPool;
-	const FOptimusConstantIdentifier DomainConstantIdentifier;
-	const FOptimusDeformerInstanceComponentLodContext LodContext;
+	FName DataInterfaceName;
+	int32 LODIndex;
 
 	FRDGBufferRef Buffer;
 	FRDGBufferSRVRef BufferSRV;
diff --git a/Engine/Plugins/Animation/DeformerGraph/Source/OptimusCore/Private/Nodes/OptimusNode_ComputeKernelBase.cpp b/Engine/Plugins/Animation/DeformerGraph/Source/OptimusCore/Private/Nodes/OptimusNode_ComputeKernelBase.cpp
--- a/Engine/Plugins/Animation/DeformerGraph/Source/OptimusCore/Private/Nodes/OptimusNode_ComputeKernelBase.cpp
+++ b/Engine/Plugins/Animation/DeformerGraph/Source/OptimusCore/Private/Nodes/OptimusNode_ComputeKernelBase.cpp
@@ -10,13 +10,14 @@
 
 #include "DataInterfaces/OptimusDataInterfaceGraph.h"
 #include "DataInterfaces/OptimusDataInterfaceRawBuffer.h"
+#include "DataInterfaces/OptimusDataInterfaceLoopTerminal.h"
 #include "IOptimusValueProvider.h"
 #include "OptimusComponentSource.h"
 #include "OptimusDataTypeRegistry.h"
 #include "OptimusHelpers.h"
 #include "OptimusKernelSource.h"
-#include "OptimusConstant.h"
 #include "OptimusExecutionDomain.h"
+#include "OptimusNode_LoopTerminal.h"
 #include "OptimusNode_ResourceAccessorBase.h"
 
 #include UE_INLINE_GENERATED_CPP_BY_NAME(OptimusNode_ComputeKernelBase)
@@ -55,14 +56,12 @@ static void CopyValueType(FShaderValueTypeHandle InValueType,  FShaderParamTypeD
 FOptimus_ComputeKernelResult UOptimusNode_ComputeKernelBase::CreateComputeKernel(
 	UObject *InKernelSourceOuter,
 	const FOptimusPinTraversalContext& InTraversalContext,
-	const FOptimus_NodeToDataInterfaceMap& InNodeDataInterfaceMap,
-	const FOptimus_PinToDataInterfaceMap& InLinkDataInterfaceMap,
+	const FOptimus_KernelInputMap InKernelInputs,
+	const FOptimus_KernelOutputMap InKernelOutputs,
 	const TArray<const UOptimusNode *>& InValueNodes,
-	const UComputeDataInterface* InGraphDataInterface,
 	UComputeDataInterface* InOutKernelDataInterface,
 	FOptimus_InterfaceBindingMap& OutInputDataBindings,
-	FOptimus_InterfaceBindingMap& OutOutputDataBindings,
-	FOptimusKernelConstantContainer& OutKernelConstantContainer
+	FOptimus_InterfaceBindingMap& OutOutputDataBindings
 ) const
 {
 	// Maps friendly name to unique name for each struct type
@@ -161,9 +160,8 @@ FOptimus_ComputeKernelResult UOptimusNode_ComputeKernelBase::CreateComputeKernel
 			{
 				TOptional<FText> Result = ProcessInputPinForComputeKernel(
 					InTraversalContext, SubPin, GroupIndex == 0 ? FString() : Pin->GetName(),
-					InNodeDataInterfaceMap, InLinkDataInterfaceMap, InValueNodes,
-					InGraphDataInterface,
-					KernelSource, GeneratedFunctions, OutInputDataBindings, OutKernelConstantContainer
+					InKernelInputs, InValueNodes,
+					KernelSource, GeneratedFunctions, OutInputDataBindings
 					);
 				if (Result.IsSet())
 				{
@@ -177,8 +175,8 @@ FOptimus_ComputeKernelResult UOptimusNode_ComputeKernelBase::CreateComputeKernel
 		{
 			ProcessOutputPinForComputeKernel(
 				InTraversalContext, Pin,
-				InNodeDataInterfaceMap, InLinkDataInterfaceMap,
-				KernelSource, GeneratedFunctions, OutInputDataBindings, OutOutputDataBindings, OutKernelConstantContainer);
+				InKernelOutputs,
+				KernelSource, GeneratedFunctions, OutInputDataBindings, OutOutputDataBindings);
 		}
 	}
 
@@ -188,8 +186,8 @@ FOptimus_ComputeKernelResult UOptimusNode_ComputeKernelBase::CreateComputeKernel
 			PrimaryBindings.Array()[0],
 			InOutKernelDataInterface,
 			KernelSource,
-			OutInputDataBindings,
-			OutKernelConstantContainer);
+			OutInputDataBindings
+			);
 	}
 
 	FString CookedSource;
@@ -243,24 +241,6 @@ FOptimus_ComputeKernelResult UOptimusNode_ComputeKernelBase::CreateComputeKernel
 }
 
 
-bool UOptimusNode_ComputeKernelBase::HasMutableInput() const
-{
-	for (const UOptimusNodePin* Pin: GetPins())
-	{
-		if (Pin->IsGroupingPin())
-		{
-			for (const UOptimusNodePin* SubPin: Pin->GetSubPins())
-			{
-				if (GetOwningGraph()->IsPinMutable(SubPin))
-				{
-					return true;
-				}
-			}
-		}
-	}
-	return false;
-}
-
 TOptional<FText> UOptimusNode_ComputeKernelBase::ValidateForCompile() const
 {
 	auto GetStructTypeDefFromPin = [](const UOptimusNodePin* InPin) -> TOptional<FText>
@@ -303,7 +283,7 @@ TOptional<FText> UOptimusNode_ComputeKernelBase::ValidateForCompile() const
 		}
 		return {};
 	};
-	
+
 	TArray<TSet<UOptimusComponentSourceBinding*>> BindingsPerGroup;
 	
 	for (const UOptimusNodePin* Pin: GetPins())
@@ -323,7 +303,7 @@ TOptional<FText> UOptimusNode_ComputeKernelBase::ValidateForCompile() const
 				{
 					return Result;
 				}
-				
+
 				if (TOptional<FText> Result = GetStructTypeDefFromPin(SubPin); Result.IsSet())
 				{
 					return Result;
@@ -473,112 +453,105 @@ TOptional<FText> UOptimusNode_ComputeKernelBase::ProcessInputPinForComputeKernel
 	const FOptimusPinTraversalContext& InTraversalContext,
 	const UOptimusNodePin* InInputPin,
 	const FString& InGroupName,
-    const FOptimus_NodeToDataInterfaceMap& InNodeDataInterfaceMap,
-	const FOptimus_PinToDataInterfaceMap& InLinkDataInterfaceMap,
+	const FOptimus_KernelInputMap& InKernelInputs,
 	const TArray<const UOptimusNode*>& InValueNodes,
-	const UComputeDataInterface* InGraphDataInterface,
 	UOptimusKernelSource* InKernelSource,
 	TArray<FString>& OutGeneratedFunctions,
-	FOptimus_InterfaceBindingMap& OutInputDataBindings,
-	FOptimusKernelConstantContainer& OutKernelConstantContainer
+	FOptimus_InterfaceBindingMap& OutInputDataBindings
 	) const
 {
-	const UOptimusNodePin* OutputPin = nullptr;
-	TArray<FOptimusRoutedNodePin> ConnectedPins = InInputPin->GetConnectedPinsWithRouting(InTraversalContext);
-	if (ConnectedPins.Num() == 1)
-	{
-		OutputPin = ConnectedPins[0].NodePin;
-	}
 	
 	FShaderParamTypeDefinition IndexParamDef;
 	CopyValueType(FShaderValueType::Get(EShaderFundamentalType::Uint), IndexParamDef);
 
 	const FShaderValueTypeHandle ValueType = InInputPin->GetDataType()->ShaderValueType;
-	const UOptimusNode *OutputNode = OutputPin ? OutputPin->GetOwningNode() : nullptr;
 
-	// For inputs, we only have to deal with a single read, because only one
-	// link can connect into it. 
-	if (OutputPin)
+	if (const FOptimus_KernelConnection* InputConnection = InKernelInputs.Find(InInputPin))
 	{
-		const UComputeDataInterface* DataInterface = nullptr;
+		const UOptimusNodePin* OutputPin = InputConnection->ConnectedPin;
+		const UOptimusNode* OutputNode = OutputPin->GetOwningNode();
+		UComputeDataInterface* DataInterface = InputConnection->DataInterface;
+		
 		const UOptimusComponentSourceBinding* ComponentBinding = nullptr;
 		int32 DataInterfaceFuncIndex = INDEX_NONE;
 		FString DataFunctionName;
-		
-		// Are we being connected from a scene data interface or a transient buffer?
-		if (InLinkDataInterfaceMap.Contains(OutputPin))
-		{
-			// For transient buffers we need the function index as given by the
-			// ReadValue function. 
-			DataInterface = InLinkDataInterfaceMap[OutputPin];
-			DataInterfaceFuncIndex = UOptimusRawBufferDataInterface::GetReadValueInputIndex(EOptimusBufferReadType::Default);
 
-			TArray<FShaderFunctionDefinition> ReadFunctions;
-			DataInterface->GetSupportedInputs(ReadFunctions);
-			
-			DataFunctionName = ReadFunctions[DataInterfaceFuncIndex].Name;
+		if (Cast<const IOptimusComputeKernelProvider>(OutputNode))
+		{
+			if (ensure(Cast<UOptimusRawBufferDataInterface>(DataInterface)))
+			{
+				DataInterfaceFuncIndex = UOptimusRawBufferDataInterface::GetReadValueInputIndex(EOptimusBufferReadType::Default);
+				
+				TArray<FShaderFunctionDefinition> ReadFunctions;
+				DataInterface->GetSupportedInputs(ReadFunctions);
+				DataFunctionName = ReadFunctions[DataInterfaceFuncIndex].Name;	
+			}
 		}
-		else if(InNodeDataInterfaceMap.Contains(OutputNode))
+		else if (const IOptimusDataInterfaceProvider* InterfaceProvider = Cast<const IOptimusDataInterfaceProvider>(OutputNode)) 
 		{
-			const IOptimusDataInterfaceProvider* InterfaceProvider = Cast<const IOptimusDataInterfaceProvider>(OutputNode);
-			
-			// FIXME: Sub-pin read support.
-			UOptimusComputeDataInterface const* OptimusDataInterface = InNodeDataInterfaceMap[OutputNode];
-			DataInterface = OptimusDataInterface;
-
-			DataInterfaceFuncIndex = InterfaceProvider->GetDataFunctionIndexFromPin(OutputPin);
-			
-			TArray<FShaderFunctionDefinition> ReadFunctions;
-			DataInterface->GetSupportedInputs(ReadFunctions);
-
-			DataFunctionName = ReadFunctions[DataInterfaceFuncIndex].Name;
+			if (ensure(Cast<UOptimusComputeDataInterface>(DataInterface)))
+			{
+				DataInterfaceFuncIndex = InterfaceProvider->GetDataFunctionIndexFromPin(OutputPin);
+				
+				TArray<FShaderFunctionDefinition> ReadFunctions;
+				DataInterface->GetSupportedInputs(ReadFunctions);
+				DataFunctionName = ReadFunctions[DataInterfaceFuncIndex].Name;
+			}
 		}
-		else if (IOptimusValueProvider const* ValueProvider = Cast<const IOptimusValueProvider>(OutputNode))
+		else if (const IOptimusValueProvider * ValueProvider = Cast<const IOptimusValueProvider>(OutputNode))
 		{
-			// Value nodes bind the single graph data interface.
-			DataInterface = InGraphDataInterface;
-			DataInterfaceFuncIndex = InValueNodes.Find(OutputNode);
-			check(DataInterfaceFuncIndex != INDEX_NONE);
-			DataFunctionName = ValueProvider->GetValueName();
+			if (ensure(Cast<UOptimusGraphDataInterface>(DataInterface)))
+			{
+				// Value nodes bind the single graph data interface.
+				DataInterfaceFuncIndex = InValueNodes.Find(OutputNode);
+				check(DataInterfaceFuncIndex != INDEX_NONE);
+				DataFunctionName = ValueProvider->GetValueName();	
+			}
 		}
-
-		// If we are connected from a data interface, set the input binding up now.
-		if (DataInterface)
+		else if (Cast<const UOptimusNode_LoopTerminal>(OutputNode))
 		{
-			if (!DataInterface->CanSupportUnifiedDispatch() && !InGroupName.IsEmpty())
+			if (ensure(Cast<const UOptimusLoopTerminalDataInterface>(DataInterface)))
 			{
-				return FText::Format(LOCTEXT("SecondaryGroupOnlySupportsDataInterfaceUnifiedDispatch", "Cannot connect Secondary group input binding pin '{0}' to a data interface that does not support unified dispatch, consider inserting a kernel inbetween"), FText::FromName(InInputPin->GetUniqueName()));
+				DataInterfaceFuncIndex = UOptimusNode_LoopTerminal::GetDataFunctionIndexFromPin(OutputPin);
+				TArray<FShaderFunctionDefinition> ReadFunctions;
+				DataInterface->GetSupportedInputs(ReadFunctions);
+				DataFunctionName = ReadFunctions[DataInterfaceFuncIndex].Name;
 			}
+		}
 
-			// The shader function definition that exposes the function that we use to
-			// read values to input into the kernel.
-			FShaderFunctionDefinition FuncDef;
-			FuncDef.Name = DataFunctionName;
-			FuncDef.bHasReturnType = true;
-			
-			FShaderParamTypeDefinition ParamDef;
-			CopyValueType(ValueType, ParamDef);
-			FuncDef.ParamTypes.Emplace(ParamDef);
+		if (!DataInterface->CanSupportUnifiedDispatch() && !InGroupName.IsEmpty())
+		{
+			return FText::Format(LOCTEXT("SecondaryGroupOnlySupportsDataInterfaceUnifiedDispatch", "Cannot connect Secondary group input binding pin '{0}' to a data interface that does not support unified dispatch, consider inserting a kernel inbetween"), FText::FromName(InInputPin->GetUniqueName()));
+		}
 
-			// For resources we need the index parameter.
-			if (!InInputPin->GetDataDomain().IsSingleton())
+		// The shader function definition that exposes the function that we use to
+		// read values to input into the kernel.
+		FShaderFunctionDefinition FuncDef;
+		FuncDef.Name = DataFunctionName;
+		FuncDef.bHasReturnType = true;
+		
+		FShaderParamTypeDefinition ParamDef;
+		CopyValueType(ValueType, ParamDef);
+		FuncDef.ParamTypes.Emplace(ParamDef);
+
+		// For resources we need the index parameter.
+		if (!InInputPin->GetDataDomain().IsSingleton())
+		{
+			for (int32 Count = 0; Count < InInputPin->GetDataDomain().NumDimensions(); Count++)
 			{
-				for (int32 Count = 0; Count < InInputPin->GetDataDomain().NumDimensions(); Count++)
-				{
-					FuncDef.ParamTypes.Add(IndexParamDef);
-				}
+				FuncDef.ParamTypes.Add(IndexParamDef);
 			}
-
-			FOptimus_InterfaceBinding InterfaceBinding;
-			InterfaceBinding.DataInterface = DataInterface;
-			InterfaceBinding.DataInterfaceBindingIndex = DataInterfaceFuncIndex;
-			InterfaceBinding.BindingFunctionName = FString::Printf(TEXT("Read%s"), *InInputPin->GetName());
-			InterfaceBinding.BindingFunctionNamespace = InGroupName;
-			
-			OutInputDataBindings.Add(InKernelSource->ExternalInputs.Num(), InterfaceBinding);  
-			
-			InKernelSource->ExternalInputs.Emplace(FuncDef);
 		}
+
+		FOptimus_InterfaceBinding InterfaceBinding;
+		InterfaceBinding.DataInterface = DataInterface;
+		InterfaceBinding.DataInterfaceBindingIndex = DataInterfaceFuncIndex;
+		InterfaceBinding.BindingFunctionName = FString::Printf(TEXT("Read%s"), *InInputPin->GetName());
+		InterfaceBinding.BindingFunctionNamespace = InGroupName;
+		
+		OutInputDataBindings.Add(InKernelSource->ExternalInputs.Num(), InterfaceBinding);  
+		
+		InKernelSource->ExternalInputs.Emplace(FuncDef);
 	}
 	else
 	{
@@ -620,57 +593,9 @@ TOptional<FText> UOptimusNode_ComputeKernelBase::ProcessInputPinForComputeKernel
 
 		OutGeneratedFunctions.Add(
 			FString::Printf(TEXT("%s%s Read%s(%s) { return %s; }%s"),
-				*NamespacePrefix, *ValueType->ToString(), *InInputPin->GetName(), *OptionalParamStr, *ValueStr, *NamespaceSuffix));
-	}
-
-	if (InInputPin->GetDataDomain().IsOneDimensional())
-	{
-		// Get the component binding from the upstream connection. Unconnected input pin should use the group's component binding
-		int32 ComponentBindingIndex = INDEX_NONE;
-		
-		if (const UOptimusNodePin* GroupPin = InInputPin->GetParentPin())
-		{
-			TArray<UOptimusComponentSourceBinding*> ComponentBindings = GroupPin->GetComponentSourceBindingsRecursively().Array();
-			if (ensure(ComponentBindings.Num() == 1))
-			{
-				ComponentBindingIndex = ComponentBindings[0]->GetIndex();
-			}
-		}
-		
-		// Use the domain of input pin if it is defined
-		// Use the domain of connected output pin otherwise
-		if (InInputPin->GetDataDomain().IsFullyDefined())
-		{
-			OutKernelConstantContainer.AddToKernelContainer({
-				{this, *InGroupName, InInputPin->GetFName()},
-				{InInputPin->GetDataDomain().AsExpression().GetValue()},
-				ComponentBindingIndex,
-				EOptimusConstantType::Input});
-		}
-		else if (ensure(OutputPin) && ensure(OutputPin->GetDataDomain().IsOneDimensional()) && ensure(OutputPin->GetDataDomain().IsFullyDefined()))
-		{
-			if (InLinkDataInterfaceMap.Contains(OutputPin))
-			{
-				if (const UOptimusRawBufferDataInterface* RawBufferDataInterface = Cast<UOptimusRawBufferDataInterface>(InLinkDataInterfaceMap[OutputPin]))
-				{
-					OutKernelConstantContainer.AddToKernelContainer({
-						{this, *InGroupName, InInputPin->GetFName()},
-						{RawBufferDataInterface->DomainConstantIdentifier},
-						ComponentBindingIndex,
-						EOptimusConstantType::Input});	
-				}
-			}
-			else
-			{
-				OutKernelConstantContainer.AddToKernelContainer({
-					{this, *InGroupName, InInputPin->GetFName()},
-					{OutputPin->GetDataDomain().AsExpression().Get({})},
-					ComponentBindingIndex,
-					EOptimusConstantType::Input});
-			}
-		}
+				*NamespacePrefix, *ValueType->ToString(), *InInputPin->GetName(), *OptionalParamStr, *ValueStr, *NamespaceSuffix));	
 	}
-
+	
 	return {};
 }
 
@@ -678,13 +603,11 @@ TOptional<FText> UOptimusNode_ComputeKernelBase::ProcessInputPinForComputeKernel
 void UOptimusNode_ComputeKernelBase::ProcessOutputPinForComputeKernel(
 	const FOptimusPinTraversalContext& InTraversalContext,
 	const UOptimusNodePin* InOutputPin,
-	const FOptimus_NodeToDataInterfaceMap& InNodeDataInterfaceMap,
-	const FOptimus_PinToDataInterfaceMap& InLinkDataInterfaceMap,
+	const FOptimus_KernelOutputMap& InKernelOutputs,
 	UOptimusKernelSource* InKernelSource,
 	TArray<FString>& OutGeneratedFunctions,
 	FOptimus_InterfaceBindingMap& OutInputDataBindings,
-	FOptimus_InterfaceBindingMap& OutOutputDataBindings,
-	FOptimusKernelConstantContainer& OutKernelConstantContainer
+	FOptimus_InterfaceBindingMap& OutOutputDataBindings
 	) const
 {
 	FShaderParamTypeDefinition IndexParamDef;
@@ -694,70 +617,73 @@ void UOptimusNode_ComputeKernelBase::ProcessOutputPinForComputeKernel(
 
 	const FShaderValueTypeHandle ValueType = InOutputPin->GetDataType()->ShaderValueType;
 
-	TArray<FOptimusRoutedNodePin> ConnectedPins = InOutputPin->GetConnectedPinsWithRouting(InTraversalContext);
-
-	// Get the component binding from the upstream connection.
-	UOptimusComponentSourceBinding* ComponentBinding = nullptr;
-	TArray<UOptimusComponentSourceBinding*> ComponentBindings = GetOwningGraph()->GetComponentSourceBindingsForPin(InOutputPin).Array();
-	if (ensure(ComponentBindings.Num() == 1))
-	{
-		ComponentBinding = ComponentBindings[0];
-	}
-
-	if (!ensure(ComponentBinding))
-	{
-		// Validate For Compile should have caught this
-		return;
-	}
-	
-	
-	if (!ConnectedPins.IsEmpty())
+	if (const TArray<FOptimus_KernelConnection>* KernelConnections = InKernelOutputs.Find(InOutputPin))
 	{
 		// If we have an output connection going to multiple data interfaces, then we
 		// have to wrap them all up in a single proxy function to make it still transparent
 		// to the kernel writer.
 		struct FWriteConnectionDef
 		{
-			UOptimusComputeDataInterface* DataInterface = nullptr;
+			const UComputeDataInterface* DataInterface = nullptr;
 			FString DataFunctionName;
 			FString WriteToName;
 		};
 		TArray<FWriteConnectionDef> WriteConnectionDefs;
 
-		// If we're scheduled to write to a transient data interface, do that now.
-		// There is only ever a single transient data interface per output pin.
-		if (InLinkDataInterfaceMap.Contains(InOutputPin))
+		// All connected kernels should share the same data interface
+		UComputeDataInterface* SharedRawBufferDI = nullptr;
+
+		for (const FOptimus_KernelConnection& KernelConnection: *KernelConnections)
 		{
-			UOptimusComputeDataInterface* DataInterface = InLinkDataInterfaceMap[InOutputPin]; 
+			const UOptimusNodePin* ConnectedPin = KernelConnection.ConnectedPin;
+			const UOptimusNode *ConnectedNode = ConnectedPin ? ConnectedPin->GetOwningNode() : nullptr;
+			UComputeDataInterface* DataInterface = KernelConnection.DataInterface; 
 
 			TArray<FShaderFunctionDefinition> WriteFunctions;
 			DataInterface->GetSupportedOutputs(WriteFunctions);
 
-			const int32 OutputIndex = UOptimusTransientBufferDataInterface::GetWriteValueOutputIndex(EOptimusBufferWriteType::Write);
-			WriteConnectionDefs.Add({DataInterface, WriteFunctions[OutputIndex].Name, TEXT("Transient")});
-		}
-		
-		for (const FOptimusRoutedNodePin& ConnectedPin: ConnectedPins)
-		{
-			const UOptimusNode *ConnectedNode = ConnectedPin.NodePin->GetOwningNode();
-
-			// Connected to a data interface node?
-			if(!InNodeDataInterfaceMap.Contains(ConnectedNode))
+			if (!ConnectedNode)
 			{
-				continue;
+				if (const UOptimusTransientBufferDataInterface* TransientBufferDataInterface = Cast<const UOptimusTransientBufferDataInterface>(DataInterface);
+							ensure(TransientBufferDataInterface))
+				{
+					const int32 OutputIndex = UOptimusRawBufferDataInterface::GetWriteValueOutputIndex(EOptimusBufferWriteType::Write);
+					WriteConnectionDefs.Add({TransientBufferDataInterface, WriteFunctions[OutputIndex].Name, TEXT("TransientBuffer")});	
+				}
+				
+				check(KernelConnections->Num() == 1);
+				break;
 			}
+			else
+			{
+				if (Cast<const IOptimusComputeKernelProvider>(ConnectedNode))
+				{
+					if (!SharedRawBufferDI)
+					{
+						if (const UOptimusRawBufferDataInterface* RawBufferDataInterface = Cast<const UOptimusRawBufferDataInterface>(DataInterface);
+							ensure(RawBufferDataInterface))
+						{
+							const int32 OutputIndex = UOptimusRawBufferDataInterface::GetWriteValueOutputIndex(EOptimusBufferWriteType::Write);
+							WriteConnectionDefs.Add({RawBufferDataInterface, WriteFunctions[OutputIndex].Name, TEXT("RawBuffer")});
+							SharedRawBufferDI = DataInterface;
+						}	
+					}
+					else
+					{
+						// All connected kernels should share the same data interface
+						check(SharedRawBufferDI == DataInterface);
+					}
+				}
+				else if (const IOptimusDataInterfaceProvider* InterfaceProvider = Cast<const IOptimusDataInterfaceProvider>(ConnectedNode))
+				{
+					int32 DataFunctionIndex = InterfaceProvider->GetDataFunctionIndexFromPin(ConnectedPin);
+					TArray<FShaderFunctionDefinition> FunctionDefinitions;
+					DataInterface->GetSupportedOutputs(FunctionDefinitions);
+					FString DataFunctionName = FunctionDefinitions[DataFunctionIndex].Name;
 
-			const IOptimusDataInterfaceProvider* InterfaceProvider = Cast<const IOptimusDataInterfaceProvider>(ConnectedNode);
-			
-			// FIXME: Sub-pin write support.
-			UOptimusComputeDataInterface* DataInterface = InNodeDataInterfaceMap[ConnectedNode];
-			int32 DataFunctionIndex = InterfaceProvider->GetDataFunctionIndexFromPin(ConnectedPin.NodePin);
-			
-			TArray<FShaderFunctionDefinition> FunctionDefinitions;
-			DataInterface->GetSupportedOutputs(FunctionDefinitions);
-			FString DataFunctionName = FunctionDefinitions[DataFunctionIndex].Name;
-
-			WriteConnectionDefs.Add({DataInterface, DataFunctionName, ConnectedPin.NodePin->GetName()});
+					WriteConnectionDefs.Add({DataInterface, DataFunctionName, ConnectedPin->GetName()});
+				}	
+			}
 		}
 
 		TArray<FString> WrapFunctionNameCalls;
@@ -818,71 +744,60 @@ void UOptimusNode_ComputeKernelBase::ProcessOutputPinForComputeKernel(
 		
 		if (DoesOutputPinSupportAtomic(InOutputPin) || DoesOutputPinSupportRead(InOutputPin))
 		{
-			UOptimusRawBufferDataInterface* RawBufferDataInterface = nullptr;
-			if (InLinkDataInterfaceMap.Contains(InOutputPin))
+			// Read access on output pins guarantees existence of raw buffer data interface
+			if (ensure(KernelConnections->Num() > 0))
 			{
-				RawBufferDataInterface = Cast<UOptimusRawBufferDataInterface>(InLinkDataInterfaceMap[InOutputPin]);	
-			}
-			else if (ensure(ConnectedPins.Num() == 1))
-			{
-				const UOptimusNode *ConnectedNode = ConnectedPins[0].NodePin->GetOwningNode();
-
-				// Connected to a data interface node?
-				if(ensure(InNodeDataInterfaceMap.Contains(ConnectedNode)))
-				{
-					RawBufferDataInterface = Cast<UOptimusRawBufferDataInterface>(InNodeDataInterfaceMap[ConnectedNode]);
-				}
-			}
-
-			if (ensure(RawBufferDataInterface))
-			{
-				if (DoesOutputPinSupportAtomic(InOutputPin))
+				// We may have multiple connections but they should all be tie to the same raw buffer data interface, so just use the first connection
+				if (UOptimusRawBufferDataInterface* RawBufferDataInterface = Cast<UOptimusRawBufferDataInterface>((*KernelConnections)[0].DataInterface))
 				{
-					TArray<FShaderFunctionDefinition> WriteFunctions;
-					RawBufferDataInterface->GetSupportedOutputs(WriteFunctions);
-				
-					for (uint8 WriteType = uint8(EOptimusBufferWriteType::Write) + 1; WriteType < uint8(EOptimusBufferWriteType::Count); WriteType++)
+					if (DoesOutputPinSupportAtomic(InOutputPin))
 					{
+						TArray<FShaderFunctionDefinition> WriteFunctions;
+						RawBufferDataInterface->GetSupportedOutputs(WriteFunctions);
+					
+						for (uint8 WriteType = uint8(EOptimusBufferWriteType::Write) + 1; WriteType < uint8(EOptimusBufferWriteType::Count); WriteType++)
+						{
 
-						FString WrapFunctionName = GetAtomicWriteFunctionName((EOptimusBufferWriteType)WriteType, InOutputPin->GetName());
+							FString WrapFunctionName = GetAtomicWriteFunctionName((EOptimusBufferWriteType)WriteType, InOutputPin->GetName());
 
-						const int32 OutputIndex = UOptimusRawBufferDataInterface::GetWriteValueOutputIndex((EOptimusBufferWriteType)WriteType);
+							const int32 OutputIndex = UOptimusRawBufferDataInterface::GetWriteValueOutputIndex((EOptimusBufferWriteType)WriteType);
+						
+							FShaderFunctionDefinition FuncDef = WriteFunctions[OutputIndex];
+							for (int32 Index = 0; Index < FuncDef.ParamTypes.Num(); Index++)
+							{
+								FuncDef.ParamTypes[Index].ResetTypeDeclaration();
+							}
 					
-						FShaderFunctionDefinition FuncDef = WriteFunctions[OutputIndex];
+							FOptimus_InterfaceBinding InterfaceBinding;
+							InterfaceBinding.DataInterface = RawBufferDataInterface;
+							InterfaceBinding.DataInterfaceBindingIndex = OutputIndex;
+							InterfaceBinding.BindingFunctionName = WrapFunctionName;
+					
+							OutOutputDataBindings.Add(InKernelSource->ExternalOutputs.Num(), InterfaceBinding);
+							InKernelSource->ExternalOutputs.Emplace(FuncDef);	
+						}
+					}
+
+					if (DoesOutputPinSupportRead(InOutputPin))
+					{
+						TArray<FShaderFunctionDefinition> ReadFunctions;
+						RawBufferDataInterface->GetSupportedInputs(ReadFunctions);
+						int32 InputIndex = UOptimusRawBufferDataInterface::GetReadValueInputIndex(EOptimusBufferReadType::ForceUAV);
+
+						FShaderFunctionDefinition FuncDef = ReadFunctions[InputIndex];
 						for (int32 Index = 0; Index < FuncDef.ParamTypes.Num(); Index++)
 						{
 							FuncDef.ParamTypes[Index].ResetTypeDeclaration();
 						}
-				
+					
 						FOptimus_InterfaceBinding InterfaceBinding;
 						InterfaceBinding.DataInterface = RawBufferDataInterface;
-						InterfaceBinding.DataInterfaceBindingIndex = OutputIndex;
-						InterfaceBinding.BindingFunctionName = WrapFunctionName;
-				
-						OutOutputDataBindings.Add(InKernelSource->ExternalOutputs.Num(), InterfaceBinding);
-						InKernelSource->ExternalOutputs.Emplace(FuncDef);	
-					}
-				}
-
-				if (DoesOutputPinSupportRead(InOutputPin))
-				{
-					TArray<FShaderFunctionDefinition> ReadFunctions;
-					RawBufferDataInterface->GetSupportedInputs(ReadFunctions);
-					int32 InputIndex = UOptimusRawBufferDataInterface::GetReadValueInputIndex(EOptimusBufferReadType::ForceUAV);
-
-					FShaderFunctionDefinition FuncDef = ReadFunctions[InputIndex];
-					for (int32 Index = 0; Index < FuncDef.ParamTypes.Num(); Index++)
-					{
-						FuncDef.ParamTypes[Index].ResetTypeDeclaration();
+						InterfaceBinding.DataInterfaceBindingIndex = InputIndex;
+						InterfaceBinding.BindingFunctionName = GetReadFunctionName(InOutputPin->GetName());
+					
+						OutInputDataBindings.Add(InKernelSource->ExternalInputs.Num(), InterfaceBinding);
+						InKernelSource->ExternalInputs.Emplace(FuncDef);		
 					}
-				
-					FOptimus_InterfaceBinding InterfaceBinding;
-					InterfaceBinding.DataInterface = RawBufferDataInterface;
-					InterfaceBinding.DataInterfaceBindingIndex = InputIndex;
-					InterfaceBinding.BindingFunctionName = GetReadFunctionName(InOutputPin->GetName());
-				
-					OutInputDataBindings.Add(InKernelSource->ExternalInputs.Num(), InterfaceBinding);
-					InKernelSource->ExternalInputs.Emplace(FuncDef);		
 				}
 			}
 		}
@@ -902,47 +817,22 @@ void UOptimusNode_ComputeKernelBase::ProcessOutputPinForComputeKernel(
 			FString::Printf(TEXT("void Write%s(%s, %s Value) { }"),
 				*InOutputPin->GetName(), *FString::Join(StubIndexes, TEXT(", ")), *ValueType->ToString()));
 	}
-
-	// Gather the constants
-	if (InOutputPin->GetDataDomain().IsOneDimensional() && ensure(InOutputPin->GetDataDomain().IsFullyDefined()))
-	{
-		if (TOptional<FString> Expression = InOutputPin->GetDataDomain().AsExpression())
-		{
-			// Get the component binding from the upstream connection.
-			int32 ComponentBindingIndex = ComponentBinding->GetIndex();
-			
-			OutKernelConstantContainer.AddToKernelContainer({
-				{this, NAME_None, InOutputPin->GetFName()},
-				{*Expression},
-				ComponentBindingIndex,
-				EOptimusConstantType::Output});
-		}
-	}
 }
 
 void UOptimusNode_ComputeKernelBase::BindKernelDataInterfaceForComputeKernel(
 	const UOptimusComponentSourceBinding* InKernelPrimaryComponentSourceBinding,
 	UComputeDataInterface* InOutKernelDataInterface,
 	UOptimusKernelSource* InKernelSource,
-	FOptimus_InterfaceBindingMap& OutInputDataBindings,
-	FOptimusKernelConstantContainer& OutKernelConstantContainer
+	FOptimus_InterfaceBindingMap& OutInputDataBindings
 	) const
 {
 	check(InOutKernelDataInterface);
 	check(InKernelPrimaryComponentSourceBinding);
-
-	FOptimusConstantIdentifier ExecutionDomainIdentifier(this, NAME_None, TEXT("@ExecutionDomain"));
-	
-	OutKernelConstantContainer.AddToKernelContainer({
-		ExecutionDomainIdentifier,
-		{GetExecutionDomain().AsExpression()},
-		InKernelPrimaryComponentSourceBinding->GetIndex(),
-		EOptimusConstantType::Output});
-	
 	IOptimusComputeKernelDataInterface* KernelDataInterface = Cast<IOptimusComputeKernelDataInterface>(InOutKernelDataInterface);
 	if (ensure(KernelDataInterface))
 	{
-		KernelDataInterface->SetExecutionDomainConstant(ExecutionDomainIdentifier);
+		KernelDataInterface->SetExecutionDomain(GetExecutionDomain().AsExpression());
+		KernelDataInterface->SetComponentBinding(InKernelPrimaryComponentSourceBinding);
 	}
 
 	
diff --git a/Engine/Plugins/Animation/DeformerGraph/Source/OptimusCore/Private/Nodes/OptimusNode_ComputeKernelBase.h b/Engine/Plugins/Animation/DeformerGraph/Source/OptimusCore/Private/Nodes/OptimusNode_ComputeKernelBase.h
--- a/Engine/Plugins/Animation/DeformerGraph/Source/OptimusCore/Private/Nodes/OptimusNode_ComputeKernelBase.h
+++ b/Engine/Plugins/Animation/DeformerGraph/Source/OptimusCore/Private/Nodes/OptimusNode_ComputeKernelBase.h
@@ -39,14 +39,12 @@ public:
 	FOptimus_ComputeKernelResult CreateComputeKernel(
 		UObject* InKernelSourceOuter,
 		const FOptimusPinTraversalContext& InTraversalContext,
-		const FOptimus_NodeToDataInterfaceMap& InNodeDataInterfaceMap,
-		const FOptimus_PinToDataInterfaceMap& InLinkDataInterfaceMap,
+		const FOptimus_KernelInputMap InKernelInputs,
+		const FOptimus_KernelOutputMap InKernelOutputs,
 		const TArray<const UOptimusNode*>& InValueNodes,
-		const UComputeDataInterface* InGraphDataInterface,
 		UComputeDataInterface* InOutKernelDataInterface,
 		FOptimus_InterfaceBindingMap& OutInputDataBindings,
-		FOptimus_InterfaceBindingMap& OutOutputDataBindings,
-		FOptimusKernelConstantContainer& OutKernelConstantContainer
+		FOptimus_InterfaceBindingMap& OutOutputDataBindings
 	) const override;
 
 	FOptimusExecutionDomain GetExecutionDomain() const override PURE_VIRTUAL(UOptimusNode_ComputeKernelBase::GetExecutionDomain, return {}; );
@@ -54,7 +52,6 @@ public:
 	UComputeDataInterface* MakeKernelDataInterface(UObject* InOuter) const override PURE_VIRTUAL(UOptimusNode_ComputeKernelBase::MakeKernelDataInterface, return {}; );
 	bool DoesOutputPinSupportAtomic(const UOptimusNodePin* InPin) const override PURE_VIRTUAL(UOptimusNode_ComputeKernelBase::DoesOutputPinSupportAtomic, return false; );
 	bool DoesOutputPinSupportRead(const UOptimusNodePin* InPin) const override PURE_VIRTUAL(UOptimusNode_ComputeKernelBase::DoesOutputPinSupportRead, return false; );
-	bool HasMutableInput() const override;
 	
 	// -- UOptimusNode overrides
 	TOptional<FText> ValidateForCompile() const override;
@@ -111,34 +108,28 @@ private:
 		const FOptimusPinTraversalContext& InTraversalContext,
 		const UOptimusNodePin* InInputPin,
 		const FString& InGroupName,
-		const FOptimus_NodeToDataInterfaceMap& InNodeDataInterfaceMap,
-		const FOptimus_PinToDataInterfaceMap& InLinkDataInterfaceMap,
+		const FOptimus_KernelInputMap& InKernelInputs,
 		const TArray<const UOptimusNode*>& InValueNodes,
-		const UComputeDataInterface* InGraphDataInterface,
 		UOptimusKernelSource* InKernelSource,
 		TArray<FString>& OutGeneratedFunctions,
-		FOptimus_InterfaceBindingMap& OutInputDataBindings,
-		FOptimusKernelConstantContainer& OutKernelConstantContainer
+		FOptimus_InterfaceBindingMap& OutInputDataBindings
 		) const;
 
 	void ProcessOutputPinForComputeKernel(
 		const FOptimusPinTraversalContext& InTraversalContext,
 		const UOptimusNodePin* InOutputPin,
-		const FOptimus_NodeToDataInterfaceMap& InNodeDataInterfaceMap,
-		const FOptimus_PinToDataInterfaceMap& InLinkDataInterfaceMap,
+		const FOptimus_KernelOutputMap& InKernelOutputs,
 		UOptimusKernelSource* InKernelSource,
 		TArray<FString>& OutGeneratedFunctions,
 		FOptimus_InterfaceBindingMap& OutInputDataBindings,
-		FOptimus_InterfaceBindingMap& OutOutputDataBindings,
-		FOptimusKernelConstantContainer& OutKernelConstantContainer
+		FOptimus_InterfaceBindingMap& OutOutputDataBindings
 		) const;
 
 	void BindKernelDataInterfaceForComputeKernel(
 		const UOptimusComponentSourceBinding* InKernelPrimaryComponentSourceBinding,
 		UComputeDataInterface* InOutKernelDataInterface,
 		UOptimusKernelSource* InKernelSource,
-		FOptimus_InterfaceBindingMap& OutInputDataBindings,
-		FOptimusKernelConstantContainer& OutKernelConstantContainer
+		FOptimus_InterfaceBindingMap& OutInputDataBindings
 		) const;
 };
 
diff --git a/Engine/Plugins/Animation/DeformerGraph/Source/OptimusCore/Private/Nodes/OptimusNode_ComputeKernelFunction.h b/Engine/Plugins/Animation/DeformerGraph/Source/OptimusCore/Private/Nodes/OptimusNode_ComputeKernelFunction.h
--- a/Engine/Plugins/Animation/DeformerGraph/Source/OptimusCore/Private/Nodes/OptimusNode_ComputeKernelFunction.h
+++ b/Engine/Plugins/Animation/DeformerGraph/Source/OptimusCore/Private/Nodes/OptimusNode_ComputeKernelFunction.h
@@ -81,7 +81,6 @@ public:
 	UComputeDataInterface* MakeKernelDataInterface(UObject* InOuter) const override { return nullptr; };
 	bool DoesOutputPinSupportAtomic(const UOptimusNodePin* InPin) const override {return false;};
 	bool DoesOutputPinSupportRead(const UOptimusNodePin* InPin) const override {return false;};
-	bool HasMutableInput() const override {return false;};
 private:
 	UOptimusNode_ComputeKernelFunctionGeneratorClass *GetGeneratorClass() const;
 };
diff --git a/Engine/Plugins/Animation/DeformerGraph/Source/OptimusCore/Private/Nodes/OptimusNode_ConstantValue.h b/Engine/Plugins/Animation/DeformerGraph/Source/OptimusCore/Private/Nodes/OptimusNode_ConstantValue.h
--- a/Engine/Plugins/Animation/DeformerGraph/Source/OptimusCore/Private/Nodes/OptimusNode_ConstantValue.h
+++ b/Engine/Plugins/Animation/DeformerGraph/Source/OptimusCore/Private/Nodes/OptimusNode_ConstantValue.h
@@ -78,7 +78,7 @@ public:
 		) const override;
 
 	// IOptimusPinMutabilityDefiner implementation
-	bool IsOutputPinMutable(const UOptimusNodePin* InPin) const override { return false; };
+	EOptimusPinMutability GetOutputPinMutability(const UOptimusNodePin* InPin) const override { return EOptimusPinMutability::Immutable; };
 	
 	
 protected:
diff --git a/Engine/Plugins/Animation/DeformerGraph/Source/OptimusCore/Private/Nodes/OptimusNode_CustomComputeKernel.cpp b/Engine/Plugins/Animation/DeformerGraph/Source/OptimusCore/Private/Nodes/OptimusNode_CustomComputeKernel.cpp
--- a/Engine/Plugins/Animation/DeformerGraph/Source/OptimusCore/Private/Nodes/OptimusNode_CustomComputeKernel.cpp
+++ b/Engine/Plugins/Animation/DeformerGraph/Source/OptimusCore/Private/Nodes/OptimusNode_CustomComputeKernel.cpp
@@ -16,6 +16,7 @@
 #include "ComponentSources/OptimusSkinnedMeshComponentSource.h"
 #include "DataInterfaces/OptimusDataInterfaceCustomComputeKernel.h"
 #include "IOptimusDeprecatedExecutionDataInterface.h"
+#include "IOptimusUnnamedNodePinProvider.h"
 #include "OptimusNode_ResourceAccessorBase.h"
 #include "Engine/UserDefinedStruct.h"
 
@@ -435,7 +436,8 @@ TArray<IOptimusNodeAdderPinProvider::FAdderPinAction> UOptimusNode_CustomCompute
 
 TArray<UOptimusNodePin*> UOptimusNode_CustomComputeKernel::TryAddPinFromPin(
 	const FAdderPinAction& InSelectedAction,
-	UOptimusNodePin* InSourcePin
+	UOptimusNodePin* InSourcePin,
+	FName InNameToUse
 	)
 {
 	TArray<UOptimusNodePin*> AddedPins;
@@ -526,7 +528,7 @@ TArray<UOptimusNodePin*> UOptimusNode_CustomComputeKernel::TryAddPinFromPin(
 	}
 	
 	FOptimusParameterBinding Binding;
-	Binding.Name = GetSanitizedNewPinName(ParentPin, InSourcePin->GetFName());
+	Binding.Name = GetSanitizedNewPinName(ParentPin, InNameToUse);
 	Binding.DataType = {InSourcePin->GetDataType()};
 	Binding.DataDomain = InSourcePin->GetDataDomain();
 	
@@ -627,13 +629,13 @@ FName UOptimusNode_CustomComputeKernel::GetSanitizedNewPinName(
 	
 	if (InParentPin && InParentPin != PrimaryGroupPin)
 	{
-		NewName = Optimus::GetUniqueNameForScope(InParentPin, NewName);
+		NewName = GetAvailablePinNameStable(InParentPin, NewName);
 	}
 	else
 	{
 		// Primary Input and Output need to name check against node scope as well as primary group pin scope
-		NewName = Optimus::GetUniqueNameForScope(this, NewName);
-		NewName = Optimus::GetUniqueNameForScope(PrimaryGroupPin, NewName);
+		NewName = GetAvailablePinNameStable(this, NewName);
+		NewName = GetAvailablePinNameStable(PrimaryGroupPin, NewName);
 	}
 
 
@@ -1354,41 +1356,75 @@ void UOptimusNode_CustomComputeKernel::PropertyArrayItemMoved(
 		TFunction<bool(const UOptimusNodePin *)> InPinPredicate
 		)
 	{
-		int32 BindingIndex = 0;
+		FName NameToMove = NAME_None;
 		
 		// Find the first entry that's different. That's an element we can consider moved. Since array move only
 		// deals with a single item moving either forward or backward.
-		for (const UOptimusNodePin* Pin: InPins)
 		{
-			if (InPinPredicate(Pin))
+			int32 DivergeIndex = INDEX_NONE;
+			FName PinNameAtDiverge = NAME_None;
+			FName BindingNameAtDiverge = NAME_None;
+			
+			int32 BindingIndex = 0;
+			for (const UOptimusNodePin* Pin: InPins)
 			{
-				if (Pin->GetFName() != BindingNameArray[BindingIndex])
-				{
-					break;
-				}
-				
-				BindingIndex++;
-				
-				if (BindingIndex == BindingNameArray.Num())
+				if (InPinPredicate(Pin))
 				{
-					// Nothing got moved.
-					return;
+					if (Pin->GetFName() != BindingNameArray[BindingIndex] && DivergeIndex == INDEX_NONE)
+					{
+						DivergeIndex = BindingIndex;
+						PinNameAtDiverge = Pin->GetFName();
+						BindingNameAtDiverge = BindingNameArray[BindingIndex];
+						BindingIndex++;
+						continue;
+					}
+
+					if (DivergeIndex != INDEX_NONE)
+					{
+						if (BindingNameAtDiverge == Pin->GetFName())
+						{
+							NameToMove = PinNameAtDiverge;
+						}
+						else if (ensure(BindingNameArray[BindingIndex] == PinNameAtDiverge))
+						{
+							NameToMove = BindingNameAtDiverge;
+						}
+
+						break;
+					}
+					
+					BindingIndex++;
+					
+					
+					if (BindingIndex == BindingNameArray.Num())
+					{
+						// Nothing got moved.
+						return;
+					}
 				}
 			}
 		}
 
-		UOptimusNodePin* MovedPin = *InPins.FindByPredicate([Name=BindingNameArray[BindingIndex]](const UOptimusNodePin* InPin)
+		UOptimusNodePin* MovedPin = *InPins.FindByPredicate([NameToMove](const UOptimusNodePin* InPin)
 		{
-			return InPin->GetFName() == Name;
+			return InPin->GetFName() == NameToMove;
 		});
 
+		int32 NameIndex = BindingNameArray.IndexOfByPredicate([NameToMove](FName InBindingName)
+		{
+			return InBindingName == NameToMove;
+		});
+
+		const int32 NextNameIndex = NameIndex + 1;
+
+		const FName NextName = BindingNameArray.IsValidIndex(NextNameIndex) ? BindingNameArray[NextNameIndex] : NAME_None;
+
 		const UOptimusNodePin* NextPin = nullptr;
-		const FName NextPinName = (BindingIndex < (BindingNameArray.Num() - 1)) ? BindingNameArray[BindingIndex + 1] : NAME_None;
-		if (!NextPinName.IsNone())
+		if (!NextName.IsNone())
 		{
-			NextPin = *InPins.FindByPredicate([NextPinName](const UOptimusNodePin* InPin)
+			NextPin = *InPins.FindByPredicate([NextName](const UOptimusNodePin* InPin)
 			{
-				return InPin->GetFName() == NextPinName;
+				return InPin->GetFName() == NextName;
 			});
 		}
 
diff --git a/Engine/Plugins/Animation/DeformerGraph/Source/OptimusCore/Private/Nodes/OptimusNode_CustomComputeKernel.h b/Engine/Plugins/Animation/DeformerGraph/Source/OptimusCore/Private/Nodes/OptimusNode_CustomComputeKernel.h
--- a/Engine/Plugins/Animation/DeformerGraph/Source/OptimusCore/Private/Nodes/OptimusNode_CustomComputeKernel.h
+++ b/Engine/Plugins/Animation/DeformerGraph/Source/OptimusCore/Private/Nodes/OptimusNode_CustomComputeKernel.h
@@ -87,8 +87,9 @@ public:
 		) const override;
 
 	TArray<UOptimusNodePin*> TryAddPinFromPin(
-		const FAdderPinAction& InSelectedAction,
-		UOptimusNodePin* InSourcePin
+		const FAdderPinAction& InSelectedAction, 
+		UOptimusNodePin* InSourcePin, 
+		FName InNameToUse
 		) override;
 	
 	bool RemoveAddedPins(
diff --git a/Engine/Plugins/Animation/DeformerGraph/Source/OptimusCore/Private/Nodes/OptimusNode_DataInterface.cpp b/Engine/Plugins/Animation/DeformerGraph/Source/OptimusCore/Private/Nodes/OptimusNode_DataInterface.cpp
--- a/Engine/Plugins/Animation/DeformerGraph/Source/OptimusCore/Private/Nodes/OptimusNode_DataInterface.cpp
+++ b/Engine/Plugins/Animation/DeformerGraph/Source/OptimusCore/Private/Nodes/OptimusNode_DataInterface.cpp
@@ -11,6 +11,9 @@
 #include "OptimusObjectVersion.h"
 
 #include "ComputeFramework/ShaderParamTypeDefinition.h"
+#include "Serialization/ObjectAndNameAsStringProxyArchive.h"
+#include "Serialization/MemoryReader.h"
+#include "Serialization/MemoryWriter.h"
 
 #include UE_INLINE_GENERATED_CPP_BY_NAME(OptimusNode_DataInterface)
 
@@ -94,6 +97,27 @@ TOptional<FText> UOptimusNode_DataInterface::ValidateForCompile() const
 	return {};
 }
 
+void UOptimusNode_DataInterface::SaveState(FArchive& Ar) const
+{
+	Super::SaveState(Ar);
+	// This fella does the heavy lifting of serializing object references. 
+	// FMemoryWriter and fam do not handle UObject* serialization on their own.
+	FObjectAndNameAsStringProxyArchive NodeProxyArchive(
+			Ar, /* bInLoadIfFindFails=*/ false);
+	DataInterfaceData->SerializeScriptProperties(NodeProxyArchive);	
+}
+
+void UOptimusNode_DataInterface::RestoreState(FArchive& Ar)
+{
+	Super::RestoreState(Ar);
+
+	DataInterfaceData = NewObject<UOptimusComputeDataInterface>(this, DataInterfaceClass);
+	
+	FObjectAndNameAsStringProxyArchive NodeProxyArchive(
+			Ar, /* bInLoadIfFindFails=*/true);
+	DataInterfaceData->SerializeScriptProperties(NodeProxyArchive);
+}
+
 bool UOptimusNode_DataInterface::IsComponentSourceCompatible(const UOptimusComponentSource* InComponentSource) const
 {
 	return InComponentSource && InComponentSource->GetComponentClass()->IsChildOf(DataInterfaceData->GetRequiredComponentClass());
@@ -197,7 +221,7 @@ UOptimusComponentSourceBinding* UOptimusNode_DataInterface::GetComponentBinding(
 	return nullptr;
 }
 
-bool UOptimusNode_DataInterface::IsOutputPinMutable(const UOptimusNodePin* InPin) const
+EOptimusPinMutability UOptimusNode_DataInterface::GetOutputPinMutability(const UOptimusNodePin* InPin) const
 {
 	const TArray<FOptimusCDIPinDefinition> PinDefinitions = DataInterfaceData->GetPinDefinitions();
 
@@ -212,10 +236,10 @@ bool UOptimusNode_DataInterface::IsOutputPinMutable(const UOptimusNodePin* InPin
 	}
 	if (!ensure(PinDefinitionIndex != INDEX_NONE))
 	{
-		return true;
+		return EOptimusPinMutability::Mutable;
 	}
 
-	return PinDefinitions[PinDefinitionIndex].bMutable;
+	return PinDefinitions[PinDefinitionIndex].bMutable ? EOptimusPinMutability::Mutable : EOptimusPinMutability::Immutable;
 }
 
 
diff --git a/Engine/Plugins/Animation/DeformerGraph/Source/OptimusCore/Private/Nodes/OptimusNode_DataInterface.h b/Engine/Plugins/Animation/DeformerGraph/Source/OptimusCore/Private/Nodes/OptimusNode_DataInterface.h
--- a/Engine/Plugins/Animation/DeformerGraph/Source/OptimusCore/Private/Nodes/OptimusNode_DataInterface.h
+++ b/Engine/Plugins/Animation/DeformerGraph/Source/OptimusCore/Private/Nodes/OptimusNode_DataInterface.h
@@ -48,7 +48,7 @@ public:
 	UOptimusComponentSourceBinding* GetComponentBinding() const override;
 
 	// -- IOptimusPinMutabilityDefiner
-	bool IsOutputPinMutable(const UOptimusNodePin* InPin) const override;
+	EOptimusPinMutability GetOutputPinMutability(const UOptimusNodePin* InPin) const override;
 	
 protected:
 	// -- UOptimusNode overrides
@@ -56,6 +56,9 @@ protected:
 	bool ValidateConnection(const UOptimusNodePin& InThisNodesPin, const UOptimusNodePin& InOtherNodesPin, FString* OutReason) const override;
 	TOptional<FText> ValidateForCompile() const override;
 
+	void SaveState(FArchive& Ar) const override;
+	void RestoreState(FArchive& Ar) override;
+
 	// -- UObject overrides
 	void PostLoad() override;
 	void PostDuplicate(EDuplicateMode::Type DuplicateMode) override;
diff --git a/Engine/Plugins/Animation/DeformerGraph/Source/OptimusCore/Private/Nodes/OptimusNode_GetVariable.h b/Engine/Plugins/Animation/DeformerGraph/Source/OptimusCore/Private/Nodes/OptimusNode_GetVariable.h
--- a/Engine/Plugins/Animation/DeformerGraph/Source/OptimusCore/Private/Nodes/OptimusNode_GetVariable.h
+++ b/Engine/Plugins/Animation/DeformerGraph/Source/OptimusCore/Private/Nodes/OptimusNode_GetVariable.h
@@ -57,7 +57,7 @@ public:
 	FShaderValueType::FValue GetShaderValue() const override;
 
 	//IOptimusPinMutabilityDefiner overrides 
-	bool IsOutputPinMutable(const UOptimusNodePin* InPin) const override { return true; };
+	EOptimusPinMutability GetOutputPinMutability(const UOptimusNodePin* InPin) const override { return EOptimusPinMutability::Mutable; };
 	
 protected:
 	void ConstructNode() override;
diff --git a/Engine/Plugins/Animation/DeformerGraph/Source/OptimusCore/Private/Nodes/OptimusNode_LoopTerminal.cpp b/Engine/Plugins/Animation/DeformerGraph/Source/OptimusCore/Private/Nodes/OptimusNode_LoopTerminal.cpp
new file mode 100644
--- /dev/null
+++ b/Engine/Plugins/Animation/DeformerGraph/Source/OptimusCore/Private/Nodes/OptimusNode_LoopTerminal.cpp
@@ -0,0 +1,723 @@
+﻿// Copyright Epic Games, Inc. All Rights Reserved.
+
+#include "OptimusNode_LoopTerminal.h"
+
+#include "OptimusDataTypeRegistry.h"
+#include "OptimusHelpers.h"
+#include "OptimusNodeGraph.h"
+#include "ComponentSources/OptimusSkeletalMeshComponentSource.h"
+#include "DataInterfaces/OptimusDataInterfaceLoopTerminal.h"
+
+#include UE_INLINE_GENERATED_CPP_BY_NAME(OptimusNode_LoopTerminal)
+
+#define LOCTEXT_NAMESPACE "OptimusNodeLoopTerminal"
+
+
+static FName GetHiddenPinName(FName InBindingName)
+{
+	return *(TEXT("(Hidden)") + InBindingName.ToString());
+}
+
+UOptimusNode_LoopTerminal::UOptimusNode_LoopTerminal()
+{
+	EnableDynamicPins();
+}
+#if WITH_EDITOR
+void UOptimusNode_LoopTerminal::PostEditChangeProperty(FPropertyChangedEvent& PropertyChangedEvent)
+{
+	Super::PostEditChangeProperty(PropertyChangedEvent);
+
+	if (PropertyChangedEvent.ChangeType & EPropertyChangeType::ValueSet)
+	{
+		if (CastField<FArrayProperty>(PropertyChangedEvent.Property) &&
+			PropertyChangedEvent.GetArrayIndex(PropertyChangedEvent.GetPropertyName().ToString()) == INDEX_NONE)
+		{
+			PropertyArrayPasted(PropertyChangedEvent);
+		}
+		else
+		{
+			PropertyValueChanged(PropertyChangedEvent);
+		}
+		
+	}
+	else if (PropertyChangedEvent.ChangeType & EPropertyChangeType::ArrayAdd)
+	{
+		PropertyArrayItemAdded(PropertyChangedEvent);
+	}
+	else if (PropertyChangedEvent.ChangeType & EPropertyChangeType::ArrayRemove)
+	{
+		PropertyArrayItemRemoved(PropertyChangedEvent);
+	}
+	else if (PropertyChangedEvent.ChangeType & EPropertyChangeType::ArrayClear)
+	{
+		PropertyArrayCleared(PropertyChangedEvent);
+	}
+	else if (PropertyChangedEvent.ChangeType & EPropertyChangeType::ArrayMove)
+	{
+		PropertyArrayItemMoved(PropertyChangedEvent);
+	}
+
+}
+#endif
+
+
+FText UOptimusNode_LoopTerminal::GetDisplayName() const
+{
+	switch(TerminalType)
+	{
+	case EOptimusTerminalType::Entry:
+		return LOCTEXT("LoopTerminalType_Entry", "Loop Entry"); 
+	case EOptimusTerminalType::Return:
+		return LOCTEXT("LoopTerminalType_Return", "Loop Return");
+	case EOptimusTerminalType::Unknown:
+		checkNoEntry();
+	}
+	return FText();
+}
+
+
+void UOptimusNode_LoopTerminal::ConstructNode()
+{
+	if (TerminalType == EOptimusTerminalType::Entry)
+	{
+		FOptimusDataTypeHandle UintType = FOptimusDataTypeRegistry::Get().FindType(*FUInt32Property::StaticClass());
+		IndexPin = AddPinDirect(TEXT("Index"),  EOptimusNodePinDirection::Output, FOptimusDataDomain(), UintType);
+		CountPin = AddPinDirect(TEXT("Count"),  EOptimusNodePinDirection::Output, FOptimusDataDomain(), UintType);
+	}
+
+	
+	if (FOptimusLoopTerminalInfo* LoopInfoPtr = GetLoopInfo())
+	{
+		PinPairInfos.Reset();
+		for (const FOptimusParameterBinding& Binding : LoopInfoPtr->Bindings)
+		{
+			AddPinPairsDirect(Binding);
+		}
+	}
+}
+
+bool UOptimusNode_LoopTerminal::ValidateConnection(const UOptimusNodePin& InThisNodesPin, const UOptimusNodePin& InOtherNodesPin, FString* OutReason) const
+{
+	if (InThisNodesPin.GetDirection() == EOptimusNodePinDirection::Input)
+	{
+		if (TerminalType == EOptimusTerminalType::Entry)
+		{
+			if (InOtherNodesPin.GetOwningNode() == GetOtherTerminal())
+			{
+				if (OutReason)
+				{
+					*OutReason = TEXT("Cannot connect output of return terminal to input of entry terminal");
+				}
+				return false;
+			}
+		}
+	}
+
+	return true;
+}
+
+TArray<IOptimusNodeAdderPinProvider::FAdderPinAction> UOptimusNode_LoopTerminal::GetAvailableAdderPinActions(
+	const UOptimusNodePin* InSourcePin, EOptimusNodePinDirection InNewPinDirection, FString* OutReason) const
+{
+	if (InSourcePin->GetDataDomain().IsSingleton())
+	{
+		if (OutReason)
+		{
+			*OutReason = TEXT("Can't add parameter pin ");
+		}
+
+		return {};
+	}
+	
+	FAdderPinAction Action;
+	Action.bCanAutoLink = InSourcePin->GetDirection() != InNewPinDirection;
+	Action.NewPinDirection = InNewPinDirection;
+	return {Action};
+}
+
+TArray<UOptimusNodePin*> UOptimusNode_LoopTerminal::TryAddPinFromPin(const FAdderPinAction& InSelectedAction, UOptimusNodePin* InSourcePin, FName InNameToUse)
+{
+	FOptimusParameterBinding Binding;
+	Binding.Name = InNameToUse;
+	Binding.DataType = {InSourcePin->GetDataType()};
+	Binding.DataDomain = InSourcePin->GetDataDomain();	
+
+	SanitizeBinding(Binding, NAME_None);
+
+	check(GetLoopInfo());
+	GetLoopInfo()->Bindings.InnerArray.Add(Binding);
+
+	TArray<UOptimusNodePin*> AddedPins = GetTerminalByType(EOptimusTerminalType::Entry)->AddPinPairsDirect(Binding);	
+	AddedPins.Append(GetTerminalByType(EOptimusTerminalType::Return)->AddPinPairsDirect(Binding));	
+
+	// The last pin should be the one to auto link
+	if (TerminalType == EOptimusTerminalType::Entry)
+	{
+		if (InSelectedAction.NewPinDirection == EOptimusNodePinDirection::Input)
+		{
+			AddedPins.Swap(0, 3);
+		}
+		else
+		{
+			AddedPins.Swap(1, 3);
+		}
+	}
+	else
+	{
+		if (InSelectedAction.NewPinDirection == EOptimusNodePinDirection::Input)
+		{
+			AddedPins.Swap(2, 3);
+		}
+	}
+
+	return AddedPins;
+}
+
+bool UOptimusNode_LoopTerminal::RemoveAddedPins(TConstArrayView<UOptimusNodePin*> InAddedPinsToRemove)
+{
+	if (!ensure(!InAddedPinsToRemove.IsEmpty()))
+	{
+		return false;
+	}
+
+	int32 PairIndex = GetPairIndex(InAddedPinsToRemove[0]);
+
+	if (!ensure(PairIndex != INDEX_NONE))
+	{
+		return false;
+	}
+
+	GetTerminalByType(EOptimusTerminalType::Entry)->RemovePinPairDirect(PairIndex);
+	GetTerminalByType(EOptimusTerminalType::Return)->RemovePinPairDirect(PairIndex);
+	check(GetLoopInfo());
+	GetLoopInfo()->Bindings.InnerArray.RemoveAt(PairIndex);;
+
+	return true;
+}
+
+bool UOptimusNode_LoopTerminal::IsPinNameHidden(UOptimusNodePin* InPin) const
+{
+	if (!InPin->GetDataDomain().IsSingleton())
+	{
+		if (TerminalType == EOptimusTerminalType::Entry && InPin->GetDirection() == EOptimusNodePinDirection::Output)
+		{
+			return true;
+		}
+		
+		if (TerminalType == EOptimusTerminalType::Return && InPin->GetDirection() == EOptimusNodePinDirection::Input)
+		{
+			return true;
+		}
+	}
+
+	return false;
+}
+
+FName UOptimusNode_LoopTerminal::GetNameForAdderPin(UOptimusNodePin* InPin) const
+{
+	if (IsPinNameHidden(InPin))
+	{
+		return GetPinCounterpart(InPin, TerminalType)->GetFName();
+	}
+
+	return InPin->GetFName();
+}
+
+EOptimusPinMutability UOptimusNode_LoopTerminal::GetOutputPinMutability(const UOptimusNodePin* InPin) const
+{
+	if (InPin == IndexPin || InPin == CountPin)
+	{
+		return EOptimusPinMutability::Immutable;
+	}
+
+	return EOptimusPinMutability::Undefined;
+}
+
+void UOptimusNode_LoopTerminal::PairToCounterpartNode(const IOptimusNodePairProvider* NodePairProvider)
+{
+	if (TerminalType == EOptimusTerminalType::Return)
+	{
+		ConstructNode();
+	}
+}
+
+UOptimusNodePin* UOptimusNode_LoopTerminal::GetPinCounterpart(const UOptimusNodePin* InNodePin, EOptimusTerminalType InTerminalType, TOptional<EOptimusNodePinDirection> InDirection) const
+{
+	const int32 PairIndex = GetPairIndex(InNodePin);
+
+	const UOptimusNode_LoopTerminal* Terminal = GetTerminalByType(InTerminalType);
+	TArray<UOptimusNodePin*> PairedPins = Terminal->GetPairedPins(Terminal->PinPairInfos[PairIndex]);
+
+	EOptimusNodePinDirection CounterpartDirection =
+		InNodePin->GetDirection() == EOptimusNodePinDirection::Input ?
+		EOptimusNodePinDirection::Output : EOptimusNodePinDirection::Input;
+
+	if (InDirection)
+	{
+		CounterpartDirection = *InDirection;
+	}
+	
+	return CounterpartDirection == EOptimusNodePinDirection::Input ? PairedPins[0] : PairedPins[1];
+}
+
+UOptimusNode_LoopTerminal* UOptimusNode_LoopTerminal::GetOtherTerminal() const
+{
+	return Cast<UOptimusNode_LoopTerminal>(GetOwningGraph()->GetNodeCounterpart(this));
+}
+
+int32 UOptimusNode_LoopTerminal::GetLoopCount() const
+{
+	check(GetLoopInfo());
+	return GetLoopInfo()->Count;
+}
+
+EOptimusTerminalType UOptimusNode_LoopTerminal::GetTerminalType() const
+{
+	return TerminalType;
+}
+
+int32 UOptimusNode_LoopTerminal::GetDataFunctionIndexFromPin(const UOptimusNodePin* InPin)
+{
+	TArray<FOptimusCDIPinDefinition> PinDefinitions = UOptimusLoopTerminalDataInterface::StaticClass()->GetDefaultObject<UOptimusLoopTerminalDataInterface>()->GetPinDefinitions();
+	int32 PinIndex = INDEX_NONE;
+	for (int32 Index = 0 ; Index < PinDefinitions.Num(); ++Index)
+	{
+		if (InPin->GetUniqueName() == PinDefinitions[Index].PinName)
+		{
+			PinIndex = Index;
+			break;
+		}
+	}
+	
+	if (!ensure(PinIndex != INDEX_NONE))
+	{
+		return INDEX_NONE;
+	}
+
+	return PinIndex;
+}
+
+TArray<UOptimusNodePin*> UOptimusNode_LoopTerminal::AddPinPairs(const FOptimusParameterBinding& InBinding)
+{
+	TArray<UOptimusNodePin*> AddedPins;
+
+	UOptimusNodePin* OutputBeforePin = nullptr;
+	if (IndexPin)
+	{
+		OutputBeforePin = IndexPin;
+	}
+
+	FName InputName = InBinding.Name;
+	FName OutputName = GetHiddenPinName(InBinding.Name);
+
+	if (TerminalType == EOptimusTerminalType::Return)
+	{
+		Swap(InputName, OutputName);
+	}
+	
+	AddedPins.Add(AddPin(InputName, EOptimusNodePinDirection::Input, InBinding.DataDomain, InBinding.DataType, nullptr));
+	AddedPins.Add(AddPin(OutputName, EOptimusNodePinDirection::Output, InBinding.DataDomain, InBinding.DataType, OutputBeforePin, nullptr ));
+
+	PinPairInfos.Add({AddedPins[0]->GetPinNamePath(), AddedPins[1]->GetPinNamePath()});
+
+	return AddedPins;
+}
+
+TArray<UOptimusNodePin*> UOptimusNode_LoopTerminal::AddPinPairsDirect(const FOptimusParameterBinding& InBinding)
+{
+	TArray<UOptimusNodePin*> AddedPins;
+
+	UOptimusNodePin* OutputBeforePin = nullptr;
+	if (IndexPin)
+	{
+		OutputBeforePin = IndexPin;
+	}
+
+	FName InputName = InBinding.Name;
+	FName OutputName = GetHiddenPinName(InBinding.Name);
+
+	if (TerminalType == EOptimusTerminalType::Return)
+	{
+		Swap(InputName, OutputName);
+	}
+	
+	AddedPins.Add(AddPinDirect(InputName, EOptimusNodePinDirection::Input, InBinding.DataDomain, InBinding.DataType, nullptr));
+	AddedPins.Add(AddPinDirect(OutputName, EOptimusNodePinDirection::Output, InBinding.DataDomain, InBinding.DataType, OutputBeforePin, nullptr ));
+
+	PinPairInfos.Add({AddedPins[0]->GetPinNamePath(), AddedPins[1]->GetPinNamePath()});
+
+	return AddedPins;
+}
+
+
+TArray<UOptimusNodePin*> UOptimusNode_LoopTerminal::GetPairedPins(const FOptimusPinPairInfo& InPair) const
+{
+	return {FindPinFromPath(InPair.InputPinPath), FindPinFromPath(InPair.OutputPinPath)};
+}
+
+int32 UOptimusNode_LoopTerminal::GetPairIndex(const UOptimusNodePin* Pin)
+{
+	auto PinPairFilter = [Pin](const FOptimusPinPairInfo& InPinPair)
+	{
+		if (InPinPair.InputPinPath == Pin->GetPinNamePath() || InPinPair.OutputPinPath == Pin->GetPinNamePath())
+		{
+			return true;
+		}
+		return false;
+	};
+	
+	const int32 PairIndex = CastChecked<UOptimusNode_LoopTerminal>(Pin->GetOwningNode())->PinPairInfos.IndexOfByPredicate(PinPairFilter);
+
+	return PairIndex;
+}
+
+void UOptimusNode_LoopTerminal::RemovePinPair(int32 InPairIndex)
+{
+	if (ensure(PinPairInfos.IsValidIndex(InPairIndex)))
+	{
+		FOptimusPinPairInfo PinPair = PinPairInfos[InPairIndex];
+
+		for (UOptimusNodePin* Pin : GetPairedPins(PinPair))
+		{
+			RemovePin(Pin);
+		}
+
+		PinPairInfos.Remove(PinPair);
+	}	
+}
+
+void UOptimusNode_LoopTerminal::RemovePinPairDirect(int32 InPairIndex)
+{
+	if (ensure(PinPairInfos.IsValidIndex(InPairIndex)))
+	{
+		FOptimusPinPairInfo PinPair = PinPairInfos[InPairIndex];
+
+		for (UOptimusNodePin* Pin : GetPairedPins(PinPair))
+		{
+			RemovePinDirect(Pin);
+		}
+
+		PinPairInfos.Remove(PinPair);
+	}	
+}
+
+void UOptimusNode_LoopTerminal::ClearPinPairs()
+{
+	for (const FOptimusPinPairInfo& Pair : PinPairInfos)
+	{
+		for (UOptimusNodePin* Pin : GetPairedPins(Pair))
+		{
+			RemovePin(Pin);
+		}
+	}
+
+	PinPairInfos.Reset();	
+}
+
+void UOptimusNode_LoopTerminal::MovePinPair()
+{
+	FName PinNameToMove = NAME_None;
+	int32 DivergeIndex = INDEX_NONE;
+
+	TArray<TArray<UOptimusNodePin*>> PinPairs;
+
+	FOptimusParameterBindingArray& Bindings = GetLoopInfo()->Bindings;
+
+	check(Bindings.Num() == PinPairInfos.Num())
+	
+	for (const FOptimusPinPairInfo& PinPairInfo : PinPairInfos)
+	{
+		PinPairs.Add(GetPairedPins(PinPairInfo));
+	}
+
+	int32 NamedPinIndex = 0;
+	if (TerminalType == EOptimusTerminalType::Return)
+	{
+		NamedPinIndex = 1;
+	}
+	
+	for (int32 Index = 0; Index < Bindings.Num(); Index++)
+	{
+		if (Bindings[Index].Name != PinPairs[Index][NamedPinIndex]->GetFName() && DivergeIndex == INDEX_NONE)
+		{
+			DivergeIndex = Index;
+			continue;
+		}
+
+		if (DivergeIndex != INDEX_NONE)
+		{
+			if (Bindings[DivergeIndex].Name == PinPairs[Index][NamedPinIndex]->GetFName())
+			{
+				PinNameToMove = PinPairs[DivergeIndex][NamedPinIndex]->GetFName();
+			}
+			else if (ensure(Bindings[Index].Name == PinPairs[DivergeIndex][NamedPinIndex]->GetFName()))
+			{
+				PinNameToMove = Bindings[DivergeIndex].Name;
+			}
+			break;
+		}
+	}
+
+	if (DivergeIndex != INDEX_NONE)
+	{
+		const int32 PinIndex = PinPairs.IndexOfByPredicate([PinNameToMove, NamedPinIndex](const TArray<UOptimusNodePin*>& InPinPair)
+		{
+			return InPinPair[NamedPinIndex]->GetFName() == PinNameToMove;
+		});
+
+		const int32 NameIndex = Bindings.InnerArray.IndexOfByPredicate([PinNameToMove](const FOptimusParameterBinding& InBinding)
+		{
+			return InBinding.Name == PinNameToMove;
+		}); 
+
+		const int32 NextNameIndex = NameIndex + 1;
+		FName NextName = Bindings.IsValidIndex(NextNameIndex) ?  FName(Bindings[NextNameIndex].Name) : NAME_None;
+		int32 NextPinIndex = INDEX_NONE;
+		if (!NextName.IsNone())
+		{
+			NextPinIndex = PinPairs.IndexOfByPredicate([NextName, NamedPinIndex](const TArray<UOptimusNodePin*>& InPinPair)
+			{
+				return InPinPair[NamedPinIndex]->GetFName() == NextName;
+			});	
+		}
+
+		MovePin(PinPairs[PinIndex][0], NextPinIndex != INDEX_NONE ? PinPairs[NextPinIndex][0]: nullptr);
+		MovePin(PinPairs[PinIndex][1], NextPinIndex != INDEX_NONE ? PinPairs[NextPinIndex][1]: IndexPin);
+
+
+		FOptimusPinPairInfo Pair = PinPairInfos[PinIndex];
+
+		FOptimusPinPairInfo PairBefore = NextPinIndex != INDEX_NONE ? PinPairInfos[NextPinIndex] : FOptimusPinPairInfo();
+
+		PinPairInfos.RemoveAt(PinIndex);
+				
+		if (NextPinIndex != INDEX_NONE)
+		{
+			const int32 BeforeIndex = PinPairInfos.IndexOfByKey(PairBefore);
+			PinPairInfos.Insert(Pair, BeforeIndex);
+		}
+		else
+		{
+			PinPairInfos.Add(Pair);
+		}
+	}
+}
+
+void UOptimusNode_LoopTerminal::UpdatePinPairs()
+{
+	// PostEditChangeProperty does not give array index for changes to inner properties of a struct in an array, hence the traversal here
+	TArray<TArray<UOptimusNodePin*>> PinPairs;
+
+	check(GetLoopInfo());
+	const FOptimusParameterBindingArray& Bindings = GetLoopInfo()->Bindings;
+
+	check(Bindings.Num() == PinPairInfos.Num())
+	
+	for (const FOptimusPinPairInfo& PinPairInfo : PinPairInfos)
+	{
+		PinPairs.Add(GetPairedPins(PinPairInfo));
+	}
+
+	int32 NamedPinIndex = 0;
+	int32 UnnamedPinIndex = 1;
+	if (TerminalType == EOptimusTerminalType::Return)
+	{
+		Swap(NamedPinIndex, UnnamedPinIndex);
+	}
+	
+	
+	for (int32 Index = 0 ; Index < Bindings.Num(); Index++)
+	{
+		if (Bindings[Index].Name != PinPairs[Index][NamedPinIndex]->GetFName())
+		{
+			SetPinName(PinPairs[Index][NamedPinIndex], Bindings[Index].Name);
+			SetPinName(PinPairs[Index][UnnamedPinIndex], GetHiddenPinName(Bindings[Index].Name));
+			PinPairInfos[Index] = {PinPairs[Index][0]->GetPinNamePath(), PinPairs[Index][1]->GetPinNamePath()};
+		}
+		
+		if (Bindings[Index].DataType.Resolve() != PinPairs[Index][0]->GetDataType())
+		{
+			SetPinDataType(PinPairs[Index][0], Bindings[Index].DataType);
+			SetPinDataType(PinPairs[Index][1], Bindings[Index].DataType);
+		}
+
+		if (Bindings[Index].DataDomain != PinPairs[Index][0]->GetDataDomain())
+		{
+			SetPinDataDomain(PinPairs[Index][0], Bindings[Index].DataDomain);
+			SetPinDataDomain(PinPairs[Index][1], Bindings[Index].DataDomain);
+		}
+	}
+}
+
+FOptimusLoopTerminalInfo* UOptimusNode_LoopTerminal::GetLoopInfo()
+{
+	if (TerminalType == EOptimusTerminalType::Entry)
+	{
+		return &LoopInfo;
+	}
+	
+	if (UOptimusNode_LoopTerminal* Entry = GetOtherTerminal())
+	{
+		return &Entry->LoopInfo;
+	}
+
+	return nullptr;
+}
+
+const FOptimusLoopTerminalInfo* UOptimusNode_LoopTerminal::GetLoopInfo() const
+{
+	if (TerminalType == EOptimusTerminalType::Entry)
+	{
+		return &LoopInfo;
+	}
+	
+	if (UOptimusNode_LoopTerminal* Entry = GetOtherTerminal())
+	{
+		return &Entry->LoopInfo;
+	}
+
+	return nullptr;
+}
+
+void UOptimusNode_LoopTerminal::SanitizeBinding(FOptimusParameterBinding& InOutBinding, FName InOldName)
+{
+	InOutBinding.Name = GetSanitizedBindingName(InOutBinding.Name, InOldName);
+	
+	if (!InOutBinding.DataType.IsValid())
+	{
+		InOutBinding.DataType = FOptimusDataTypeRegistry::Get().FindType(*FFloatProperty::StaticClass());;
+	}
+
+	if (InOutBinding.DataDomain.IsSingleton())
+	{
+		InOutBinding.DataDomain = FOptimusDataDomain(TArray<FName>({UOptimusSkeletalMeshComponentSource::Domains::Vertex}));
+	}
+}
+
+UOptimusNode_LoopTerminal* UOptimusNode_LoopTerminal::GetTerminalByType(EOptimusTerminalType InType)
+{
+	if (InType == EOptimusTerminalType::Unknown)
+	{
+		return nullptr;
+	}
+	
+	if (TerminalType == InType)
+	{
+		return this;
+	}
+
+	return GetOtherTerminal();
+}
+
+const UOptimusNode_LoopTerminal* UOptimusNode_LoopTerminal::GetTerminalByType(EOptimusTerminalType InType) const
+{
+	if (InType == EOptimusTerminalType::Unknown)
+	{
+		return nullptr;
+	}
+	
+	if (TerminalType == InType)
+	{
+		return this;
+	}
+
+	return GetOtherTerminal();	
+}
+#if WITH_EDITOR
+void UOptimusNode_LoopTerminal::PropertyArrayPasted(const FPropertyChangedEvent& InPropertyChangedEvent)
+{
+	GetOtherTerminal()->Modify();
+	
+	ClearPinPairs();
+	GetOtherTerminal()->ClearPinPairs();
+
+	for(FOptimusParameterBinding& Binding : LoopInfo.Bindings)
+	{
+		SanitizeBinding(Binding, NAME_None);
+		AddPinPairs(Binding);
+		GetOtherTerminal()->AddPinPairs(Binding);
+	}	
+}
+
+void UOptimusNode_LoopTerminal::PropertyValueChanged(const FPropertyChangedEvent& InPropertyChangedEvent)
+{
+	if (InPropertyChangedEvent.Property == FOptimusLoopTerminalInfo::StaticStruct()->FindPropertyByName(GET_MEMBER_NAME_CHECKED(FOptimusLoopTerminalInfo, Count)))
+	{
+		return;
+	}
+
+	check(PinPairInfos.Num() == LoopInfo.Bindings.Num());
+	
+	for (int32 Index = 0; Index < LoopInfo.Bindings.Num(); Index++)
+	{
+		const FName OldName = GetPairedPins(PinPairInfos[Index])[0]->GetFName();
+		SanitizeBinding(LoopInfo.Bindings[Index], OldName);
+	}
+	
+	UpdatePinPairs();
+	GetOtherTerminal()->Modify();
+	GetOtherTerminal()->UpdatePinPairs();
+}
+
+void UOptimusNode_LoopTerminal::PropertyArrayItemAdded(const FPropertyChangedEvent& InPropertyChangedEvent)
+{
+	const int32 ArrayIndex = InPropertyChangedEvent.GetArrayIndex(GET_MEMBER_NAME_STRING_CHECKED(FOptimusParameterBindingArray, InnerArray));
+
+	// We should only support appending to the end currently, given FOptimusParameterBindingArray detail customizaiton
+	if (ensure(ArrayIndex == LoopInfo.Bindings.Num() - 1))
+	{
+		FOptimusParameterBinding& Binding = LoopInfo.Bindings[ArrayIndex];
+
+		SanitizeBinding(Binding, NAME_None);
+		
+		AddPinPairs(Binding);
+		GetOtherTerminal()->Modify();
+		GetOtherTerminal()->AddPinPairs(Binding);
+	}
+}
+
+void UOptimusNode_LoopTerminal::PropertyArrayItemRemoved(const FPropertyChangedEvent& InPropertyChangedEvent)
+{
+	const int32 ArrayIndex = InPropertyChangedEvent.GetArrayIndex(GET_MEMBER_NAME_STRING_CHECKED(FOptimusParameterBindingArray, InnerArray));
+	
+	RemovePinPair(ArrayIndex);
+	GetOtherTerminal()->Modify();
+	GetOtherTerminal()->RemovePinPair(ArrayIndex);
+}
+
+void UOptimusNode_LoopTerminal::PropertyArrayCleared(const FPropertyChangedEvent& InPropertyChangedEvent)
+{
+	ClearPinPairs();
+	GetOtherTerminal()->Modify();
+	GetOtherTerminal()->ClearPinPairs();
+}
+
+void UOptimusNode_LoopTerminal::PropertyArrayItemMoved(const FPropertyChangedEvent& InPropertyChangedEvent)
+{
+	MovePinPair();
+	GetOtherTerminal()->Modify();
+	GetOtherTerminal()->MovePinPair();
+}
+#endif
+
+FName UOptimusNode_LoopTerminal::GetSanitizedBindingName(FName InNewName, FName InOldName)
+{
+	FName Name = InNewName;
+	
+	if (Name == NAME_None)
+	{
+		Name = TEXT("EmptyName");
+	}
+
+	Name = Optimus::GetSanitizedNameForHlsl(Name);
+
+	if (Name != InOldName)
+	{
+		Name = GetAvailablePinNameStable(this, Name);
+	}
+
+	return Name;
+}
+
+
+#undef LOCTEXT_NAMESPACE
diff --git a/Engine/Plugins/Animation/DeformerGraph/Source/OptimusCore/Private/Nodes/OptimusNode_LoopTerminal.h b/Engine/Plugins/Animation/DeformerGraph/Source/OptimusCore/Private/Nodes/OptimusNode_LoopTerminal.h
new file mode 100644
--- /dev/null
+++ b/Engine/Plugins/Animation/DeformerGraph/Source/OptimusCore/Private/Nodes/OptimusNode_LoopTerminal.h
@@ -0,0 +1,146 @@
+﻿// Copyright Epic Games, Inc. All Rights Reserved.
+#pragma once
+
+#include "IOptimusNodeAdderPinProvider.h"
+#include "IOptimusNodePairProvider.h"
+#include "IOptimusPinMutabilityDefiner.h"
+#include "IOptimusUnnamedNodePinProvider.h"
+#include "OptimusBindingTypes.h"
+#include "OptimusNode.h"
+#include "OptimusNode_GraphTerminal.h"
+
+#include "OptimusNode_LoopTerminal.generated.h"
+
+class UOptimusNode_LoopTerminal;
+
+
+USTRUCT()
+struct FOptimusPinPairInfo
+{
+	GENERATED_BODY()
+	// Using PinNamePath here such that it plays well with default UObject undo/redo
+	
+	UPROPERTY()
+	TArray<FName> InputPinPath;
+	
+	UPROPERTY()
+	TArray<FName> OutputPinPath;
+
+	bool operator==(const FOptimusPinPairInfo& InOther) const
+	{
+		return InOther.InputPinPath == InputPinPath && InOther.OutputPinPath == OutputPinPath;
+	}
+
+};
+
+USTRUCT()
+struct FOptimusLoopTerminalInfo
+{
+	GENERATED_BODY()
+
+	UPROPERTY(EditAnywhere, Category="Loop Terminal")
+	int32 Count = 1;
+
+	UPROPERTY(EditAnywhere, Category="Loop Terminal", meta=(FullyExpand="true", NoResetToDefault))
+	FOptimusParameterBindingArray Bindings;	
+};
+
+UCLASS(Hidden)
+class OPTIMUSCORE_API UOptimusNode_LoopTerminal :
+	public UOptimusNode,
+	public IOptimusNodeAdderPinProvider,
+	public IOptimusUnnamedNodePinProvider,
+	public IOptimusNodePairProvider,
+	public IOptimusPinMutabilityDefiner
+{
+	GENERATED_BODY()
+	
+public:
+	UOptimusNode_LoopTerminal();
+
+#if WITH_EDITOR
+	// UObject overrides
+	void PostEditChangeProperty(FPropertyChangedEvent& PropertyChangedEvent) override;
+#endif
+	
+	// UOptimusNode overrides
+	bool CanUserDeleteNode() const override { return true; }
+	FName GetNodeCategory() const override { return NAME_None; }
+	FText GetDisplayName() const override;
+	void ConstructNode() override;
+	bool ValidateConnection(const UOptimusNodePin& InThisNodesPin, const UOptimusNodePin& InOtherNodesPin, FString* OutReason) const override;
+
+	// IOptimusNodeAdderPinProvider
+	TArray<FAdderPinAction> GetAvailableAdderPinActions(const UOptimusNodePin* InSourcePin, EOptimusNodePinDirection InNewPinDirection, FString* OutReason) const override;
+	TArray<UOptimusNodePin*> TryAddPinFromPin(const FAdderPinAction& InSelectedAction, UOptimusNodePin* InSourcePin, FName InNameToUse) override;
+	bool RemoveAddedPins(TConstArrayView<UOptimusNodePin*> InAddedPinsToRemove) override;
+
+	// IOptimusUnnamedNodePinProvider
+	bool IsPinNameHidden(UOptimusNodePin* InPin) const override;
+	FName GetNameForAdderPin(UOptimusNodePin* InPin) const override;
+
+	// IOptimusPinMutabilityDefiner
+	EOptimusPinMutability GetOutputPinMutability(const UOptimusNodePin* InPin) const override;
+
+	// IOptimusNodePairProvider
+	void PairToCounterpartNode(const IOptimusNodePairProvider* NodePairProvider) override;
+	
+	UOptimusNodePin* GetPinCounterpart(const UOptimusNodePin* InNodePin, EOptimusTerminalType InTerminalType, TOptional<EOptimusNodePinDirection> InDirection = {}) const;
+	UOptimusNode_LoopTerminal* GetOtherTerminal() const;
+	int32 GetLoopCount() const;
+	EOptimusTerminalType GetTerminalType() const;
+
+	static int32 GetDataFunctionIndexFromPin(const UOptimusNodePin* InPin);
+
+protected:
+	friend class UOptimusNodeGraph;
+	
+	/** Indicates whether this is an entry or a return terminal node */
+	UPROPERTY(VisibleAnywhere, Category="Loop Terminal")
+	EOptimusTerminalType TerminalType;
+
+	
+	UPROPERTY(EditAnywhere, Category="Loop Terminal", DisplayName="Settings", meta=(FullyExpand="true", EditCondition="TerminalType==EOptimusTerminalType::Entry", EditConditionHides))
+	FOptimusLoopTerminalInfo LoopInfo;
+
+
+	UPROPERTY()
+	UOptimusNodePin* IndexPin;
+	
+	UPROPERTY()
+	UOptimusNodePin* CountPin;
+
+	UPROPERTY()
+	TArray<FOptimusPinPairInfo> PinPairInfos;
+	
+private:
+#if WITH_EDITOR
+	void PropertyArrayPasted(const FPropertyChangedEvent& InPropertyChangedEvent);
+	void PropertyValueChanged(const FPropertyChangedEvent& InPropertyChangedEvent);
+	void PropertyArrayItemAdded(const FPropertyChangedEvent& InPropertyChangedEvent);
+	void PropertyArrayItemRemoved(const FPropertyChangedEvent& InPropertyChangedEvent);
+	void PropertyArrayCleared(const FPropertyChangedEvent& InPropertyChangedEvent);
+	void PropertyArrayItemMoved(const FPropertyChangedEvent& InPropertyChangedEvent);
+#endif
+
+	TArray<UOptimusNodePin*> AddPinPairs(const FOptimusParameterBinding& InBinding);
+	TArray<UOptimusNodePin*> AddPinPairsDirect(const FOptimusParameterBinding& InBinding);
+	TArray<UOptimusNodePin*> GetPairedPins(const FOptimusPinPairInfo& InPair) const;
+	static int32 GetPairIndex(const UOptimusNodePin* Pin);
+	void RemovePinPair(int32 InPairIndex);
+	void RemovePinPairDirect(int32 InPairIndex);
+	void ClearPinPairs();
+	void MovePinPair();
+	void UpdatePinPairs();
+	FOptimusLoopTerminalInfo* GetLoopInfo();
+	const FOptimusLoopTerminalInfo* GetLoopInfo() const;
+	
+	void SanitizeBinding(FOptimusParameterBinding& InOutBinding, FName InOldName);
+	UOptimusNode_LoopTerminal* GetTerminalByType(EOptimusTerminalType InType);
+	const UOptimusNode_LoopTerminal* GetTerminalByType(EOptimusTerminalType InType) const;
+	
+	FName GetSanitizedBindingName(
+		FName InNewName,
+		FName InOldName
+	);
+};
diff --git a/Engine/Plugins/Animation/DeformerGraph/Source/OptimusCore/Private/Nodes/OptimusNode_ResourceAccessorBase.h b/Engine/Plugins/Animation/DeformerGraph/Source/OptimusCore/Private/Nodes/OptimusNode_ResourceAccessorBase.h
--- a/Engine/Plugins/Animation/DeformerGraph/Source/OptimusCore/Private/Nodes/OptimusNode_ResourceAccessorBase.h
+++ b/Engine/Plugins/Animation/DeformerGraph/Source/OptimusCore/Private/Nodes/OptimusNode_ResourceAccessorBase.h
@@ -67,7 +67,7 @@ public:
 	PRAGMA_ENABLE_DEPRECATION_WARNINGS
 
 	//IOptimusPinMutabilityDefiner overrides 
-	bool IsOutputPinMutable(const UOptimusNodePin* InPin) const override { return true; };
+	EOptimusPinMutability GetOutputPinMutability(const UOptimusNodePin* InPin) const override { return EOptimusPinMutability::Mutable; };
 	
 protected:
 	void PreDuplicateRequirementActions(const UOptimusNodeGraph* InTargetGraph, FOptimusCompoundAction* InCompoundAction) override;
diff --git a/Engine/Plugins/Animation/DeformerGraph/Source/OptimusCore/Private/OptimusDataDomain.cpp b/Engine/Plugins/Animation/DeformerGraph/Source/OptimusCore/Private/OptimusDataDomain.cpp
--- a/Engine/Plugins/Animation/DeformerGraph/Source/OptimusCore/Private/OptimusDataDomain.cpp
+++ b/Engine/Plugins/Animation/DeformerGraph/Source/OptimusCore/Private/OptimusDataDomain.cpp
@@ -174,6 +174,25 @@ TOptional<FString> FOptimusDataDomain::AsExpression() const
 	return {};
 }
 
+TSet<FName> FOptimusDataDomain::GetUsedConstants() const
+{
+	if (ensure(Type == EOptimusDataDomainType::Expression))
+	{
+		using namespace Optimus::Expression;
+		FEngine Engine;
+		TVariant<FExpressionObject, FParseError> ParseResult = Engine.Parse(Expression);
+
+		if (ParseResult.IsType<FParseError>())
+		{
+			return {};
+		}
+
+		return ParseResult.Get<FExpressionObject>().GetUsedConstants();
+	}
+
+	return {};
+}
+
 FString FOptimusDataDomain::GetDisplayName() const
 {
 	switch(Type)
@@ -212,8 +231,12 @@ FString FOptimusDataDomain::GetDisplayName() const
 bool FOptimusDataDomain::AreCompatible(const FOptimusDataDomain& InOutput, const FOptimusDataDomain& InInput,
                                        FString* OutReason)
 {
-	if (!InOutput.IsFullyDefined())
+	if (!InOutput.IsFullyDefined() || !InInput.IsFullyDefined())
 	{
+		if (OutReason)
+		{
+			*OutReason = TEXT("One of the pins has undefined datadomain");
+		}
 		return false;
 	}
 
@@ -228,23 +251,10 @@ bool FOptimusDataDomain::AreCompatible(const FOptimusDataDomain& InOutput, const
 		return false;
 	}
 
-	if (InOutput.IsSingleton() && !InInput.IsSingleton())
-	{
-		if (!InInput.IsFullyDefined())
-		{
-			if (OutReason)
-			{
-				*OutReason = TEXT("Can't connect a value output to a input with undefined data domain");
-			}
-			return false;
-		}
-		
-	}
-	
 	// If it's resource -> resource, check that the domains are compatible
 	if (!InOutput.IsSingleton() && !InInput.IsSingleton())
 	{
-		if (InOutput != InInput && InInput.IsFullyDefined())
+		if (InOutput != InInput)
 		{
 			if (OutReason)
 			{
diff --git a/Engine/Plugins/Animation/DeformerGraph/Source/OptimusCore/Private/OptimusDeformer.cpp b/Engine/Plugins/Animation/DeformerGraph/Source/OptimusCore/Private/OptimusDeformer.cpp
--- a/Engine/Plugins/Animation/DeformerGraph/Source/OptimusCore/Private/OptimusDeformer.cpp
+++ b/Engine/Plugins/Animation/DeformerGraph/Source/OptimusCore/Private/OptimusDeformer.cpp
@@ -11,6 +11,7 @@
 #include "Containers/Queue.h"
 #include "DataInterfaces/OptimusDataInterfaceGraph.h"
 #include "DataInterfaces/OptimusDataInterfaceRawBuffer.h"
+#include "DataInterfaces/OptimusDataInterfaceLoopTerminal.h"
 #include "IOptimusComputeKernelProvider.h"
 #include "IOptimusDataInterfaceProvider.h"
 #include "IOptimusValueProvider.h"
@@ -44,6 +45,7 @@
 #include "Nodes/OptimusNode_DataInterface.h"
 #include "Nodes/OptimusNode_GetVariable.h"
 #include "Nodes/OptimusNode_ResourceAccessorBase.h"
+#include "Nodes/OptimusNode_LoopTerminal.h"
 
 #include "IOptimusDeprecatedExecutionDataInterface.h"
 #include "OptimusNodeLink.h"
@@ -1341,7 +1343,9 @@ UOptimusComponentSourceBinding* UOptimusDeformer::ResolveComponentBinding(
 // Do a breadth-first collection of nodes starting from the seed nodes (terminal data interfaces).
 static void CollectNodes(
 	const TArray<const UOptimusNode*>& InSeedNodes,
-	TArray<FOptimusRoutedConstNode>& OutCollectedNodes
+	TArray<FOptimusRoutedConstNode>& OutCollectedNodes,
+	TMap<FOptimusRoutedConstNode, TArray<FOptimusRoutedConstNode>>& OutNodeToInputNodes,
+	TMap<FOptimusRoutedConstNode, TArray<FOptimusRoutedConstNode>>& OutNodeToOutputNodes 
 	)
 {
 	TSet<FOptimusRoutedConstNode> VisitedNodes;
@@ -1355,7 +1359,7 @@ static void CollectNodes(
 		OutCollectedNodes.Add({Node, FOptimusPinTraversalContext{}});
 	}
 
-	auto CollectFromInputPins = [&WorkingSet, &VisitedNodes, &UniqueNeighborNodes, &OutCollectedNodes](const FOptimusRoutedConstNode& InWorkItem, const UOptimusNodePin* InPin)
+	auto CollectFromInputPins = [&WorkingSet, &VisitedNodes, &UniqueNeighborNodes, &OutCollectedNodes, &OutNodeToInputNodes, &OutNodeToOutputNodes](const FOptimusRoutedConstNode& InWorkItem, const UOptimusNodePin* InPin)
 	{
 		for (const FOptimusRoutedNodePin& ConnectedPin: InPin->GetConnectedPinsWithRouting(InWorkItem.TraversalContext))
 		{
@@ -1367,6 +1371,9 @@ static void CollectNodes(
 				{
 					UniqueNeighborNodes.Add(CollectedNode);
 					
+					OutNodeToInputNodes.FindOrAdd(InWorkItem).Add(CollectedNode);
+					OutNodeToOutputNodes.FindOrAdd(CollectedNode).Add(InWorkItem);
+					
 					WorkingSet.Enqueue(CollectedNode);
 					
 					if (!VisitedNodes.Contains(CollectedNode))
@@ -1430,7 +1437,6 @@ bool UOptimusDeformer::Compile()
 	}
 
 	ComputeGraphs.Reset();
-	ConstantContainer.Reset();
 	
 	CompileBeginDelegate.Broadcast(this);
 	
@@ -1458,13 +1464,12 @@ bool UOptimusDeformer::Compile()
 		TArray<FOptimusComputeGraphInfo> ComputeGraphInfos = CompileNodeGraphToComputeGraphs(Graph, ErrorReporter);
 		ComputeGraphs.Append(ComputeGraphInfos);
 	}
-	
+
 	CompileEndDelegate.Broadcast(this);
 
 	if (Status == EOptimusDeformerStatus::HasErrors)
 	{
 		ComputeGraphs.Reset();
-		ConstantContainer.Reset();
 		return false;
 	}
 	
@@ -1511,6 +1516,37 @@ TArray<UOptimusNode*> UOptimusDeformer::GetAllNodesOfClass(UClass* InNodeClass)
 	return NodesFound;
 }
 
+struct FOptimusInstancedNode
+{
+	FOptimusRoutedConstNode	RoutedNode;
+	int32 LoopIndex;
+		
+	friend uint32 GetTypeHash(const FOptimusInstancedNode& InNode)
+	{
+		return HashCombineFast(GetTypeHash(InNode.RoutedNode), InNode.LoopIndex);
+	}
+
+	bool operator==(const FOptimusInstancedNode& InOther) const
+	{
+		return RoutedNode == InOther.RoutedNode && LoopIndex == InOther.LoopIndex;
+	}
+};
+
+struct FOptimusInstancedPin
+{
+	FOptimusInstancedNode InstancedNode;
+	const UOptimusNodePin* Pin;
+
+	friend uint32 GetTypeHash(const FOptimusInstancedPin& InPin)
+	{
+		return HashCombineFast(GetTypeHash(InPin.InstancedNode), GetTypeHash(InPin.Pin));
+	}
+
+	bool operator==(const FOptimusInstancedPin& InOther) const
+	{
+		return InstancedNode == InOther.InstancedNode && Pin == InOther.Pin;
+	}
+};
 
 TArray<FOptimusComputeGraphInfo> UOptimusDeformer::CompileNodeGraphToComputeGraphs(
 	const UOptimusNodeGraph* InNodeGraph,
@@ -1587,7 +1623,9 @@ TArray<FOptimusComputeGraphInfo> UOptimusDeformer::CompileNodeGraphToComputeGrap
 	}
 
 	TArray<FOptimusRoutedConstNode> ConnectedNodes;
-	CollectNodes(TerminalNodes, ConnectedNodes);
+	TMap<FOptimusRoutedConstNode, TArray<FOptimusRoutedConstNode>> NodeToInputNodes;
+	TMap<FOptimusRoutedConstNode, TArray<FOptimusRoutedConstNode>> NodeToOutputNodes; 
+	CollectNodes(TerminalNodes, ConnectedNodes, NodeToInputNodes, NodeToOutputNodes);
 
 	// Since we now have the connected nodes in a breadth-first list, reverse the list which
 	// will give use the same list but topologically sorted in kernel execution order.
@@ -1609,108 +1647,233 @@ TArray<FOptimusComputeGraphInfo> UOptimusDeformer::CompileNodeGraphToComputeGrap
 		return {};
 	}
 
-	TArray<FOptimusComputeGraphInfo> GraphInfos;
-	TMap<const IOptimusComputeKernelProvider*, int> KernelComputeGraphIndexMap;
-	
-	if (InNodeGraph->GraphType != EOptimusNodeGraphType::Update)
+	// Mark zero-count loops as skippable
+	TSet<FOptimusRoutedConstNode> LoopTerminalToSkip; 
+	for (const FOptimusRoutedConstNode& ConnectedNode : ConnectedNodes )
 	{
-		FOptimusComputeGraphInfo GraphInfo;
-		GraphInfo.GraphName = MakeUniqueObjectName(this, UOptimusComputeGraph::StaticClass(), InNodeGraph->GetFName());;
-		GraphInfo.GraphType = InNodeGraph->GetGraphType();
-		GraphInfo.ComputeGraph = NewObject<UOptimusComputeGraph>(this, GraphInfo.GraphName);
-		GraphInfos.Add(GraphInfo);
+		if (const UOptimusNode_LoopTerminal* LoopTerminal = Cast<const UOptimusNode_LoopTerminal>(ConnectedNode.Node);
+			LoopTerminal && LoopTerminal->GetLoopCount() == 0)
+		{
+			LoopTerminalToSkip.Add(ConnectedNode);
+		}
+	}
+
+	// Collect looping kernels and skippable kernels
+	TMap<FOptimusRoutedConstNode, TArray<FOptimusRoutedConstNode>> LoopEntryToKernelNodes; 
+	TMap<FOptimusRoutedConstNode, FOptimusRoutedConstNode> KernelNodeToLoopEntry; 
+	TSet<FOptimusRoutedConstNode> KernelToSkip; 
 
-		for (FOptimusRoutedConstNode ConnectedNode: ConnectedNodes)
+	for (const FOptimusRoutedConstNode& ConnectedNode : ConnectedNodes )
+	{
+		if (const IOptimusComputeKernelProvider *KernelProvider = Cast<const IOptimusComputeKernelProvider>(ConnectedNode.Node))
 		{
-			if (const IOptimusComputeKernelProvider *KernelProvider = Cast<const IOptimusComputeKernelProvider>(ConnectedNode.Node))
+			TSet<FOptimusRoutedConstNode> LoopTerminals = ConnectedNode.Node->GetOwningGraph()->GetLoopEntryTerminalForNode(ConnectedNode.Node, ConnectedNode.TraversalContext);
+			if (LoopTerminals.Num() > 0)
 			{
-				KernelComputeGraphIndexMap.Add(KernelProvider) = 0;
+				FOptimusRoutedConstNode LoopEntry= *LoopTerminals.CreateConstIterator();
+
+				LoopEntryToKernelNodes.FindOrAdd(LoopEntry).Add(ConnectedNode);
+				KernelNodeToLoopEntry.Add(ConnectedNode) = LoopEntry;
+
+				if (LoopTerminalToSkip.Contains(LoopEntry))
+				{
+					KernelToSkip.Add(ConnectedNode);
+				}
 			}
 		}
 	}
-	else
+
+	// Mark additional kernels with no meaningful output as skippable
+	for (int32 Index = ConnectedNodes.Num() - 1; Index >= 0 ; Index-- )
+	{
+		const FOptimusRoutedConstNode& ConnectedNode = ConnectedNodes[Index];
+
+		if (const IOptimusComputeKernelProvider *KernelProvider = Cast<const IOptimusComputeKernelProvider>(ConnectedNode.Node))
+		{
+			bool bShouldSkip = true;
+			for (const FOptimusRoutedConstNode& OutputNode : NodeToOutputNodes[ConnectedNode])
+			{
+				// Do not skip if this node has at least one meaningful output
+				if (!KernelToSkip.Contains(OutputNode))
+				{
+					bShouldSkip = false;
+				}
+			}
+
+			if (bShouldSkip)
+			{
+				KernelToSkip.Add(ConnectedNode);
+			}
+		}
+	}
+
+	// Used to establish kernel - kernel , kernel - data interface dependencies so that we know what data interface to assign to the output pin at loop boundaries
+	TMap<FOptimusRoutedConstNodePin, TArray<FOptimusRoutedConstNodePin>> TargetResourcePinConnections;
+	TMap<FOptimusRoutedConstNodePin, FOptimusRoutedConstNodePin> SourceResourcePinConnections;
+	
+	for (FOptimusRoutedConstNode ConnectedNode : ConnectedNodes)
 	{
-		bool bHasRunOnceKernels = false;
-		bool bHasRunAlwaysKernels = false;
+		if (Cast<const IOptimusDataInterfaceProvider>(ConnectedNode.Node) || Cast<const IOptimusComputeKernelProvider>(ConnectedNode.Node))
+		{
+			if (!KernelToSkip.Contains(ConnectedNode))
+			{
+				for (const UOptimusNodePin* Pin: ConnectedNode.Node->GetPinsByDirection(EOptimusNodePinDirection::Input, true))
+				{
+					if (!Pin->GetDataDomain().IsSingleton())
+					{
+						TArray<FOptimusRoutedConstNodePin> SourceResourcePins;
+
+						FOptimusRoutedConstNodePin WorkPin;
+						TQueue<FOptimusRoutedConstNodePin> PinQueue;
+						TSet<FOptimusRoutedConstNodePin> VisitedPin;
+						PinQueue.Enqueue({Pin, ConnectedNode.TraversalContext});
+						while(PinQueue.Dequeue(WorkPin))
+						{
+							if (VisitedPin.Contains(WorkPin))
+							{
+								continue;
+							}
+							VisitedPin.Add({Pin, ConnectedNode.TraversalContext});
+							
+							TArray<FOptimusRoutedNodePin> NextRoutedPin = WorkPin.NodePin->GetConnectedPinsWithRouting(WorkPin.TraversalContext);
+
+							if (NextRoutedPin.Num() == 1)
+							{
+								UOptimusNodePin* NextPin = NextRoutedPin[0].NodePin;
+								FOptimusRoutedConstNode NextRoutedNode = {NextPin->GetOwningNode(), NextRoutedPin[0].TraversalContext};
+								if (const UOptimusNode_LoopTerminal* LoopTerminal = Cast<UOptimusNode_LoopTerminal>(NextRoutedNode.Node))
+								{
+									if (LoopTerminal->GetTerminalType() == EOptimusTerminalType::Return)
+									{
+										if (!LoopTerminalToSkip.Contains(NextRoutedNode))
+										{
+											PinQueue.Enqueue({LoopTerminal->GetPinCounterpart(NextPin, EOptimusTerminalType::Return), NextRoutedPin[0].TraversalContext});
+										}
+										else
+										{
+											PinQueue.Enqueue({LoopTerminal->GetPinCounterpart(NextPin, EOptimusTerminalType::Entry), NextRoutedPin[0].TraversalContext});
+										}
+									}
+									else if (LoopTerminal->GetTerminalType() == EOptimusTerminalType::Entry)
+									{
+										if (ensure(!LoopTerminalToSkip.Contains(NextRoutedNode)))
+										{
+											PinQueue.Enqueue( {LoopTerminal->GetPinCounterpart(NextPin, EOptimusTerminalType::Entry), NextRoutedPin[0].TraversalContext});
+										}
+									}
+								}
+								else
+								{
+									if (!KernelToSkip.Contains(ConnectedNode))
+									{
+										SourceResourcePins.Add({NextRoutedPin[0].NodePin, NextRoutedPin[0].TraversalContext});	
+									}
+								}
+							}
+						}
+
+						if (SourceResourcePins.Num() == 1)
+						{
+							TargetResourcePinConnections.FindOrAdd(SourceResourcePins[0]).Add({Pin, ConnectedNode.TraversalContext});
+							SourceResourcePinConnections.Add({Pin, ConnectedNode.TraversalContext}) = SourceResourcePins[0];
+						}
+					}
+				}	
+			}
+		}
+	}
+
+	TSet<EOptimusNodeGraphType> GraphTypes;
+	TMap<FOptimusRoutedConstNode, EOptimusNodeGraphType> KernelToGraphType; 
 	
-		for (FOptimusRoutedConstNode ConnectedNode: ConnectedNodes)
+	for (FOptimusRoutedConstNode ConnectedNode : ConnectedNodes)
+	{
+		if (const IOptimusComputeKernelProvider *KernelProvider = Cast<const IOptimusComputeKernelProvider>(ConnectedNode.Node))
 		{
-			if (const IOptimusComputeKernelProvider *KernelProvider = Cast<const IOptimusComputeKernelProvider>(ConnectedNode.Node))
+			if (InNodeGraph->GraphType != EOptimusNodeGraphType::Update)
 			{
-				if (!KernelProvider->HasMutableInput())
+				KernelToGraphType.Add(ConnectedNode) = InNodeGraph->GraphType;
+				GraphTypes.Add(InNodeGraph->GraphType);
+			}
+			else
+			{
+				if (ConnectedNode.Node->GetOwningGraph()->DoesNodeHaveMutableInput(ConnectedNode.Node))
 				{
-					bHasRunOnceKernels = true;
+					KernelToGraphType.Add(ConnectedNode) = EOptimusNodeGraphType::Update;
+					GraphTypes.Add(EOptimusNodeGraphType::Update);
 				}
 				else
 				{
-					bHasRunAlwaysKernels = true;
+					KernelToGraphType.Add(ConnectedNode) = EOptimusNodeGraphType::Setup;
+					GraphTypes.Add(EOptimusNodeGraphType::Setup);
 				}
 			}
 		}
+	}
 
-		// Insert a setup graph if applicable
-		if (bHasRunOnceKernels)
-		{
-			FName GraphName = InNodeGraph->GetFName();
-			GraphName = *(GraphName.ToString() + TEXT("_Setup"));
-			GraphName = MakeUniqueObjectName(this, UOptimusComputeGraph::StaticClass(), GraphName);
-		
-			FOptimusComputeGraphInfo GraphInfo;
-			GraphInfo.GraphName = GraphName;
-			GraphInfo.GraphType = EOptimusNodeGraphType::Setup;
-			GraphInfo.ComputeGraph = NewObject<UOptimusComputeGraph>(this, GraphInfo.GraphName);
-			GraphInfos.Add(GraphInfo);
-		}
+	
 
-		if (bHasRunAlwaysKernels)
+	// Instance looped nodes
+	TArray<FOptimusInstancedNode> InstancedNodes;
+	TMap<FOptimusRoutedConstNode, int32> NodeToMaxLoopIndex;
+
+	for (const FOptimusRoutedConstNode& ConnectedNode : ConnectedNodes)
+	{
+		if (KernelToSkip.Contains(ConnectedNode))
 		{
-			FOptimusComputeGraphInfo GraphInfo;
-			GraphInfo.GraphName = MakeUniqueObjectName(this, UOptimusComputeGraph::StaticClass(), InNodeGraph->GetFName());;
-			GraphInfo.GraphType = InNodeGraph->GetGraphType();
-			GraphInfo.ComputeGraph = NewObject<UOptimusComputeGraph>(this, GraphInfo.GraphName);
-			GraphInfos.Add(GraphInfo);
+			continue;
 		}
 
-		for (FOptimusRoutedConstNode ConnectedNode: ConnectedNodes)
+		if (const UOptimusNode_LoopTerminal* LoopTerminal = Cast<const UOptimusNode_LoopTerminal>(ConnectedNode.Node))
 		{
-			if (const IOptimusComputeKernelProvider *KernelProvider = Cast<const IOptimusComputeKernelProvider>(ConnectedNode.Node))
+			if (!LoopTerminalToSkip.Contains(ConnectedNode))
 			{
-				int32 ComputeGraphIndex = 0;
-				if (GraphInfos.Num() == 2)
+				if (LoopTerminal->GetTerminalType() == EOptimusTerminalType::Return)
 				{
-					ComputeGraphIndex = KernelProvider->HasMutableInput() ? 1 : 0;
+					FOptimusRoutedConstNode LoopEntry = {LoopTerminal->GetOtherTerminal(), ConnectedNode.TraversalContext};
+
+					for (int32 Index = 1; Index < LoopTerminal->GetLoopCount(); Index++)
+					{
+						for (const FOptimusRoutedConstNode& KernelNode : LoopEntryToKernelNodes[LoopEntry])
+						{
+							InstancedNodes.Add({ KernelNode , Index });
+						}
+					}
+					
+					for (const FOptimusRoutedConstNode& KernelNode : LoopEntryToKernelNodes[LoopEntry])
+					{
+						NodeToMaxLoopIndex.FindOrAdd(KernelNode) = LoopTerminal->GetLoopCount() - 1;
+					}
 				}
-				KernelComputeGraphIndexMap.Add(KernelProvider) = ComputeGraphIndex;
 			}
 		}
+
+		InstancedNodes.Add({ConnectedNode, 0});
+		NodeToMaxLoopIndex.FindOrAdd(ConnectedNode) = 0;
 	}
 
-	check(GraphInfos.Num() <= 2);
 
 	
+	// Create all the data interfaces: node, graph, kernel outputs, loop terminal data
+	
+	// The component binding for the graph data is the primary binding on the deformer.
+	UOptimusComponentSourceBinding* GraphDataComponentBinding = Bindings->Bindings[0];
+	
+	TMap<const UComputeDataInterface*, int32> DataInterfaceToBindingIndexMap;
+	
 	// Find all data interface nodes and create their data interfaces.
 	FOptimus_NodeToDataInterfaceMap NodeDataInterfaceMap;
-
-	// Find all resource links from one compute kernel directly to another. The pin here is
-	// the output pin from a kernel node that connects to another. We don't map from input pins
-	// because a resource output may be used multiple times, but only written into once.
-	FOptimus_PinToDataInterfaceMap LinkDataInterfaceMap;
-
-	// Each kernel spawns a data interface to pass along kernel specific data to the kernel itself
-	FOptimus_KernelNodeToKernelDataInterfaceMap KernelDataInterfaceMap;
-
+	
 	// Find all value nodes (constant and variable) 
 	TArray<const UOptimusNode *> ValueNodes;
 
-	TMap<const UComputeDataInterface*, int32> DataInterfaceToBindingIndexMap;
-
+	TMap<FOptimusRoutedConstNode, TArray<UOptimusComputeDataInterface*>> LoopEntryToLoopDataInterfaces;
 	
-
-	for (FOptimusRoutedConstNode ConnectedNode: ConnectedNodes)
+	for (const FOptimusRoutedConstNode& ConnectedNode : ConnectedNodes)
 	{
-		if (const IOptimusDataInterfaceProvider* DataInterfaceNode = Cast<const IOptimusDataInterfaceProvider>(ConnectedNode.Node))
+		if (const IOptimusDataInterfaceProvider* NodeDataInterfaceProvider = Cast<const IOptimusDataInterfaceProvider>(ConnectedNode.Node))
 		{
-			UOptimusComputeDataInterface* DataInterface = DataInterfaceNode->GetDataInterface(this);
+			UOptimusComputeDataInterface* DataInterface = NodeDataInterfaceProvider->GetDataInterface(this);
 			if (!DataInterface)
 			{
 				AddDiagnostic(EOptimusDiagnosticLevel::Error, LOCTEXT("NoDataInterfaceOnProvider", "No data interface object returned from node. Compilation aborted."));
@@ -1718,9 +1881,63 @@ TArray<FOptimusComputeGraphInfo> UOptimusDeformer::CompileNodeGraphToComputeGrap
 			}
 
 			NodeDataInterfaceMap.Add(ConnectedNode.Node, DataInterface);
-			DataInterfaceToBindingIndexMap.Add(DataInterface) = DataInterfaceNode->GetComponentBinding()->GetIndex();
+			DataInterfaceToBindingIndexMap.Add(DataInterface) = NodeDataInterfaceProvider->GetComponentBinding()->GetIndex();
+		}
+		else if (Cast<const IOptimusValueProvider>(ConnectedNode.Node))
+		{
+			ValueNodes.AddUnique(ConnectedNode.Node);
+		}
+		else if (const UOptimusNode_LoopTerminal* LoopTerminal = Cast<const UOptimusNode_LoopTerminal>(ConnectedNode.Node))
+		{
+			if (LoopTerminal->GetTerminalType() == EOptimusTerminalType::Entry)
+			{
+				TArray<UOptimusComputeDataInterface*>& LoopDataInterfaces = LoopEntryToLoopDataInterfaces.Add(ConnectedNode);
+				for (int32 LoopIndex = 0; LoopIndex < LoopTerminal->GetLoopCount(); LoopIndex++)
+				{
+					UOptimusLoopTerminalDataInterface* LoopDataInterface = NewObject<UOptimusLoopTerminalDataInterface>(this);
+					LoopDataInterface->Index = LoopIndex;
+					LoopDataInterface->Count = LoopTerminal->GetLoopCount();
+					LoopDataInterfaces.Add(LoopDataInterface);
+					// Loop Data are constant values for now so it does not matter which binding is used
+					DataInterfaceToBindingIndexMap.Add(LoopDataInterface) = GraphDataComponentBinding->GetIndex();
+				}
+			}
+		}
+	}
+
+	// Create the graph data interface and fill it with the value nodes.
+	UOptimusGraphDataInterface* GraphDataInterface = NewObject<UOptimusGraphDataInterface>(this);
+
+	TArray<FOptimusGraphVariableDescription> ValueNodeDescriptions;
+	ValueNodeDescriptions.Reserve(ValueNodes.Num());
+	for (UOptimusNode const* ValueNode : ValueNodes)
+	{
+		if (IOptimusValueProvider const* ValueProvider = Cast<const IOptimusValueProvider>(ValueNode))
+		{
+			FOptimusGraphVariableDescription& ValueNodeDescription = ValueNodeDescriptions.AddDefaulted_GetRef();
+			ValueNodeDescription.Name = ValueProvider->GetValueName();
+			ValueNodeDescription.ValueType = ValueProvider->GetValueType()->ShaderValueType;
+
+			if (UOptimusNode_ConstantValue const* ConstantNode = Cast<const UOptimusNode_ConstantValue>(ValueNode))
+			{
+				ValueNodeDescription.Value = ConstantNode->GetShaderValue().ShaderValue;
+			}
 		}
-		else if (const IOptimusComputeKernelProvider* KernelProvider = Cast<const IOptimusComputeKernelProvider>(ConnectedNode.Node))
+	}
+	GraphDataInterface->Init(ValueNodeDescriptions);
+
+	DataInterfaceToBindingIndexMap.Add(GraphDataInterface) = GraphDataComponentBinding->GetIndex();
+	
+	TMap<FOptimusInstancedNode, UComputeDataInterface*> KernelDataInterfaceMap;
+	TMap<FOptimusInstancedNode, FOptimus_KernelInputMap> KernelInputMap;
+	TMap<FOptimusInstancedNode, FOptimus_KernelOutputMap> KernelOutputMap;
+	TMap<FOptimusInstancedPin, UOptimusComputeDataInterface*> KernelOutputDataInterfaceMap;
+
+	for (const FOptimusInstancedNode& InstancedNode : InstancedNodes )
+	{
+		const FOptimusRoutedConstNode& ConnectedNode = InstancedNode.RoutedNode;
+
+		if (const IOptimusComputeKernelProvider* KernelProvider = Cast<const IOptimusComputeKernelProvider>(ConnectedNode.Node))
 		{
 			UComputeDataInterface* KernelDataInterface = KernelProvider->MakeKernelDataInterface(this);
 			TSet<UOptimusComponentSourceBinding*> KernelPrimaryBindings = KernelProvider->GetPrimaryGroupPin()->GetComponentSourceBindingsRecursively();
@@ -1737,95 +1954,260 @@ TArray<FOptimusComputeGraphInfo> UOptimusDeformer::CompileNodeGraphToComputeGrap
 			UOptimusComponentSourceBinding* KernelPrimaryBinding = *KernelPrimaryBindings.CreateConstIterator();
 			int32 PrimaryBindingIndex = KernelPrimaryBinding->GetIndex();
 
-			KernelDataInterfaceMap.Add(ConnectedNode.Node, KernelDataInterface);
+			KernelDataInterfaceMap.Add(InstancedNode, KernelDataInterface);
 			DataInterfaceToBindingIndexMap.Add(KernelDataInterface) = PrimaryBindingIndex;	
 			
-			for (const UOptimusNodePin* Pin: ConnectedNode.Node->GetPins())
+			for (const UOptimusNodePin* Pin: ConnectedNode.Node->GetPinsByDirection(EOptimusNodePinDirection::Input, true))
+			{
+				if (Pin->IsGroupingPin())
+				{
+					continue;
+				}
+				
+				FOptimusInstancedPin InstancedPin = {InstancedNode, Pin};
+				FOptimusInstancedPin SourceInstancedPin;
+
+				if (!Pin->GetDataDomain().IsSingleton())
+				{
+					for (const FOptimusRoutedNodePin& ConnectedPin: Pin->GetConnectedPinsWithRouting(ConnectedNode.TraversalContext))
+					{
+						// By default Read from the connected node in the same iteration
+						FOptimusRoutedConstNode SourceRoutedNode = {ConnectedPin.NodePin->GetOwningNode(), ConnectedPin.TraversalContext};
+						int32 SourceLoopIndex = FMath::Clamp(InstancedNode.LoopIndex, 0, NodeToMaxLoopIndex[SourceRoutedNode]);
+						SourceInstancedPin = {FOptimusInstancedNode(SourceRoutedNode, SourceLoopIndex), ConnectedPin.NodePin};
+
+						// If connecting to a terminal Read from last node from the previous iteration
+						if (const UOptimusNode_LoopTerminal* LoopTerminal = Cast<const UOptimusNode_LoopTerminal>(ConnectedPin.NodePin->GetOwningNode()))
+						{
+							if (InstancedNode.LoopIndex == 0)
+							{
+								// 0th instance reads from upstream as if there isn't a loop terminal
+								FOptimusRoutedConstNodePin RoutedPin = {Pin, ConnectedNode.TraversalContext};
+								if (const FOptimusRoutedConstNodePin* SourceRoutedPin = SourceResourcePinConnections.Find(RoutedPin))
+								{
+									SourceRoutedNode = {SourceRoutedPin->NodePin->GetOwningNode(), SourceRoutedPin->TraversalContext};
+									int32 MaxLoopIndex = NodeToMaxLoopIndex[SourceRoutedNode];
+									SourceInstancedPin = {FOptimusInstancedNode(SourceRoutedNode, MaxLoopIndex), SourceRoutedPin->NodePin};
+								}	
+							}
+							else
+							{
+								if (ensure(LoopTerminal->GetTerminalType() == EOptimusTerminalType::Entry))
+								{
+									UOptimusNodePin* PinCounterpart = LoopTerminal->GetPinCounterpart(ConnectedPin.NodePin, EOptimusTerminalType::Return);
+									for (const FOptimusRoutedNodePin& PinFromLastIteration: PinCounterpart->GetConnectedPinsWithRouting(ConnectedPin.TraversalContext))
+									{
+										SourceRoutedNode = {PinFromLastIteration.NodePin->GetOwningNode(), PinFromLastIteration.TraversalContext};
+										int32 SourceLoopIndexFromLastIteration = FMath::Clamp(InstancedNode.LoopIndex - 1, 0, NodeToMaxLoopIndex[SourceRoutedNode]);
+										SourceInstancedPin = {FOptimusInstancedNode(SourceRoutedNode, SourceLoopIndexFromLastIteration), PinFromLastIteration.NodePin};
+									}
+								}	
+							}
+						}
+					}
+				}
+				else
+				{
+					for (const FOptimusRoutedNodePin& ConnectedPin: Pin->GetConnectedPinsWithRouting(ConnectedNode.TraversalContext))
+					{
+						FOptimusRoutedConstNode SourceRoutedNode = {ConnectedPin.NodePin->GetOwningNode(), ConnectedPin.TraversalContext};
+						SourceInstancedPin = {FOptimusInstancedNode(SourceRoutedNode, 0), ConnectedPin.NodePin};
+					}
+				}
+
+				const UOptimusNodePin* SourcePin = SourceInstancedPin.Pin;
+				const FOptimusRoutedConstNode& SourceRoutedNode = SourceInstancedPin.InstancedNode.RoutedNode;
+
+				if (SourcePin)
+				{
+					if (const IOptimusValueProvider* ValueProvider = Cast<const IOptimusValueProvider>(SourcePin->GetOwningNode()))
+					{
+						KernelInputMap.FindOrAdd(InstancedNode).Add(Pin) = {GraphDataInterface, SourcePin};
+					}
+					else if (UOptimusComputeDataInterface** NodeDataInterface = NodeDataInterfaceMap.Find(SourcePin->GetOwningNode()))
+					{
+						KernelInputMap.FindOrAdd(InstancedNode).Add(Pin) = {*NodeDataInterface, SourcePin};
+					}
+					else if (UOptimusComputeDataInterface** KernelOutputDataInterface = KernelOutputDataInterfaceMap.Find(SourceInstancedPin))
+					{
+						KernelInputMap.FindOrAdd(InstancedNode).Add(Pin) = {*KernelOutputDataInterface, SourcePin};
+					}
+					else if (Cast<const UOptimusNode_LoopTerminal>(SourcePin->GetOwningNode()) && Pin->GetDataDomain().IsSingleton())
+					{
+						KernelInputMap.FindOrAdd(InstancedNode).Add(Pin) = {LoopEntryToLoopDataInterfaces[SourceRoutedNode][InstancedNode.LoopIndex] , SourcePin};
+					}
+				}
+			}
+			
+			for (const UOptimusNodePin* Pin: ConnectedNode.Node->GetPinsByDirection(EOptimusNodePinDirection::Output, true))
 			{
-				if (Pin->GetDirection() == EOptimusNodePinDirection::Output &&
-					ensure(!Pin->GetDataDomain().IsSingleton()) &&
-					!LinkDataInterfaceMap.Contains(Pin))
+				if (ensure(!Pin->GetDataDomain().IsSingleton()))
 				{
+					FOptimusInstancedPin InstancedPin = {InstancedNode, Pin};
+
+					TArray<FOptimusInstancedPin> TargetInstancedPins;
+
+					// Writing to other looped kernel instances or data interface node
 					for (const FOptimusRoutedNodePin& ConnectedPin: Pin->GetConnectedPinsWithRouting(ConnectedNode.TraversalContext))
 					{
-						// Make sure it connects to another kernel node.
-						if (const IOptimusComputeKernelProvider* ConnectedKernel = Cast<const IOptimusComputeKernelProvider>(ConnectedPin.NodePin->GetOwningNode());
-							ConnectedKernel && ensure(Pin->GetDataType().IsValid()))
+						// By default write to the connected node in the same iteration
+						UOptimusNode* TargetNode = ConnectedPin.NodePin->GetOwningNode();
+						
+						if (const UOptimusNode_LoopTerminal* LoopTerminal = Cast<const UOptimusNode_LoopTerminal>(TargetNode))
+						{
+							// Last instance of the kernel writes as if there isn't a loop terminal
+							if (InstancedNode.LoopIndex == NodeToMaxLoopIndex[ConnectedNode])
+							{
+								FOptimusRoutedConstNodePin RoutedPin = {Pin, ConnectedNode.TraversalContext};
+								if (TArray<FOptimusRoutedConstNodePin>* TargetRoutedPins = TargetResourcePinConnections.Find(RoutedPin))
+								{
+									for (FOptimusRoutedConstNodePin& TargetRoutedPin : *TargetRoutedPins)
+									{
+										FOptimusRoutedConstNode TargetRoutedNode = {TargetRoutedPin.NodePin->GetOwningNode(), TargetRoutedPin.TraversalContext};
+										TargetInstancedPins.Add({FOptimusInstancedNode(TargetRoutedNode, 0) , TargetRoutedPin.NodePin});	
+									}
+								}	
+							}
+							else
+							{
+								// For nodes with in a loop, if connecting to a terminal, write to the first node from the next iteration
+								if (ensure(LoopTerminal->GetTerminalType() == EOptimusTerminalType::Return))
+								{
+									UOptimusNodePin* PinCounterpart = LoopTerminal->GetPinCounterpart(ConnectedPin.NodePin, EOptimusTerminalType::Entry);
+									for (const FOptimusRoutedNodePin& PinFromNextIteration: PinCounterpart->GetConnectedPinsWithRouting(ConnectedPin.TraversalContext))
+									{
+										UOptimusNode* TargetNodeFromNextIteration = PinFromNextIteration.NodePin->GetOwningNode();
+										FOptimusRoutedConstNode TargetRoutedNodeFromNextIteration = {TargetNodeFromNextIteration, PinFromNextIteration.TraversalContext};
+										int32 TargetLoopIndexFromNextIteration = FMath::Clamp(InstancedNode.LoopIndex + 1, 0, NodeToMaxLoopIndex[TargetRoutedNodeFromNextIteration]);
+										TargetInstancedPins.Add({FOptimusInstancedNode(TargetRoutedNodeFromNextIteration, TargetLoopIndexFromNextIteration), PinFromNextIteration.NodePin});
+									}
+								}	
+							}
+						}
+						else
+						{
+							FOptimusRoutedConstNode TargetRoutedNode = {TargetNode, ConnectedPin.TraversalContext};
+							int32 TargetLoopIndex = FMath::Clamp(InstancedNode.LoopIndex, 0, NodeToMaxLoopIndex[TargetRoutedNode]);
+							FOptimusInstancedPin TargetInstancedPin = {FOptimusInstancedNode(TargetRoutedNode, TargetLoopIndex), ConnectedPin.NodePin};
+							TargetInstancedPins.Add(TargetInstancedPin);	
+						}
+					}	
+
+					if (TargetInstancedPins.Num() == 0)
+					{
+						if (KernelProvider->DoesOutputPinSupportAtomic(Pin) || KernelProvider->DoesOutputPinSupportRead(Pin))
+						{
+							UOptimusTransientBufferDataInterface* TransientBufferDI = NewObject<UOptimusTransientBufferDataInterface>(this);
+							if (KernelProvider->DoesOutputPinSupportAtomic(Pin))
+							{
+								CastChecked<UOptimusTransientBufferDataInterface>(TransientBufferDI)->bZeroInitForAtomicWrites = true;
+							}
+							TransientBufferDI->ValueType = Pin->GetDataType()->ShaderValueType;
+							TransientBufferDI->DataDomain = Pin->GetDataDomain();
+							TransientBufferDI->ComponentSourceBinding = KernelPrimaryBinding;
+							DataInterfaceToBindingIndexMap.Add(TransientBufferDI) = PrimaryBindingIndex;
+								
+							KernelOutputDataInterfaceMap.Add(InstancedPin) = TransientBufferDI;
+
+							KernelOutputMap.FindOrAdd(InstancedNode).FindOrAdd(Pin).Add({TransientBufferDI, nullptr});	
+						}
+					}
+					else
+					{
+						bool bShouldUseImplicitPersistentDI = false;
+						TArray<const UOptimusNodePin*> TargetKernelPins;
+						
+						for (const FOptimusInstancedPin& TargetInstancedPin : TargetInstancedPins)
 						{
-							if (KernelComputeGraphIndexMap[KernelProvider] == KernelComputeGraphIndexMap[ConnectedKernel])
+							const FOptimusRoutedConstNode TargetRoutedNode = TargetInstancedPin.InstancedNode.RoutedNode;
+							const UOptimusNode* TargetNode = TargetRoutedNode.Node;
+							const UOptimusNodePin* TargetPin = TargetInstancedPin.Pin;
+							
+							if (UOptimusComputeDataInterface** NodeDataInterface = NodeDataInterfaceMap.Find(TargetNode))
 							{
-								UOptimusTransientBufferDataInterface* TransientBufferDI =
-									NewObject<UOptimusTransientBufferDataInterface>(this);
+								KernelOutputMap.FindOrAdd(InstancedNode).FindOrAdd(Pin).Add({*NodeDataInterface, TargetPin});
+							}
+							else if (const IOptimusComputeKernelProvider* ConnectedKernel = Cast<const IOptimusComputeKernelProvider>(TargetNode))
+							{
+								if (KernelToGraphType[ConnectedNode] != KernelToGraphType[TargetRoutedNode])
+								{
+									bShouldUseImplicitPersistentDI = true;
+								}
 
-								TransientBufferDI->ValueType = Pin->GetDataType()->ShaderValueType;
-								TransientBufferDI->DataDomain = Pin->GetDataDomain();
-								TransientBufferDI->ComponentSourceBinding = KernelPrimaryBinding;
-								TransientBufferDI->DomainConstantIdentifier = {ConnectedNode.Node, NAME_None, Pin->GetFName()};
+								TargetKernelPins.Add(TargetPin);
+							}
+						}
+
+						if (TargetKernelPins.Num() > 0)
+						{
+							UOptimusRawBufferDataInterface* RawBufferDI = nullptr;
+							if (bShouldUseImplicitPersistentDI)
+							{
+								RawBufferDI = NewObject<UOptimusImplicitPersistentBufferDataInterface>(this);
 								if (KernelProvider->DoesOutputPinSupportAtomic(Pin))
 								{
-									TransientBufferDI->bZeroInitForAtomicWrites = true;
+									CastChecked<UOptimusImplicitPersistentBufferDataInterface>(RawBufferDI)->bZeroInitForAtomicWrites = true;
 								}
-							
-								LinkDataInterfaceMap.Add(Pin, TransientBufferDI);
-								DataInterfaceToBindingIndexMap.Add(TransientBufferDI) = PrimaryBindingIndex;	
 							}
 							else
 							{
-								UOptimusImplicitPersistentBufferDataInterface* ImplicitPersistentBufferDI =
-									NewObject<UOptimusImplicitPersistentBufferDataInterface>(this);
-
-								ImplicitPersistentBufferDI->ValueType = Pin->GetDataType()->ShaderValueType;
-								ImplicitPersistentBufferDI->DataDomain = Pin->GetDataDomain();
-								ImplicitPersistentBufferDI->ComponentSourceBinding = KernelPrimaryBinding;
-								ImplicitPersistentBufferDI->DomainConstantIdentifier = {ConnectedNode.Node, NAME_None, Pin->GetFName()};
+								RawBufferDI = NewObject<UOptimusTransientBufferDataInterface>(this);
 								if (KernelProvider->DoesOutputPinSupportAtomic(Pin))
 								{
-									ImplicitPersistentBufferDI->bZeroInitForAtomicWrites = true;
+									CastChecked<UOptimusTransientBufferDataInterface>(RawBufferDI)->bZeroInitForAtomicWrites = true;
 								}
-							
-								LinkDataInterfaceMap.Add(Pin, ImplicitPersistentBufferDI);
-								DataInterfaceToBindingIndexMap.Add(ImplicitPersistentBufferDI) = PrimaryBindingIndex;	
+							}
+								
+							RawBufferDI->ValueType = Pin->GetDataType()->ShaderValueType;
+							RawBufferDI->DataDomain = Pin->GetDataDomain();
+							RawBufferDI->ComponentSourceBinding = KernelPrimaryBinding;
+							DataInterfaceToBindingIndexMap.Add(RawBufferDI) = PrimaryBindingIndex;
+								
+							KernelOutputDataInterfaceMap.Add(InstancedPin) = RawBufferDI;
+
+							// All connected kernels share the same raw buffer data interface
+							for (const UOptimusNodePin* TargetKernelPin : TargetKernelPins)
+							{
+								KernelOutputMap.FindOrAdd(InstancedNode).FindOrAdd(Pin).Add({RawBufferDI, TargetKernelPin});		
 							}
 						}
 					}
 				}
 			}
 		}
-		else if (Cast<const IOptimusValueProvider>(ConnectedNode.Node))
-		{
-			ValueNodes.AddUnique(ConnectedNode.Node);
-		}
 	}
 
-	// Create the graph data interface and fill it with the value nodes.
-	UOptimusGraphDataInterface* GraphDataInterface = NewObject<UOptimusGraphDataInterface>(this);
 
-	TArray<FOptimusGraphVariableDescription> ValueNodeDescriptions;
-	ValueNodeDescriptions.Reserve(ValueNodes.Num());
-	for (UOptimusNode const* ValueNode : ValueNodes)
+	TArray<FOptimusComputeGraphInfo> GraphInfos;
+	for (EOptimusNodeGraphType GraphType : GraphTypes)
 	{
-		if (IOptimusValueProvider const* ValueProvider = Cast<const IOptimusValueProvider>(ValueNode))
+		FString Name = InNodeGraph->GetName();
+		if (GraphType != InNodeGraph->GraphType)
 		{
-			FOptimusGraphVariableDescription& ValueNodeDescription = ValueNodeDescriptions.AddDefaulted_GetRef();
-			ValueNodeDescription.Name = ValueProvider->GetValueName();
-			ValueNodeDescription.ValueType = ValueProvider->GetValueType()->ShaderValueType;
-
-			if (UOptimusNode_ConstantValue const* ConstantNode = Cast<const UOptimusNode_ConstantValue>(ValueNode))
-			{
-				ValueNodeDescription.Value = ConstantNode->GetShaderValue().ShaderValue;
-			}
+			check(GraphType == EOptimusNodeGraphType::Setup);
+			Name += TEXT("_Setup");
 		}
-	}
-	GraphDataInterface->Init(ValueNodeDescriptions);
+		
+		FOptimusComputeGraphInfo GraphInfo;
+		FName GraphName =
+			MakeUniqueObjectName(
+				this,
+				UOptimusComputeGraph::StaticClass(),
+				*Name
+				);
+		GraphInfo.GraphName = GraphName;
+		GraphInfo.GraphType = GraphType;
+		GraphInfo.ComputeGraph = NewObject<UOptimusComputeGraph>(this, GraphInfo.GraphName);
 
-	// The component binding for the graph data is the primary binding on the deformer.
-	UOptimusComponentSourceBinding* GraphDataComponentBinding = Bindings->Bindings[0];
-	DataInterfaceToBindingIndexMap.Add(GraphDataInterface) = GraphDataComponentBinding->GetIndex();
+		GraphInfos.Add(GraphInfo);
+	}
 
-	for (int32 ComputeGraphIndex = 0 ; ComputeGraphIndex < GraphInfos.Num(); ComputeGraphIndex++)
+	for (int32 GraphIndex = 0 ; GraphIndex < GraphInfos.Num(); GraphIndex++)
 	{
-		UOptimusComputeGraph* ComputeGraph = GraphInfos[ComputeGraphIndex].ComputeGraph;
-
+		FOptimusComputeGraphInfo& GraphInfo = GraphInfos[GraphIndex];
+		
+		UOptimusComputeGraph* ComputeGraph = GraphInfo.ComputeGraph;
+	
 		// Create the binding objects.
 		for (const UOptimusComponentSourceBinding* Binding: Bindings->Bindings)
 		{
@@ -1836,22 +2218,50 @@ TArray<FOptimusComputeGraphInfo> UOptimusDeformer::CompileNodeGraphToComputeGrap
 		ComputeGraph->DataInterfaces.Add(GraphDataInterface);
 		ComputeGraph->DataInterfaceToBinding.Add(0);		// Graph data interface always uses the primary binding.
 
-		for (TPair<const UOptimusNode*, UOptimusComputeDataInterface*>& Item : NodeDataInterfaceMap)
-		{
-			ComputeGraph->DataInterfaces.Add(Item.Value);
-			ComputeGraph->DataInterfaceToBinding.Add(DataInterfaceToBindingIndexMap[Item.Value]);
-		}
-		for (TPair<const UOptimusNodePin *, UOptimusComputeDataInterface *>&Item: LinkDataInterfaceMap)
+		for (const FOptimusInstancedNode& InstancedNode : InstancedNodes )
 		{
-			ComputeGraph->DataInterfaces.Add(Item.Value);
-			ComputeGraph->DataInterfaceToBinding.Add(DataInterfaceToBindingIndexMap[Item.Value]);
-		}
-		for (TPair<const UOptimusNode *, UComputeDataInterface *>&Item: KernelDataInterfaceMap)
-		{
-			ComputeGraph->DataInterfaces.Add(Item.Value);
-			ComputeGraph->DataInterfaceToBinding.Add(DataInterfaceToBindingIndexMap[Item.Value]);
-		}
+			const FOptimusRoutedConstNode& ConnectedNode = InstancedNode.RoutedNode;
+
+			if (const IOptimusComputeKernelProvider* KernelProvider = Cast<const IOptimusComputeKernelProvider>(ConnectedNode.Node))
+			{
+				if (KernelToGraphType[ConnectedNode] != GraphInfo.GraphType)
+				{
+					continue;
+				}
+				
+				const FOptimus_KernelInputMap KernelInputs = KernelInputMap[InstancedNode];
+				const FOptimus_KernelOutputMap KernelOutputs = KernelOutputMap[InstancedNode];
 
+				for (const TPair<const UOptimusNodePin*, FOptimus_KernelConnection>& Item : KernelInputs)
+				{
+					UComputeDataInterface* DataInterface = Item.Value.DataInterface;
+					if (!ComputeGraph->DataInterfaces.Contains(DataInterface))
+					{
+						ComputeGraph->DataInterfaces.Add(DataInterface);
+						ComputeGraph->DataInterfaceToBinding.Add(DataInterfaceToBindingIndexMap[DataInterface]);
+					}
+				}
+
+				for (const TPair<const UOptimusNodePin*, TArray<FOptimus_KernelConnection>>& Item : KernelOutputs)
+				{
+					for (const FOptimus_KernelConnection& Connection : Item.Value)
+					{
+						UComputeDataInterface* DataInterface = Connection.DataInterface;
+						if (!ComputeGraph->DataInterfaces.Contains(DataInterface))
+						{
+							ComputeGraph->DataInterfaces.Add(DataInterface);
+							ComputeGraph->DataInterfaceToBinding.Add(DataInterfaceToBindingIndexMap[DataInterface]);
+						}	
+					}
+				}
+
+				UComputeDataInterface* KernelDataInterface = KernelDataInterfaceMap[InstancedNode];
+
+				ComputeGraph->DataInterfaces.Add(KernelDataInterface);
+				ComputeGraph->DataInterfaceToBinding.Add(DataInterfaceToBindingIndexMap[KernelDataInterface]);
+			}
+		}
+			
 		// Create bound kernels
 		struct FKernelWithDataBindings
 		{
@@ -1861,11 +2271,14 @@ TArray<FOptimusComputeGraphInfo> UOptimusDeformer::CompileNodeGraphToComputeGrap
 		};
 
 		TArray<FKernelWithDataBindings> BoundKernels;
-		for (FOptimusRoutedConstNode ConnectedNode: ConnectedNodes)
+		
+		for (const FOptimusInstancedNode& InstancedNode : InstancedNodes )
 		{
-			if (const IOptimusComputeKernelProvider *KernelProvider = Cast<const IOptimusComputeKernelProvider>(ConnectedNode.Node))
+			const FOptimusRoutedConstNode& ConnectedNode = InstancedNode.RoutedNode;
+
+			if (const IOptimusComputeKernelProvider* KernelProvider = Cast<const IOptimusComputeKernelProvider>(ConnectedNode.Node))
 			{
-				if (KernelComputeGraphIndexMap[KernelProvider] != ComputeGraphIndex)
+				if (KernelToGraphType[ConnectedNode] != GraphInfo.GraphType)
 				{
 					continue;
 				}
@@ -1874,16 +2287,16 @@ TArray<FOptimusComputeGraphInfo> UOptimusDeformer::CompileNodeGraphToComputeGrap
 
 				BoundKernel.Kernel = NewObject<UComputeKernel>(this);
 
-				FOptimusKernelConstantContainer& KernelConstantContainer = ConstantContainer.AddContainerForKernel();
-				UComputeDataInterface* KernelDataInterface = KernelDataInterfaceMap[ConnectedNode.Node];
+				UComputeDataInterface* KernelDataInterface = KernelDataInterfaceMap[InstancedNode];
+
+				const FOptimus_KernelInputMap& KernelInputs = KernelInputMap[InstancedNode];
+				const FOptimus_KernelOutputMap& KernelOutputs = KernelOutputMap[InstancedNode];
 				
 				FOptimus_ComputeKernelResult KernelSourceResult = KernelProvider->CreateComputeKernel(	
 					BoundKernel.Kernel, ConnectedNode.TraversalContext,
-					NodeDataInterfaceMap, LinkDataInterfaceMap,
-					ValueNodes,
-					GraphDataInterface,
+					KernelInputs, KernelOutputs, ValueNodes,
 					KernelDataInterface,
-					BoundKernel.InputDataBindings, BoundKernel.OutputDataBindings, KernelConstantContainer
+					BoundKernel.InputDataBindings, BoundKernel.OutputDataBindings
 				);
 				if (FText* ErrorMessage = KernelSourceResult.TryGet<FText>())
 				{
@@ -2692,6 +3105,7 @@ UMeshDeformerInstance* UOptimusDeformer::CreateInstance(
 	Instance->SetInstanceSettings(Cast<UOptimusDeformerInstanceSettings>(InSettings));
 	Instance->SetupFromDeformer(this);
 
+	CompileEndDelegate.RemoveAll(Instance);
 	// Make sure all the instances know when we finish compiling so they can update their local state to match.
 	CompileEndDelegate.AddUObject(Instance, &UOptimusDeformerInstance::SetupFromDeformer);
 	ConstantValueUpdateDelegate.AddUObject(Instance, &UOptimusDeformerInstance::SetConstantValueDirect);
diff --git a/Engine/Plugins/Animation/DeformerGraph/Source/OptimusCore/Private/OptimusDeformerInstance.cpp b/Engine/Plugins/Animation/DeformerGraph/Source/OptimusCore/Private/OptimusDeformerInstance.cpp
--- a/Engine/Plugins/Animation/DeformerGraph/Source/OptimusCore/Private/OptimusDeformerInstance.cpp
+++ b/Engine/Plugins/Animation/DeformerGraph/Source/OptimusCore/Private/OptimusDeformerInstance.cpp
@@ -58,8 +58,8 @@ void FOptimusPersistentBufferPool::GetResourceBuffers(
 
 void FOptimusPersistentBufferPool::GetImplicitPersistentBuffers(
 	FRDGBuilder& GraphBuilder,
-	const FOptimusConstantIdentifier& InIdentifier, 
-	const FOptimusDeformerInstanceComponentLodContext& InLodContext,
+	FName DataInterfaceName,
+	int32 InLODIndex,
 	int32 InElementStride, 
 	int32 InRawStride,
 	TArray<int32> const& InElementCounts,
@@ -69,17 +69,19 @@ void FOptimusPersistentBufferPool::GetImplicitPersistentBuffers(
 	OutBuffers.Reset();
 	bOutJustAllocated = false;
 
-	TMap<FOptimusConstantIdentifier, TArray<FOptimusPersistentStructuredBuffer>>& LODResources = ImplicitBuffersMap.FindOrAdd(InLodContext);  
-	TArray<FOptimusPersistentStructuredBuffer>& ResourceBuffers = LODResources.FindOrAdd(InIdentifier);
-	if (ResourceBuffers.IsEmpty())
+	TMap<int32, TArray<FOptimusPersistentStructuredBuffer>>& LODResources = ImplicitBuffersMap.FindOrAdd(DataInterfaceName);  
+	TArray<FOptimusPersistentStructuredBuffer>* ResourceBuffersPtr = LODResources.Find(InLODIndex);
+	if (ResourceBuffersPtr == nullptr)
 	{
 		// Create pooled buffers and store.
+		TArray<FOptimusPersistentStructuredBuffer> ResourceBuffers;
 		AllocateBuffers(GraphBuilder, InElementStride, InRawStride, InElementCounts, ResourceBuffers, OutBuffers);
+		LODResources.Add(InLODIndex, MoveTemp(ResourceBuffers));
 		bOutJustAllocated = true;
 	}
 	else
 	{
-		ValidateAndGetBuffers(GraphBuilder,InElementStride, InElementCounts, ResourceBuffers, OutBuffers);
+		ValidateAndGetBuffers(GraphBuilder,InElementStride, InElementCounts, *ResourceBuffersPtr, OutBuffers);
 	}
 }
 
@@ -380,9 +382,6 @@ void UOptimusDeformerInstance::SetupFromDeformer(UOptimusDeformer* InDeformer)
 	ComputeGraphExecInfos.Reset();
 	GraphsToRunOnNextTick.Reset();
 	
-	ConstantContainer = InDeformer->ConstantContainer;
-	ConstantValuesPerContext.Reset();
-	
 	for (int32 GraphIndex = 0; GraphIndex < InDeformer->ComputeGraphs.Num(); ++GraphIndex)
 	{
 		FOptimusComputeGraphInfo const& ComputeGraphInfo = InDeformer->ComputeGraphs[GraphIndex];
@@ -472,103 +471,6 @@ void UOptimusDeformerInstance::SetCanBeActive(bool bInCanBeActive)
 	bCanBeActive = bInCanBeActive;
 }
 
-FOptimusConstantEvaluationResult UOptimusDeformerInstance::GetConstantValuePerInvocation(
-	const FOptimusConstantIdentifier& InIdentifier)
-{
-	FOptimusConstantEvaluationResult Result;
-
-	check(WeakBoundComponents.Num() == WeakComponentSources.Num());
-	int32 NumBindings = WeakBoundComponents.Num();
-	for (int32 BindingIndex = 0; BindingIndex < NumBindings; BindingIndex++)
-	{
-		int32 LodIndex = INDEX_NONE;
-		if (UActorComponent* Component = WeakBoundComponents[BindingIndex].Get())
-		{
-			if (const UOptimusComponentSource* ComponentSource = WeakComponentSources[BindingIndex].Get())
-			{
-				LodIndex = ComponentSource->GetLodIndex(Component);
-			}
-		}
-
-
-		if (LodIndex != INDEX_NONE)
-		{
-			Result.LodContext.LodIndexPerComponent.Add(LodIndex);
-		}
-		else
-		{
-			return {};
-		}
-	}
-	
-	
-	if (!InIdentifier.IsValid())
-	{
-		return Result;
-	}
-
-	if (const FOptimusConstantContainerInstance* ConstantValues = ConstantValuesPerContext.Find(Result.LodContext))
-	{
-		Result.ValuePerInvocation = ConstantValues->GetConstantValuePerInvocation(InIdentifier);
-		return Result;
-	}
-
-	TMap<int32, TMap<FName, TArray<float>>> BindingIndexToConstantValues;
-
-	auto CollectAllBindingConstants = [&]()
-	{
-		for (int32 BindingIndex = 0; BindingIndex < NumBindings; BindingIndex++)
-		{
-			TMap<FName, TArray<float>>& BindingConstantValues = BindingIndexToConstantValues.Add(BindingIndex);
-	
-			if (UActorComponent* Component = WeakBoundComponents[BindingIndex].Get())
-			{
-				if (const UOptimusComponentSource* ComponentSource = WeakComponentSources[BindingIndex].Get())
-				{
-					const int32 LodIndex = ComponentSource->GetLodIndex(Component);
-					TArray<FName> ComponentConstants = ComponentSource->GetExecutionDomains();
-
-					for (int32 ConstantIndex = 0; ConstantIndex < ComponentConstants.Num(); ConstantIndex++)
-					{
-						const FName& ConstantName = ComponentConstants[ConstantIndex];
-				
-						TArray<float>& Values = BindingConstantValues.Add(ConstantName);
-						
-						TArray<int32> Counts;
-						if (!ComponentSource->GetComponentElementCountsForExecutionDomain(
-							ConstantName,
-							Component, LodIndex, Counts))
-						{
-							return false;
-						}
-
-						Values.Reserve(Counts.Num());
-						for (const int32 Count : Counts)
-						{
-							Values.Add(Count);
-						}
-					}
-				}
-			}
-		}
-		
-		return true;
-	};
-
-
-	if (CollectAllBindingConstants())
-	{
-		FOptimusConstantContainerInstance& NewConstantValues = ConstantValuesPerContext.Add(Result.LodContext);
-		if (NewConstantValues.Initialize(ConstantContainer, BindingIndexToConstantValues))
-		{
-			Result.ValuePerInvocation = NewConstantValues.GetConstantValuePerInvocation(InIdentifier);
-			return Result;
-		}
-	}
-
-	return Result;
-}
-
 void UOptimusDeformerInstance::AllocateResources()
 {
 }
@@ -617,6 +519,9 @@ void UOptimusDeformerInstance::EnqueueWork(FEnqueueWorkDesc const& InDesc)
 		Swap(GraphsToRunOnNextTick, GraphsToRun);
 	}
 	
+	
+	
+	
 	// Enqueue work.
 	bool bIsWorkEnqueued = false;
 	if (bCanBeActive)
@@ -625,7 +530,24 @@ void UOptimusDeformerInstance::EnqueueWork(FEnqueueWorkDesc const& InDesc)
 		{
 			if (Info.GraphType == EOptimusNodeGraphType::Update || GraphsToRun.Contains(Info.GraphName))
 			{
-				bIsWorkEnqueued |= Info.ComputeGraphInstance.EnqueueWork(Info.ComputeGraph, InDesc.Scene, ExecutionGroupName, InDesc.OwnerName, InDesc.FallbackDelegate, this);
+				FSimpleDelegate FallbackDelegate = InDesc.FallbackDelegate;
+
+#if WITH_EDITOR
+				// Pending shader compilation may cause queued graph to not execute, so requeue on failure
+				// Ideally all compute graph should be ready before the deformer instance starts to enqueue work
+				if (Info.GraphType != EOptimusNodeGraphType::Update)
+				{
+					FallbackDelegate = FSimpleDelegate::CreateLambda([InDesc, this, GraphName=Info.GraphName]()
+					{
+						{
+							UE::TScopeLock<FCriticalSection> Lock(GraphsToRunOnNextTickLock);
+							GraphsToRunOnNextTick.Add(GraphName);
+						}	
+					});		
+				}
+#endif
+				
+				bIsWorkEnqueued |= Info.ComputeGraphInstance.EnqueueWork(Info.ComputeGraph, InDesc.Scene, ExecutionGroupName, InDesc.OwnerName, FallbackDelegate, this);
 			}
 		}
 	}
diff --git a/Engine/Plugins/Animation/DeformerGraph/Source/OptimusCore/Private/OptimusNode.cpp b/Engine/Plugins/Animation/DeformerGraph/Source/OptimusCore/Private/OptimusNode.cpp
--- a/Engine/Plugins/Animation/DeformerGraph/Source/OptimusCore/Private/OptimusNode.cpp
+++ b/Engine/Plugins/Animation/DeformerGraph/Source/OptimusCore/Private/OptimusNode.cpp
@@ -17,6 +17,9 @@
 #include "Actions/OptimusNodeGraphActions.h"
 
 #include "Algo/Reverse.h"
+#include "Serialization/MemoryReader.h"
+#include "Serialization/MemoryWriter.h"
+#include "Serialization/ObjectAndNameAsStringProxyArchive.h"
 #include "UObject/Package.h"
 #include "UObject/UObjectIterator.h"
 #include "UObject/UObjectGlobals.h"
@@ -125,7 +128,7 @@ FString UOptimusNode::GetNodePath() const
 		GraphPath = Graph->GetGraphPath();
 	}
 
-	return FString::Printf(TEXT("%s/%s"), *GraphPath, *GetName());
+	return UOptimusNodeGraph::ConstructPath(GraphPath, GetName(), {});
 }
 
 
@@ -134,6 +137,29 @@ UOptimusNodeGraph* UOptimusNode::GetOwningGraph() const
 	return Cast<UOptimusNodeGraph>(GetOuter());
 }
 
+TArray<UOptimusNodePin*> UOptimusNode::GetPinsByDirection(EOptimusNodePinDirection InDirection,
+	bool bInRecursive) const
+{
+	TArray<UOptimusNodePin*> Results;
+
+	for (UOptimusNodePin* Pin : Pins)
+	{
+		if (Pin->GetDirection() == InDirection)
+		{
+			if (bInRecursive)
+			{
+				Results.Append(Pin->GetSubPinsRecursively(true));
+			}
+			else
+			{
+				Results.Add(Pin);
+			}
+		}
+	}
+
+	return Results;
+}
+
 bool UOptimusNode::CanConnectPinToPin(
 	const UOptimusNodePin& InThisNodesPin,
 	const UOptimusNodePin& InOtherNodesPin,
@@ -351,6 +377,51 @@ TArray<UClass*> UOptimusNode::GetAllNodeClasses()
 	return NodeClasses;
 }
 
+FName UOptimusNode::GetAvailablePinNameStable(const UObject* InNodeOrPin, FName InName)
+{
+	TArray<const UOptimusNodePin*> PinsToCompare;
+	if (const UOptimusNode* Node = Cast<const UOptimusNode>(InNodeOrPin))
+	{
+		PinsToCompare.Append(Node->GetPins());
+	}
+	else if (const UOptimusNodePin* Pin = Cast<const UOptimusNodePin>(InNodeOrPin))
+	{
+		PinsToCompare.Append(Pin->GetSubPins());
+	}
+
+	TMap<FName, int32> BaseNameToMaxNumber;
+
+	for(const UOptimusNodePin* Pin : PinsToCompare)
+	{
+		const int32 Number = Pin->GetFName().GetNumber();
+		FName BaseName = Pin->GetFName();
+		BaseName.SetNumber(0);
+
+		if (const int32* ExistingNumber = BaseNameToMaxNumber.Find(BaseName))
+		{
+			BaseNameToMaxNumber[BaseName] = FMath::Max(Number, *ExistingNumber);
+		}
+		else
+		{
+			BaseNameToMaxNumber.Add(BaseName) = Number;
+		}
+	}
+
+	FName NewName = InName;
+	const int32 InputNumber = InName.GetNumber();
+	FName InputBaseName = InName;
+	InputBaseName.SetNumber(0);
+	if (const int32* ExistingNumber = BaseNameToMaxNumber.Find(InputBaseName))
+	{
+		if (InputNumber <= *ExistingNumber)
+		{
+			NewName.SetNumber(*ExistingNumber+1);
+		}
+	}
+
+	return NewName;
+}
+
 
 void UOptimusNode::PostCreateNode()
 {
@@ -422,6 +493,25 @@ void UOptimusNode::ConstructNode()
 }
 
 
+void UOptimusNode::SaveState(FArchive& Ar) const
+{
+	// Take a copy of the node's contents but not sub-data (like pins).
+	// Derived nodes may add additional data to the archive
+	
+	// This fella does the heavy lifting of serializing object references. 
+	// FMemoryWriter and fam do not handle UObject* serialization on their own.
+	FObjectAndNameAsStringProxyArchive NodeProxyArchive(
+			Ar, /* bInLoadIfFindFails=*/ false);
+	SerializeScriptProperties(NodeProxyArchive);
+}
+
+void UOptimusNode::RestoreState(FArchive& Ar)
+{
+	FObjectAndNameAsStringProxyArchive NodeProxyArchive(
+		Ar, /* bInLoadIfFindFails=*/true);
+	SerializeScriptProperties(NodeProxyArchive);
+}
+
 void UOptimusNode::EnableDynamicPins()
 {
 	bDynamicPins = true;
diff --git a/Engine/Plugins/Animation/DeformerGraph/Source/OptimusCore/Private/OptimusNodeGraph.cpp b/Engine/Plugins/Animation/DeformerGraph/Source/OptimusCore/Private/OptimusNodeGraph.cpp
--- a/Engine/Plugins/Animation/DeformerGraph/Source/OptimusCore/Private/OptimusNodeGraph.cpp
+++ b/Engine/Plugins/Animation/DeformerGraph/Source/OptimusCore/Private/OptimusNodeGraph.cpp
@@ -4,6 +4,7 @@
 
 #include "OptimusDeformer.h"
 #include "OptimusNode.h"
+#include "OptimusNodePair.h"
 #include "OptimusNodeLink.h"
 #include "OptimusNodePin.h"
 #include "OptimusActionStack.h"
@@ -27,6 +28,7 @@
 
 #include <limits>
 
+#include "IOptimusNodePairProvider.h"
 #include "IOptimusPinMutabilityDefiner.h"
 #include "OptimusNodeSubGraph.h"
 #include "Nodes/OptimusNode_CustomComputeKernel.h"
@@ -37,6 +39,7 @@
 
 #include "DataInterfaces/OptimusDataInterfaceAnimAttribute.h"
 #include "Nodes/OptimusNode_ComponentSource.h"
+#include "Nodes/OptimusNode_LoopTerminal.h"
 
 #include "Serialization/MemoryReader.h"
 #include "Serialization/MemoryWriter.h"
@@ -248,6 +251,27 @@ bool UOptimusNodeGraph::IsValidUserGraphName(
 	return FName::IsValidXName(InGraphName, TEXT("./$"), OutFailureReason);
 }
 
+FString UOptimusNodeGraph::ConstructPath(const FString& GraphPath, const FString& NodeName, const FString& PinPath)
+{
+	if (!ensure(!NodeName.IsEmpty()))
+	{
+		return {};
+	}
+	
+	FString Result = NodeName;
+	if (!GraphPath.IsEmpty())
+	{
+		Result = GraphPath + TEXT("/") + Result;
+	}
+
+	if (!PinPath.IsEmpty())
+	{
+		Result = Result + TEXT(".") + PinPath;
+	}
+
+	return Result;
+}
+
 int32 UOptimusNodeGraph::GetGraphIndex() const
 {
 	return GetCollectionOwner()->GetGraphs().IndexOfByKey(this);
@@ -291,6 +315,72 @@ UOptimusNode* UOptimusNodeGraph::AddNodeInternal(
 	return AddNodeAction->GetNode(GetActionStack()->GetGraphCollectionRoot());
 }
 
+TArray<UOptimusNode*> UOptimusNodeGraph::AddNodePairInternal(const TSubclassOf<UOptimusNode> InNodeClass,
+	const FVector2D& InPosition, TFunction<void(UOptimusNode*)> InFirstNodeConfigFunc, TFunction<void(UOptimusNode*)> InSecondNodeConfigFunc)
+{
+	FName NodeName = Optimus::GetUniqueNameForScope(this, *(InNodeClass->GetName() + TEXT("_Entry")));
+
+	FOptimusCompoundAction *Action = new FOptimusCompoundAction(TEXT("Add Loop Nodes"));
+	
+	FOptimusNodeGraphAction_AddNode *AddFirstNodeAction = new FOptimusNodeGraphAction_AddNode(
+		GetGraphPath(), InNodeClass, NodeName,
+		[InFirstNodeConfigFunc, InPosition](UOptimusNode *InNode) {
+			if (InFirstNodeConfigFunc)
+			{
+				InFirstNodeConfigFunc(InNode);
+			}
+			return InNode->SetGraphPositionDirect(InPosition); 
+		});
+
+	Action->AddSubAction(AddFirstNodeAction);
+	
+	NodeName = Optimus::GetUniqueNameForScope(this, *(InNodeClass->GetName() + TEXT("_Exit")));
+	
+	FOptimusNodeGraphAction_AddNode *AddSecondNodeAction = new FOptimusNodeGraphAction_AddNode(
+			GetGraphPath(), InNodeClass, NodeName,
+			[InSecondNodeConfigFunc, InPosition](UOptimusNode *InNode) {
+				if (InSecondNodeConfigFunc)
+				{
+					InSecondNodeConfigFunc(InNode);
+				}
+				return InNode->SetGraphPositionDirect(InPosition + FVector2D(100, 0)); 
+			});
+
+	Action->AddSubAction(AddSecondNodeAction);
+	
+	if (!GetActionStack()->RunAction(Action))
+	{
+		return {};
+	}
+
+	TArray<UOptimusNode*> AddedNodes;
+
+	AddedNodes.Add(AddFirstNodeAction->GetNode(GetActionStack()->GetGraphCollectionRoot()));
+	AddedNodes.Add(AddSecondNodeAction->GetNode(GetActionStack()->GetGraphCollectionRoot()));
+	
+	return AddedNodes;
+}
+
+void UOptimusNodeGraph::AddPairNodesToArray(
+	const TArray<UOptimusNode*>& InNodes,
+	TArray<UOptimusNode*>& OutNodes,
+	TArray<UOptimusNodePair*>& OutNodePairs
+	) const
+{
+	OutNodes = InNodes;
+	for (UOptimusNode* Node : InNodes)
+	{
+		if (Cast<IOptimusNodePairProvider>(Node))
+		{
+			if (UOptimusNodePair* NodePair = Node->GetOwningGraph()->GetNodePair(Node); ensure(NodePair))
+			{
+				OutNodes.AddUnique(NodePair->GetNodeCounterpart(Node));
+				OutNodePairs.AddUnique(NodePair);
+			}
+		}
+	}	
+}
+
 
 UOptimusNode* UOptimusNodeGraph::AddNode(
 	const TSubclassOf<UOptimusNode> InNodeClass, 
@@ -332,6 +422,50 @@ UOptimusNode* UOptimusNodeGraph::AddDataInterfaceNode(
 		});
 }
 
+TArray<UOptimusNode*> UOptimusNodeGraph::AddLoopTerminalNodes(const FVector2D& InPosition)
+{
+	FOptimusCompoundAction *Action = new FOptimusCompoundAction(TEXT("Add Loop Nodes"));
+
+	FName NodeName = Optimus::GetUniqueNameForScope(this, *(UOptimusNode_LoopTerminal::StaticClass()->GetName() + TEXT("_Entry")));
+	FOptimusNodeGraphAction_AddNode *AddEntryNodeAction = new FOptimusNodeGraphAction_AddNode(
+		GetGraphPath(), UOptimusNode_LoopTerminal::StaticClass(), NodeName,
+		[InPosition](UOptimusNode *InNode) {
+			UOptimusNode_LoopTerminal* Node = CastChecked<UOptimusNode_LoopTerminal>(InNode);
+			Node->TerminalType = EOptimusTerminalType::Entry;
+			return InNode->SetGraphPositionDirect(InPosition); 
+		});
+
+	TWeakPtr<FOptimusNodeGraphAction_AddNode> WeakAddEntryNodeAction =
+		StaticCastWeakPtr<FOptimusNodeGraphAction_AddNode>(Action->AddSubAction(AddEntryNodeAction));
+
+	NodeName = Optimus::GetUniqueNameForScope(this, *(UOptimusNode_LoopTerminal::StaticClass()->GetName() + TEXT("_Return"))); 
+
+	FOptimusNodeGraphAction_AddNode *AddReturnNodeAction = new FOptimusNodeGraphAction_AddNode(
+		GetGraphPath(), UOptimusNode_LoopTerminal::StaticClass(), NodeName,
+		[InPosition](UOptimusNode *InNode) {
+			UOptimusNode_LoopTerminal* Return = CastChecked<UOptimusNode_LoopTerminal>(InNode);
+			Return->TerminalType = EOptimusTerminalType::Return;
+			return InNode->SetGraphPositionDirect(InPosition + FVector2d(300, 0)); 
+		});
+
+	TWeakPtr<FOptimusNodeGraphAction_AddNode> WeakAddReturnNodeAction =
+	StaticCastWeakPtr<FOptimusNodeGraphAction_AddNode>(Action->AddSubAction(AddReturnNodeAction));
+
+	Action->AddSubAction<FOptimusNodeGraphAction_AddNodePair>(WeakAddEntryNodeAction, WeakAddReturnNodeAction);
+
+	if (!GetActionStack()->RunAction(Action))
+	{
+		return {};
+	}
+
+	TArray<UOptimusNode*> AddedNodes;
+
+	AddedNodes.Add(AddEntryNodeAction->GetNode(GetActionStack()->GetGraphCollectionRoot()));
+	AddedNodes.Add(AddReturnNodeAction->GetNode(GetActionStack()->GetGraphCollectionRoot()));
+	
+	return AddedNodes;
+}
+
 
 UOptimusNode* UOptimusNodeGraph::AddResourceNode(
 	UOptimusResourceDescription* InResourceDesc,
@@ -451,7 +585,11 @@ bool UOptimusNodeGraph::RemoveNodesToAction(
 		return false;
 	}
 
-	for (UOptimusNode* Node : InNodes)
+	TArray<UOptimusNode*> NodesToRemove;
+	TArray<UOptimusNodePair*> NodePairsToRemove;
+	AddPairNodesToArray(InNodes, NodesToRemove, NodePairsToRemove);
+	
+	for (UOptimusNode* Node : NodesToRemove)
 	{
 		if (Node == nullptr || Node->GetOwningGraph() != this)
 		{
@@ -462,7 +600,7 @@ bool UOptimusNodeGraph::RemoveNodesToAction(
 	TSet<int32> AllLinkIndexes;
 
 	// Get all unique links for all the given nodes and remove them *before* we remove the nodes.
-	for (const UOptimusNode* Node : InNodes)
+	for (const UOptimusNode* Node : NodesToRemove)
 	{
 		AllLinkIndexes.Append(GetAllLinkIndexesToNode(Node));
 	}
@@ -472,7 +610,13 @@ bool UOptimusNodeGraph::RemoveNodesToAction(
 		InAction->AddSubAction<FOptimusNodeGraphAction_RemoveLink>(Links[LinkIndex]);
 	}
 
-	for (UOptimusNode* Node : InNodes)
+	for (const UOptimusNodePair* NodePair: NodePairsToRemove)
+	{
+		InAction->AddSubAction<FOptimusNodeGraphAction_RemoveNodePair>(NodePair);
+	}	
+	
+	
+	for (UOptimusNode* Node : NodesToRemove)
 	{
 		InAction->AddSubAction<FOptimusNodeGraphAction_RemoveNode>(Node);
 	}
@@ -490,6 +634,13 @@ UOptimusNode* UOptimusNodeGraph::DuplicateNode(
 	{
 		return nullptr;
 	}
+
+	// This API does not support duplicating loop terminal currently
+	// Use DuplicateNodes() below
+	if (!ensure(!Cast<UOptimusNode_LoopTerminal>(InNode)))
+	{
+		return nullptr;
+	}
 	
 	const FName NodeName = Optimus::GetUniqueNameForScope(this, InNode->GetFName());
 	
@@ -542,6 +693,11 @@ bool UOptimusNodeGraph::DuplicateNodes(
 		return false;
 	}
 
+
+	TArray<UOptimusNode*> NodesToDuplicate;
+	TArray<UOptimusNodePair*> NodePairsToDuplicate;
+	AddPairNodesToArray(InNodes, NodesToDuplicate, NodePairsToDuplicate);
+
 	// Figure out the non-clashing names to use, to avoid collisions during actual execution.
 	TSet<FName> ExistingObjects;
 	for (const UOptimusNode* Node: Nodes)
@@ -565,16 +721,35 @@ bool UOptimusNodeGraph::DuplicateNodes(
 	using FloatType = decltype(FVector2D::X);
 	FVector2D TopLeft{std::numeric_limits<FloatType>::max()};
 	TMap<UOptimusNode*, FName> NewNodeNameMap;
-	for (UOptimusNode* Node: InNodes)
+	for (UOptimusNode* Node: NodesToDuplicate)
 	{
 		TopLeft = FVector2D::Min(TopLeft, Node->GraphPosition);
 		NewNodeNameMap.Add(Node, MakeUniqueNodeName(Node->GetFName()));
 	}
 	FVector2D NodeOffset = InPosition - TopLeft;
 
+	const FString GraphPath = GetGraphPath();
+	
+	TArray<TPair<FString, FString>> NewNodePairs;
+	for (const UOptimusNodePair* NodePair : NodePairsToDuplicate)
+	{
+		// Genernerate new pair
+		const UOptimusNode *FirstNode = NodePair->GetFirst();
+		const UOptimusNode *SecondNode = NodePair->GetSecond();
+
+		if (ensure(NewNodeNameMap.Contains(FirstNode) && NewNodeNameMap.Contains(SecondNode)))
+		{
+			FString NewFirstNodePath = ConstructPath(GraphPath, NewNodeNameMap[FirstNode].ToString(), {});
+			FString NewSecondNodePath = ConstructPath(GraphPath, NewNodeNameMap[SecondNode].ToString(), {});
+
+			NewNodePairs.Add({NewFirstNodePath, NewSecondNodePath});
+		}	
+	}
+
+	
 	/// Collect the links between these existing nodes. 
 	TArray<TPair<FString, FString>> NodeLinks;
-	const FString GraphPath = GetGraphPath();
+	
 	for (const UOptimusNodeLink* Link: SourceGraph->GetAllLinks())
 	{
 		const UOptimusNode *OutputNode = Link->GetNodeOutputPin()->GetOwningNode();
@@ -584,34 +759,32 @@ bool UOptimusNodeGraph::DuplicateNodes(
 		{
 			// FIXME: This should be a utility function, along with all the other path creation
 			// functions.
-			FString NodeOutputPinPath = FString::Printf(TEXT("%s/%s.%s"),
-				*GraphPath, *NewNodeNameMap[OutputNode].ToString(), *Link->GetNodeOutputPin()->GetUniqueName().ToString());
-			FString NodeInputPinPath = FString::Printf(TEXT("%s/%s.%s"),
-				*GraphPath, *NewNodeNameMap[InputNode].ToString(), *Link->GetNodeInputPin()->GetUniqueName().ToString());
-
+			FString NodeOutputPinPath = ConstructPath(GraphPath, NewNodeNameMap[OutputNode].ToString(), Link->GetNodeOutputPin()->GetUniqueName().ToString());
+			FString NodeInputPinPath = ConstructPath(GraphPath, NewNodeNameMap[InputNode].ToString(), Link->GetNodeInputPin()->GetUniqueName().ToString());
+			
 			NodeLinks.Add(MakeTuple(NodeOutputPinPath, NodeInputPinPath));
 		}
 	}
 
 	FOptimusCompoundAction *Action = new FOptimusCompoundAction;
-	if (InNodes.Num() == 1)
+	if (NodesToDuplicate.Num() == 1)
 	{
 		Action->SetTitlef(TEXT("%s Node"), *InActionName);
 	}
 	else
 	{
-		Action->SetTitlef(TEXT("%s %d Nodes"), *InActionName, InNodes.Num());
+		Action->SetTitlef(TEXT("%s %d Nodes"), *InActionName, NodesToDuplicate.Num());
 	}
 
 	// Add all the pre-duplicate requirement actions first. This allows certain nodes to set up their
 	// operating environment correctly (e.g. a missing variable description for get variable nodes, etc.)
-	for (UOptimusNode* Node: InNodes)
+	for (UOptimusNode* Node: NodesToDuplicate)
 	{
 		Node->PreDuplicateRequirementActions(this, Action);
 	}
 
 	// Duplicate the nodes and place them correctly
-	for (UOptimusNode* Node: InNodes)
+	for (UOptimusNode* Node: NodesToDuplicate)
 	{
 		FOptimusNodeGraphAction_DuplicateNode *DuplicateNodeAction = new FOptimusNodeGraphAction_DuplicateNode(
 			GetGraphPath(), Node, NewNodeNameMap[Node],
@@ -622,6 +795,11 @@ bool UOptimusNodeGraph::DuplicateNodes(
 		Action->AddSubAction(DuplicateNodeAction);
 	}
 
+	for (const TPair<FString, FString>& NewPair : NewNodePairs)
+	{
+		Action->AddSubAction<FOptimusNodeGraphAction_AddNodePair>(NewPair.Key, NewPair.Value);
+	}
+	
 	// Add any links that the nodes may have had. These operations are allowed to fail, in which case
 	// we just end up with dangling nodes.
 	// In the future we would like to introduce connections that are in an error state (will fail compile)
@@ -1118,6 +1296,30 @@ bool UOptimusNodeGraph::IsSubGraphReference(UOptimusNode* InNode) const
 	return Cast<UOptimusNode_SubGraphReference>(InNode) != nullptr;
 }
 
+UOptimusNodePair* UOptimusNodeGraph::GetNodePair(const UOptimusNode* InNode) const
+{
+	for (int32 Index = 0; Index < NodePairs.Num(); Index++)
+	{
+		UOptimusNodePair* NodePair = NodePairs[Index];
+
+		if (NodePair->Contains(InNode))
+		{
+			return NodePair;
+		}
+	}
+
+	return nullptr;
+}
+
+UOptimusNode* UOptimusNodeGraph::GetNodeCounterpart(const UOptimusNode* InNode) const
+{
+	if (UOptimusNodePair* NodePair = GetNodePair(InNode))
+	{
+		return NodePair->GetNodeCounterpart(InNode);
+	}
+
+	return nullptr;
+}
 
 UOptimusNode* UOptimusNodeGraph::CreateNodeDirect(
 	const UClass* InNodeClass,
@@ -1219,6 +1421,49 @@ bool UOptimusNodeGraph::RemoveNodeDirect(
 	return true;
 }
 
+bool UOptimusNodeGraph::AddNodePairDirect(UOptimusNode* InFirstNode, UOptimusNode* InSecondNode)
+{
+	if (!ensure(Cast<IOptimusNodePairProvider>(InFirstNode) && Cast<IOptimusNodePairProvider>(InSecondNode)))
+	{
+		return false;
+	}
+
+	UOptimusNodePair* NewPair = NewObject<UOptimusNodePair>(this);
+	NewPair->First = InFirstNode;
+	NewPair->Second = InSecondNode;
+
+	NodePairs.Add(NewPair);
+
+	IOptimusNodePairProvider* FirstProvider = CastChecked<IOptimusNodePairProvider>(InFirstNode);
+	IOptimusNodePairProvider* SecondProvider = CastChecked<IOptimusNodePairProvider>(InSecondNode);
+
+	FirstProvider->PairToCounterpartNode(SecondProvider);
+	SecondProvider->PairToCounterpartNode(FirstProvider);
+	
+	return true;
+}
+
+bool UOptimusNodeGraph::RemoveNodePairDirect(UOptimusNode* InFirstNode, UOptimusNode* InSecondNode)
+{
+	if (!ensure(Cast<IOptimusNodePairProvider>(InFirstNode) && Cast<IOptimusNodePairProvider>(InSecondNode)))
+	{
+		return false;
+	}
+	
+	for (int32 Index = 0; Index < NodePairs.Num(); Index++)
+	{
+		const UOptimusNodePair* NodePair = NodePairs[Index];
+
+		if (NodePair->Contains(InFirstNode, InSecondNode))
+		{
+			RemoveNodePairByIndex(Index);
+			return true;
+		}
+	}
+
+	return false;
+}
+
 
 bool UOptimusNodeGraph::AddLinkDirect(UOptimusNodePin* NodeOutputPin, UOptimusNodePin* NodeInputPin)
 {
@@ -1529,18 +1774,137 @@ bool UOptimusNodeGraph::IsPinMutable(const UOptimusNodePin* InNodePin) const
 
 		if (const IOptimusPinMutabilityDefiner* PinMutabilityDefiner = Cast<const IOptimusPinMutabilityDefiner>(Node))
 		{
-			if (PinMutabilityDefiner->IsOutputPinMutable(NodePin))
+			EOptimusPinMutability Mutability = PinMutabilityDefiner->GetOutputPinMutability(NodePin);
+			if (Mutability != EOptimusPinMutability::Undefined)
 			{
-				return true;
+				if (Mutability == EOptimusPinMutability::Mutable)
+				{
+					return true;
+				}
+
+				// Dead-end on a immutable pin, consume this work item and don't look further
+				continue;
 			}
+
+			// Mutability is undefined, keep gathering from connected upstream nodes
 		}
 		
-		// Grab all inputs.
-		for (const UOptimusNodePin* Pin: Node->GetPins())
+		// Traverse in the direction of inputs to outputs (up the graph).
+		for (const UOptimusNodePin* Pin: Node->GetPinsByDirection(EOptimusNodePinDirection::Input, true))
+		{
+			for (const FOptimusRoutedNodePin& ConnectedPin: Pin->GetConnectedPinsWithRouting(WorkItem.TraversalContext))
+			{
+				if (ensure(ConnectedPin.NodePin != nullptr))
+				{
+					FOptimusRoutedConstNodePin CollectedNodePin{ConnectedPin.NodePin, ConnectedPin.TraversalContext};
+					const UOptimusNode *NextNode = ConnectedPin.NodePin->GetOwningNode();
+					FOptimusRoutedConstNode CollectedNode{NextNode, ConnectedPin.TraversalContext};
+
+					if (!VisitedNodes.Contains(CollectedNode))
+					{
+						VisitedNodes.Add(CollectedNode);
+						WorkingSet.Enqueue(CollectedNodePin);
+					}
+				}
+			}
+		}
+
+		if (const UOptimusNode_LoopTerminal* LoopTerminal = Cast<const UOptimusNode_LoopTerminal>(Node))
 		{
-			if (Pin->GetDirection() == EOptimusNodePinDirection::Input)
+			if (LoopTerminal->GetTerminalType() == EOptimusTerminalType::Entry && ensure(!NodePin->GetDataDomain().IsSingleton()))
 			{
-				InputPins.Append(Pin->GetSubPinsRecursively(true));
+				const UOptimusNodePin* NextPin = LoopTerminal->GetPinCounterpart(NodePin, EOptimusTerminalType::Return, EOptimusNodePinDirection::Output);
+				
+				const UOptimusNode*	NextNode = LoopTerminal->GetOtherTerminal();
+				FOptimusRoutedConstNode CollectedNode{NextNode, WorkItem.TraversalContext};
+
+				if (!VisitedNodes.Contains(CollectedNode))
+				{
+					VisitedNodes.Add(CollectedNode);
+					WorkingSet.Enqueue({NextPin, WorkItem.TraversalContext});
+				}
+			}
+		}
+	}
+
+	return false;
+}
+
+bool UOptimusNodeGraph::DoesNodeHaveMutableInput(const UOptimusNode* InNode) const
+{
+	for (UOptimusNodePin* Pin : InNode->GetPinsByDirection(EOptimusNodePinDirection::Input, true))
+	{
+		if (IsPinMutable(Pin))
+		{
+			return true;
+		}
+	}
+
+	return false;
+}
+
+TSet<FOptimusRoutedConstNode> UOptimusNodeGraph::GetLoopEntryTerminalForPin(
+	const UOptimusNodePin* InNodePin,
+	const FOptimusPinTraversalContext& InContext) const
+{
+	TSet<FOptimusRoutedConstNode> VisitedNodes;
+	TSet<FOptimusRoutedConstNode> EntryTerminals;
+	
+	TQueue<FOptimusRoutedConstNode> WorkingSet;
+
+	// If given an input pin, find the other side.
+	FOptimusRoutedConstNode StartNode;
+	if (InNodePin->GetDirection() == EOptimusNodePinDirection::Input)
+	{
+		TArray<FOptimusRoutedNodePin> RoutedPins = GetConnectedPinsWithRouting(InNodePin, InContext);
+		if (RoutedPins.IsEmpty())
+		{
+			// If the input pin is directly on a compute kernel and not connected to anything
+			// we consider it to not have a component source, even if other pins in the same group
+			// has one. Pin group level gathering is caller's responsibility
+			return {};
+		}
+		
+		check(RoutedPins.Num() == 1);
+		StartNode = {RoutedPins[0].NodePin->GetOwningNode(), RoutedPins[0].TraversalContext};
+	}
+	else
+	{
+		StartNode = {InNodePin->GetOwningNode(), InContext};
+	}
+
+	WorkingSet.Enqueue(StartNode);
+	
+	FOptimusRoutedConstNode WorkItem;
+	while (WorkingSet.Dequeue(WorkItem))
+	{
+		TArray<const UOptimusNodePin*> InputPins; 
+		const UOptimusNode* Node = WorkItem.Node;
+
+		if (const UOptimusNode_LoopTerminal* LoopTerminal = Cast<const UOptimusNode_LoopTerminal>(Node))
+		{
+			if (LoopTerminal->GetTerminalType() == EOptimusTerminalType::Entry)
+			{
+				EntryTerminals.Add(WorkItem);
+			}
+			continue;
+		}
+		
+		if (const IOptimusComputeKernelProvider* KernelProvider = Cast<const IOptimusComputeKernelProvider>(Node))
+		{
+			const UOptimusNodePin* PrimaryGroupPin = KernelProvider->GetPrimaryGroupPin();
+
+			InputPins.Append(PrimaryGroupPin->GetSubPins());
+		}
+		else
+		{
+			// Grab all inputs.
+			for (const UOptimusNodePin* Pin: Node->GetPins())
+			{
+				if (Pin->GetDirection() == EOptimusNodePinDirection::Input)
+				{
+					InputPins.Add(Pin);
+				}
 			}
 		}
 		
@@ -1551,21 +1915,33 @@ bool UOptimusNodeGraph::IsPinMutable(const UOptimusNodePin* InNodePin) const
 			{
 				if (ensure(ConnectedPin.NodePin != nullptr))
 				{
-					FOptimusRoutedConstNodePin CollectedNodePin{ConnectedPin.NodePin, ConnectedPin.TraversalContext};
-					const UOptimusNode *NextNode = ConnectedPin.NodePin->GetOwningNode();
+					UOptimusNode *NextNode = ConnectedPin.NodePin->GetOwningNode();
 					FOptimusRoutedConstNode CollectedNode{NextNode, ConnectedPin.TraversalContext};
-
+					
 					if (!VisitedNodes.Contains(CollectedNode))
 					{
 						VisitedNodes.Add(CollectedNode);
-						WorkingSet.Enqueue(CollectedNodePin);
+						WorkingSet.Enqueue(CollectedNode);
 					}
 				}
 			}
 		}
 	}
 
-	return false;
+	return EntryTerminals;
+}
+
+TSet<FOptimusRoutedConstNode> UOptimusNodeGraph::GetLoopEntryTerminalForNode(
+	const UOptimusNode* InNode, 
+	const FOptimusPinTraversalContext& InContext) const
+{
+	TSet<FOptimusRoutedConstNode> ConnectedEntryTerminals; 
+	for (UOptimusNodePin* Pin : InNode->GetPinsByDirection(EOptimusNodePinDirection::Input, true))
+	{
+		ConnectedEntryTerminals.Append(GetLoopEntryTerminalForPin(Pin, InContext));
+	}
+
+	return ConnectedEntryTerminals;
 }
 
 
@@ -1591,6 +1967,15 @@ TArray<const UOptimusNodeLink*> UOptimusNodeGraph::GetPinLinks(
 }
 
 
+void UOptimusNodeGraph::RemoveNodePairByIndex(int32 NodePairIndex)
+{
+	UOptimusNodePair* NodePair = NodePairs[NodePairIndex];
+
+	NodePairs.RemoveAt(NodePairIndex);
+
+	Optimus::RemoveObject(NodePair);
+}
+
 void UOptimusNodeGraph::RemoveLinkByIndex(int32 LinkIndex)
 {
 	UOptimusNodeLink* Link = Links[LinkIndex];
diff --git a/Engine/Plugins/Animation/DeformerGraph/Source/OptimusCore/Private/OptimusNodePair.cpp b/Engine/Plugins/Animation/DeformerGraph/Source/OptimusCore/Private/OptimusNodePair.cpp
new file mode 100644
--- /dev/null
+++ b/Engine/Plugins/Animation/DeformerGraph/Source/OptimusCore/Private/OptimusNodePair.cpp
@@ -0,0 +1,43 @@
+// Copyright Epic Games, Inc. All Rights Reserved.
+
+#include "OptimusNodePair.h"
+#include "OptimusNode.h"
+#include "IOptimusNodePairProvider.h"
+
+#include UE_INLINE_GENERATED_CPP_BY_NAME(OptimusNodePair)
+
+void UOptimusNodePair::PostDuplicate(bool bDuplicateForPIE)
+{
+	Super::PostDuplicate(bDuplicateForPIE);
+
+	IOptimusNodePairProvider* FirstProvider = CastChecked<IOptimusNodePairProvider>(First);
+	IOptimusNodePairProvider* SecondProvider = CastChecked<IOptimusNodePairProvider>(Second);
+
+	FirstProvider->PairToCounterpartNode(SecondProvider);
+	SecondProvider->PairToCounterpartNode(FirstProvider);
+}
+
+bool UOptimusNodePair::Contains(const UOptimusNode* InFirst, const UOptimusNode* InSecond) const
+{
+	return (First == InFirst && Second == InSecond) || (First == InSecond && Second == InFirst);	
+}
+
+bool UOptimusNodePair::Contains(const UOptimusNode* InNode) const
+{
+	return First == InNode || Second == InNode;	
+}
+
+UOptimusNode* UOptimusNodePair::GetNodeCounterpart(const UOptimusNode* InNode) const
+{
+	if (First == InNode)
+	{
+		return Second;
+	}
+
+	if (Second == InNode)
+	{
+		return First;
+	}
+
+	return nullptr;
+}
diff --git a/Engine/Plugins/Animation/DeformerGraph/Source/OptimusCore/Private/OptimusNodePin.cpp b/Engine/Plugins/Animation/DeformerGraph/Source/OptimusCore/Private/OptimusNodePin.cpp
--- a/Engine/Plugins/Animation/DeformerGraph/Source/OptimusCore/Private/OptimusNodePin.cpp
+++ b/Engine/Plugins/Animation/DeformerGraph/Source/OptimusCore/Private/OptimusNodePin.cpp
@@ -147,7 +147,10 @@ FText UOptimusNodePin::GetTooltipText() const
 
 FString UOptimusNodePin::GetPinPath() const
 {
-	return FString::Printf(TEXT("%s.%s"), *GetOwningNode()->GetNodePath(), *GetUniqueName().ToString());
+	return UOptimusNodeGraph::ConstructPath(
+		GetOwningNode()->GetOwningGraph()->GetGraphPath(),
+		GetOwningNode()->GetName(),
+		GetUniqueName().ToString());
 }
 
 
diff --git a/Engine/Plugins/Animation/DeformerGraph/Source/OptimusCore/Public/IOptimusComputeKernelDataInterface.h b/Engine/Plugins/Animation/DeformerGraph/Source/OptimusCore/Public/IOptimusComputeKernelDataInterface.h
--- a/Engine/Plugins/Animation/DeformerGraph/Source/OptimusCore/Public/IOptimusComputeKernelDataInterface.h
+++ b/Engine/Plugins/Animation/DeformerGraph/Source/OptimusCore/Public/IOptimusComputeKernelDataInterface.h
@@ -6,7 +6,7 @@
 
 #include "IOptimusComputeKernelDataInterface.generated.h"
 
-struct FOptimusConstantIdentifier;
+class UOptimusComponentSourceBinding;
 
 UINTERFACE()
 class OPTIMUSCORE_API UOptimusComputeKernelDataInterface :
@@ -22,5 +22,7 @@ class OPTIMUSCORE_API IOptimusComputeKernelDataInterface
 {
 public:
 	GENERATED_BODY()
-	virtual void SetExecutionDomainConstant(const FOptimusConstantIdentifier& InExecutionDomainConstantIdentifier) = 0;
+	virtual void SetExecutionDomain(const FString& InExecutionDomain) = 0;
+	virtual void SetComponentBinding(const UOptimusComponentSourceBinding* InBinding) = 0;
+	
 };
diff --git a/Engine/Plugins/Animation/DeformerGraph/Source/OptimusCore/Public/IOptimusComputeKernelProvider.h b/Engine/Plugins/Animation/DeformerGraph/Source/OptimusCore/Public/IOptimusComputeKernelProvider.h
--- a/Engine/Plugins/Animation/DeformerGraph/Source/OptimusCore/Public/IOptimusComputeKernelProvider.h
+++ b/Engine/Plugins/Animation/DeformerGraph/Source/OptimusCore/Public/IOptimusComputeKernelProvider.h
@@ -14,9 +14,11 @@ class UOptimusComponentSourceBinding;
 class UOptimusKernelSource;
 class UOptimusNode;
 class UOptimusNodePin;
+class UOptimusNode_LoopTerminal;
 struct FOptimusPinTraversalContext;
 struct FOptimusKernelConstantContainer;
 struct FOptimusExecutionDomain;
+struct FOptimusRoutedConstNodePin;
 
 // Maps the data interface's data binding index to the function we would like to have present
 // during kernel compilation to read/write values from/to that data interface's resource.
@@ -55,6 +57,15 @@ using FOptimus_KernelNodeToKernelDataInterfaceMap = TMap<const UOptimusNode*, UC
 
 using FOptimus_ComputeKernelResult = TVariant<UOptimusKernelSource* /* Kernel */, FText /* Error */>;
 
+struct FOptimus_KernelConnection
+{
+	UComputeDataInterface* DataInterface;
+	const UOptimusNodePin* ConnectedPin;
+};
+
+using FOptimus_KernelInputMap = TMap<const UOptimusNodePin*, FOptimus_KernelConnection>;
+using FOptimus_KernelOutputMap = TMap<const UOptimusNodePin*, TArray<FOptimus_KernelConnection>>;
+
 UINTERFACE()
 class OPTIMUSCORE_API UOptimusComputeKernelProvider :
 	public UInterface
@@ -75,26 +86,20 @@ public:
 	 * this interface.
 	 * @param InKernelSourceOuter The outer object that will own the new kernel source.
 	 * @param InTraversalContext The current context being used to traverse from the graph. Used mainly to keep track of graph nesting.
-	 * @param InNodeDataInterfaceMap Map from UOptimusNode objects to data interfaces.
-	 * @param InLinkDataInterfaceMap
-	 * @param InValueNodes
-	 * @param InGraphDataInterface
+	 * @param InValueNodes 
 	 * @param InOutKernelDataInterface
 	 * @param OutInputDataBindings
 	 * @param OutOutputDataBindings
-	 * @param OutKernelConstantContainer
 	 */
 	virtual FOptimus_ComputeKernelResult CreateComputeKernel(
 		UObject* InKernelSourceOuter,
 		const FOptimusPinTraversalContext& InTraversalContext,
-		const FOptimus_NodeToDataInterfaceMap& InNodeDataInterfaceMap,
-		const FOptimus_PinToDataInterfaceMap& InLinkDataInterfaceMap,
+		const FOptimus_KernelInputMap InKernelInputs,
+		const FOptimus_KernelOutputMap InKernelOutputs,
 		const TArray<const UOptimusNode*>& InValueNodes,
-		const UComputeDataInterface* InGraphDataInterface,
 		UComputeDataInterface* InOutKernelDataInterface,
 		FOptimus_InterfaceBindingMap& OutInputDataBindings,
-		FOptimus_InterfaceBindingMap& OutOutputDataBindings,
-		FOptimusKernelConstantContainer& OutKernelConstantContainer
+		FOptimus_InterfaceBindingMap& OutOutputDataBindings
 	) const = 0;
 
 	/** Returns the execution domain that this kernel should iterate over */
@@ -111,6 +116,4 @@ public:
 	
 	/** Check if a specific pin needs to support atomic operation */
 	virtual bool DoesOutputPinSupportRead(const UOptimusNodePin* InPin) const = 0;
-
-	virtual bool HasMutableInput() const = 0;
 };
diff --git a/Engine/Plugins/Animation/DeformerGraph/Source/OptimusCore/Public/IOptimusNodeAdderPinProvider.h b/Engine/Plugins/Animation/DeformerGraph/Source/OptimusCore/Public/IOptimusNodeAdderPinProvider.h
--- a/Engine/Plugins/Animation/DeformerGraph/Source/OptimusCore/Public/IOptimusNodeAdderPinProvider.h
+++ b/Engine/Plugins/Animation/DeformerGraph/Source/OptimusCore/Public/IOptimusNodeAdderPinProvider.h
@@ -28,7 +28,7 @@ public:
 		EOptimusNodePinDirection NewPinDirection;
 		FName Key;
 		FText ToolTip;
-		bool bCanAutoLink;
+		bool bCanAutoLink = false;;
 	};
 	
 	virtual TArray<FAdderPinAction> GetAvailableAdderPinActions(
@@ -36,10 +36,13 @@ public:
 		EOptimusNodePinDirection InNewPinDirection,
 		FString* OutReason = nullptr
 		) const = 0;
-	
+
+	// Make sure pins created this way don't have different names each time the function is called. Because
+	// this function is called during undo/redo, output has to be stable(same input gives same output)
 	virtual TArray<UOptimusNodePin*> TryAddPinFromPin(
 		const FAdderPinAction& InSelectedAction,
-		UOptimusNodePin* InSourcePin
+		UOptimusNodePin* InSourcePin,
+		FName InNameToUse
 		) = 0;
 	
 	virtual bool RemoveAddedPins(TConstArrayView<UOptimusNodePin*> InAddedPinsToRemove) = 0;
diff --git a/Engine/Plugins/Animation/DeformerGraph/Source/OptimusCore/Public/IOptimusNodePairProvider.h b/Engine/Plugins/Animation/DeformerGraph/Source/OptimusCore/Public/IOptimusNodePairProvider.h
new file mode 100644
--- /dev/null
+++ b/Engine/Plugins/Animation/DeformerGraph/Source/OptimusCore/Public/IOptimusNodePairProvider.h
@@ -0,0 +1,27 @@
+﻿// Copyright Epic Games, Inc. All Rights Reserved.
+
+#pragma once
+
+#include "UObject/Interface.h"
+
+#include "IOptimusNodePairProvider.generated.h"
+
+UINTERFACE()
+class OPTIMUSCORE_API UOptimusNodePairProvider :
+	public UInterface
+{
+	GENERATED_BODY()
+};
+
+/**
+* Interface that provides a mechanism to add pins to node from existing pins
+*/
+class OPTIMUSCORE_API IOptimusNodePairProvider
+{
+	GENERATED_BODY()
+
+public:
+
+	virtual void PairToCounterpartNode(const IOptimusNodePairProvider* NodePairProvider) = 0;
+};
+
diff --git a/Engine/Plugins/Animation/DeformerGraph/Source/OptimusCore/Public/IOptimusPinMutabilityDefiner.h b/Engine/Plugins/Animation/DeformerGraph/Source/OptimusCore/Public/IOptimusPinMutabilityDefiner.h
--- a/Engine/Plugins/Animation/DeformerGraph/Source/OptimusCore/Public/IOptimusPinMutabilityDefiner.h
+++ b/Engine/Plugins/Animation/DeformerGraph/Source/OptimusCore/Public/IOptimusPinMutabilityDefiner.h
@@ -15,12 +15,19 @@ class OPTIMUSCORE_API UOptimusPinMutabilityDefiner :
 	GENERATED_BODY()
 };
 
+UENUM()
+enum class EOptimusPinMutability
+{
+	Undefined,
+	Immutable,
+	Mutable,
+};
 
 class IOptimusPinMutabilityDefiner
 {
 	GENERATED_BODY()
 
 public:
-	/** Returns the component binding for the node that this interface is implemented on */
-	virtual bool IsOutputPinMutable(const UOptimusNodePin* InPin) const = 0;
+	// Returns whether the output pin is mutable, immutable or undefined
+	virtual EOptimusPinMutability GetOutputPinMutability(const UOptimusNodePin* InPin) const = 0;
 };
diff --git a/Engine/Plugins/Animation/DeformerGraph/Source/OptimusCore/Public/IOptimusUnnamedNodePinProvider.h b/Engine/Plugins/Animation/DeformerGraph/Source/OptimusCore/Public/IOptimusUnnamedNodePinProvider.h
new file mode 100644
--- /dev/null
+++ b/Engine/Plugins/Animation/DeformerGraph/Source/OptimusCore/Public/IOptimusUnnamedNodePinProvider.h
@@ -0,0 +1,28 @@
+﻿// Copyright Epic Games, Inc. All Rights Reserved.
+
+#pragma once
+
+#include "UObject/Interface.h"
+
+#include "IOptimusUnnamedNodePinProvider.generated.h"
+
+class UOptimusNodePin;
+
+UINTERFACE()
+class OPTIMUSCORE_API UOptimusUnnamedNodePinProvider :
+	public UInterface
+{
+	GENERATED_BODY()
+};
+
+/**
+* Interface that provides a mechanism to add pins to node from existing pins
+*/
+class OPTIMUSCORE_API IOptimusUnnamedNodePinProvider
+{
+	GENERATED_BODY()
+
+public:
+	virtual bool IsPinNameHidden(UOptimusNodePin* InPin) const = 0;
+	virtual FName GetNameForAdderPin(UOptimusNodePin* InPin) const = 0;
+};
diff --git a/Engine/Plugins/Animation/DeformerGraph/Source/OptimusCore/Public/OptimusBindingTypes.h b/Engine/Plugins/Animation/DeformerGraph/Source/OptimusCore/Public/OptimusBindingTypes.h
--- a/Engine/Plugins/Animation/DeformerGraph/Source/OptimusCore/Public/OptimusBindingTypes.h
+++ b/Engine/Plugins/Animation/DeformerGraph/Source/OptimusCore/Public/OptimusBindingTypes.h
@@ -76,6 +76,18 @@ struct FOptimusParameterBindingArray
 		return InnerArray.FindByPredicate(Pred);
 	}
 
+	template <typename Predicate>
+	int32 IndexOfByPredicate(Predicate Pred) const
+	{
+		return InnerArray.IndexOfByPredicate(Pred);
+	}
+
+	template <typename Predicate>
+	int32 IndexOfByPredicate(Predicate Pred)
+	{
+		return InnerArray.IndexOfByPredicate(Pred);
+	}
+
 	bool IsEmpty() const
 	{
 		return InnerArray.IsEmpty();
diff --git a/Engine/Plugins/Animation/DeformerGraph/Source/OptimusCore/Public/OptimusDataDomain.h b/Engine/Plugins/Animation/DeformerGraph/Source/OptimusCore/Public/OptimusDataDomain.h
--- a/Engine/Plugins/Animation/DeformerGraph/Source/OptimusCore/Public/OptimusDataDomain.h
+++ b/Engine/Plugins/Animation/DeformerGraph/Source/OptimusCore/Public/OptimusDataDomain.h
@@ -78,7 +78,6 @@ struct OPTIMUSCORE_API FOptimusDataDomain
 	UPROPERTY(EditAnywhere, Category = DataDomain, meta=(ClampMin=1, UIMax=8, SupportDynamicSliderMaxValue="true", EditCondition="DomainType==EOptimusDataDomainType::Dimensional"))
 	int32 Multiplier = 1;
 
-	//  
 	UPROPERTY(EditAnywhere, Category = DataDomain, meta=(EditCondition="DomainType==EOptimusDataDomainType::Expression"))
 	FString Expression;
 
@@ -141,6 +140,8 @@ struct OPTIMUSCORE_API FOptimusDataDomain
 
 	TOptional<FString> AsExpression() const;
 
+	TSet<FName> GetUsedConstants() const;
+
 	FString GetDisplayName() const;
 
 	static bool AreCompatible(const FOptimusDataDomain& InOutput, const FOptimusDataDomain& InInput, FString* OutReason);
diff --git a/Engine/Plugins/Animation/DeformerGraph/Source/OptimusCore/Public/OptimusDeformer.h b/Engine/Plugins/Animation/DeformerGraph/Source/OptimusCore/Public/OptimusDeformer.h
--- a/Engine/Plugins/Animation/DeformerGraph/Source/OptimusCore/Public/OptimusDeformer.h
+++ b/Engine/Plugins/Animation/DeformerGraph/Source/OptimusCore/Public/OptimusDeformer.h
@@ -10,7 +10,6 @@
 #include "OptimusNodeGraph.h"
 #include "OptimusComponentSource.h"
 #include "OptimusDataDomain.h"
-#include "OptimusConstant.h"
 
 #include "Animation/MeshDeformer.h"
 #include "Interfaces/Interface_PreviewMeshProvider.h"
@@ -466,10 +465,6 @@ protected:
 	UPROPERTY()
 	TArray<FOptimusComputeGraphInfo> ComputeGraphs;
 
-	// Constants used by resource pins in the graph
-	UPROPERTY()
-	FOptimusConstantContainer ConstantContainer;
-
 private:
 	void PostLoadFixupMissingComponentBindingsCompat();
 	void PostLoadFixupMismatchedResourceDataDomains();
diff --git a/Engine/Plugins/Animation/DeformerGraph/Source/OptimusCore/Public/OptimusDeformerInstance.h b/Engine/Plugins/Animation/DeformerGraph/Source/OptimusCore/Public/OptimusDeformerInstance.h
--- a/Engine/Plugins/Animation/DeformerGraph/Source/OptimusCore/Public/OptimusDeformerInstance.h
+++ b/Engine/Plugins/Animation/DeformerGraph/Source/OptimusCore/Public/OptimusDeformerInstance.h
@@ -4,13 +4,11 @@
 
 #include "Animation/MeshDeformerInstance.h"
 #include "ComputeFramework/ComputeGraphInstance.h"
-#include "OptimusConstant.h"
 #include "Engine/EngineTypes.h"
 
 #include "OptimusDeformerInstance.generated.h"
 
 class UOptimusComponentSource;
-struct FOptimusConstantIdentifier;
 enum class EOptimusNodeGraphType;
 struct FOptimusPersistentStructuredBuffer;
 class FRDGBuffer;
@@ -44,17 +42,6 @@ struct FOptimusDeformerInstanceComponentLodContext
 	}
 };
 
-struct FOptimusConstantEvaluationResult
-{
-	TArray<float> ValuePerInvocation;
-	FOptimusDeformerInstanceComponentLodContext LodContext;
-
-	bool IsValid() const
-	{
-		return ValuePerInvocation.Num() != 0;
-	};
-};
-
 class FOptimusPersistentBufferPool
 {
 public:
@@ -76,8 +63,8 @@ public:
 
 	void GetImplicitPersistentBuffers(
 		FRDGBuilder& GraphBuilder,
-		const FOptimusConstantIdentifier& InIdentifier,
-		const FOptimusDeformerInstanceComponentLodContext& InLodContext,
+		FName DataInterfaceName,
+		int32 LODIndex,
 		int32 InElementStride,
 		int32 InRawStride,
 		TArray<int32> const& InElementCounts,
@@ -106,7 +93,7 @@ private:
 		) const;
 	
 	TMap<FName, TMap<int32, TArray<FOptimusPersistentStructuredBuffer>>> ResourceBuffersMap;
-	TMap<FOptimusDeformerInstanceComponentLodContext, TMap<FOptimusConstantIdentifier, TArray<FOptimusPersistentStructuredBuffer>>> ImplicitBuffersMap;
+	TMap<FName, TMap<int32, TArray<FOptimusPersistentStructuredBuffer>>> ImplicitBuffersMap;
 };
 using FOptimusPersistentBufferPoolPtr = TSharedPtr<FOptimusPersistentBufferPool>;
 
@@ -264,9 +251,6 @@ public:
 	FOptimusPersistentBufferPoolPtr GetBufferPool() const { return BufferPool; }
 
 	void SetCanBeActive(bool bInCanBeActive);
-
-	FOptimusConstantEvaluationResult GetConstantValuePerInvocation(const FOptimusConstantIdentifier& InIdentifier);
-
 protected:
 	/** Implementation of UMeshDeformerInstance. */
 	void AllocateResources() override;
@@ -295,12 +279,7 @@ private:
 	
 	UPROPERTY()
 	TArray<TWeakObjectPtr<const UOptimusComponentSource>> WeakComponentSources;
-
-	UPROPERTY()
-	FOptimusConstantContainer ConstantContainer;
 	
-	TMap<FOptimusDeformerInstanceComponentLodContext, FOptimusConstantContainerInstance> ConstantValuesPerContext;
-
 	// List of graphs that should be run on the next tick. 
 	TSet<FName> GraphsToRunOnNextTick;
 	FCriticalSection GraphsToRunOnNextTickLock;
diff --git a/Engine/Plugins/Animation/DeformerGraph/Source/OptimusCore/Public/OptimusNode.h b/Engine/Plugins/Animation/DeformerGraph/Source/OptimusCore/Public/OptimusNode.h
--- a/Engine/Plugins/Animation/DeformerGraph/Source/OptimusCore/Public/OptimusNode.h
+++ b/Engine/Plugins/Animation/DeformerGraph/Source/OptimusCore/Public/OptimusNode.h
@@ -105,6 +105,9 @@ public:
 	/** Returns the list of all the pins on this node */
 	TArrayView<UOptimusNodePin* const> GetPins() const { return Pins; }
 
+	/** Returns the list of pins on this node filtered by direction */
+	TArray<UOptimusNodePin*> GetPinsByDirection(EOptimusNodePinDirection InDirection, bool bInRecursive) const;
+	
 	/**
 	 * Preliminary check for whether valid connection can be made between two existing pins.
 	 * Can be overridden by derived nodes to add their own additional checks. 
@@ -196,7 +199,17 @@ protected:
 	friend struct FOptimusNodeAction_SetPinDataDomain;
 	friend struct FOptimusNodeGraphAction_PackageKernelFunction;
 	friend struct FOptimusNodeGraphAction_UnpackageKernelFunction;
+	friend struct FOptimusNodeGraphAction_RemoveNode;
 
+	/**
+	 * Returns a unique name for pin. If given the same pins to compare to, the output
+	 * name will always be the same, unlike Optimus::GetUniqueNameForScope, which gives
+	 * a different name each time it is called, use this function if pin names are generated during
+	 * an action instead of prior to an action. Though ideally there should be only one way to generate
+	 * pin names
+	 */
+	static FName GetAvailablePinNameStable(const UObject* InNodeOrPin, FName InName);
+	
 	// Return the action stack for this node.
 	UOptimusActionStack* GetActionStack() const;
 
@@ -227,7 +240,9 @@ protected:
 	virtual void PreDuplicateRequirementActions(
 		const UOptimusNodeGraph* InTargetGraph, 
 		FOptimusCompoundAction *InCompoundAction) {}
-	
+
+	virtual void SaveState(FArchive& Ar) const;
+	virtual void RestoreState(FArchive& Ar);
 
 	void EnableDynamicPins();
 
diff --git a/Engine/Plugins/Animation/DeformerGraph/Source/OptimusCore/Public/OptimusNodeGraph.h b/Engine/Plugins/Animation/DeformerGraph/Source/OptimusCore/Public/OptimusNodeGraph.h
--- a/Engine/Plugins/Animation/DeformerGraph/Source/OptimusCore/Public/OptimusNodeGraph.h
+++ b/Engine/Plugins/Animation/DeformerGraph/Source/OptimusCore/Public/OptimusNodeGraph.h
@@ -11,6 +11,9 @@
 
 #include "OptimusNodeGraph.generated.h"
 
+class IOptimusNodePairProvider;
+class IOptimusComputeKernelProvider;
+class UOptimusNode_LoopTerminal;
 class UOptimusComponentSourceBinding;
 struct FOptimusCompoundAction;
 struct FOptimusPinTraversalContext;
@@ -21,6 +24,7 @@ class UOptimusComputeDataInterface;
 class IOptimusNodeGraphCollectionOwner;
 class UOptimusActionStack;
 class UOptimusNode;
+class UOptimusNodePair;
 class UOptimusNodeGraph;
 class UOptimusNodeLink;
 class UOptimusNodePin;
@@ -82,6 +86,8 @@ public:
 		const FString& InGraphName,
 		FText* OutFailureReason = nullptr
 		);
+
+	static FString ConstructPath(const FString& GraphPath, const FString& NodeName, const FString& PinPath);
 	
 	UFUNCTION(BlueprintCallable, Category = OptimusNodeGraph)
 	EOptimusNodeGraphType GetGraphType() const { return GraphType; }
@@ -128,6 +134,11 @@ public:
 		const FVector2D& InPosition
 	);
 
+	UFUNCTION(BlueprintCallable, Category = OptimusNodeGraph)
+	TArray<UOptimusNode*> AddLoopTerminalNodes(
+		const FVector2D& InPosition
+	);
+
 	UFUNCTION(BlueprintCallable, Category = OptimusNodeGraph)
 	UOptimusNode* AddResourceNode(
 		UOptimusResourceDescription* InResourceDesc,
@@ -263,6 +274,14 @@ public:
 	  */
 	UFUNCTION(BlueprintCallable, Category = OptimusNodeGraph)
 	bool IsSubGraphReference(UOptimusNode *InNode) const;
+
+	/** Returns the node pair given a IOptimusNodePairProvider
+	  */
+	UOptimusNodePair* GetNodePair(const UOptimusNode* InNode) const;
+	
+	/** Returns the paired node for the given IOptimusNodePairProvider
+	  */
+	UOptimusNode* GetNodeCounterpart(const UOptimusNode* InNode) const;
 	
 	/** Returns all pins that have a _direct_ connection to this pin. If nothing is connected 
 	  * to this pin, it returns an empty array.
@@ -279,7 +298,6 @@ public:
 		const FOptimusPinTraversalContext& InContext
 		) const;
 
-
 	/** Get all unique component bindings that lead to this pin. Note that only pins with a zero or single bindings
 	 *  are considered valid. We return all of them however for error messaging.
 	 */
@@ -287,11 +305,28 @@ public:
 		const UOptimusNodePin* InNodePin
 		) const;
 
-	/** Check if a pin is represents time varying data */
+	/** Check if a pin represents time varying data */
 	bool IsPinMutable(
 		const UOptimusNodePin* InNodePin
 		) const;
-		
+
+	/** Check if a Node has mutable input pins */
+	bool DoesNodeHaveMutableInput(
+		const UOptimusNode* InNode 
+		) const;
+
+	/** Gather connected loop entry terminals */
+	TSet<FOptimusRoutedConstNode> GetLoopEntryTerminalForPin(
+		const UOptimusNodePin* InNodePin,
+		const FOptimusPinTraversalContext& InContext = {}	
+		) const;
+
+	/** Gather connected loop entry terminals */
+	TSet<FOptimusRoutedConstNode> GetLoopEntryTerminalForNode(
+		const UOptimusNode* InNode,
+		const FOptimusPinTraversalContext& InContext = {}
+		) const;
+
 	TArray<const UOptimusNodeLink *> GetPinLinks(const UOptimusNodePin* InNodePin) const;
 
 	/// Check to see if connecting these two nodes will form a graph cycle.
@@ -368,6 +403,7 @@ protected:
 	friend struct FOptimusNodeGraphAction_AddNode;
 	friend struct FOptimusNodeGraphAction_DuplicateNode;
 	friend struct FOptimusNodeGraphAction_RemoveNode;
+	friend struct FOptimusNodeGraphAction_AddRemoveNodePair;
 	friend struct FOptimusNodeGraphAction_AddRemoveLink;
 	friend struct FOptimusNodeGraphAction_PackageKernelFunction;
 	friend struct FOptimusNodeGraphAction_UnpackageKernelFunction;
@@ -393,6 +429,10 @@ protected:
 		UOptimusNode* InNode,		
 		bool bFailIfLinks = true);
 
+	bool AddNodePairDirect(UOptimusNode* InFirstNode, UOptimusNode* InSecondNode);
+	
+	bool RemoveNodePairDirect(UOptimusNode* InFirstNode, UOptimusNode* InSecondNode);
+	
 	bool AddLinkDirect(UOptimusNodePin* InNodeOutputPin, UOptimusNodePin* InNodeInputPin);
 
 	bool RemoveLinkDirect(UOptimusNodePin* InNodeOutputPin, UOptimusNodePin* InNodeInputPin);
@@ -433,7 +473,18 @@ private:
 		const FVector2D& InPosition,
 		TFunction<void(UOptimusNode*)> InNodeConfigFunc
 	);
+	
+	TArray<UOptimusNode*> AddNodePairInternal(
+		const TSubclassOf<UOptimusNode> InNodeClass,
+		const FVector2D& InPosition,
+		TFunction<void(UOptimusNode*)> InFirstNodeConfigFunc,
+		TFunction<void(UOptimusNode*)> InSecondNodeConfigFunc
+	);
 
+	void AddPairNodesToArray(const TArray<UOptimusNode*>& InNodes, TArray<UOptimusNode*>& OutNodes, TArray<UOptimusNodePair*>& OutNodePairs) const;
+	
+	void RemoveNodePairByIndex(int32 NodePairIndex);
+	
 	void RemoveLinkByIndex(int32 LinkIndex);
 
 	/// Returns the indexes of all links that connect to the node. If a direction is specified
@@ -465,6 +516,9 @@ private:
 	UPROPERTY(NonTransactional)
 	TArray<TObjectPtr<UOptimusNodeLink>> Links;
 
+	UPROPERTY(NonTransactional)
+	TArray<TObjectPtr<UOptimusNodePair>> NodePairs;
+	
 	UPROPERTY()
 	TArray<TObjectPtr<UOptimusNodeGraph>> SubGraphs;
 
diff --git a/Engine/Plugins/Animation/DeformerGraph/Source/OptimusCore/Public/OptimusNodePair.h b/Engine/Plugins/Animation/DeformerGraph/Source/OptimusCore/Public/OptimusNodePair.h
new file mode 100644
--- /dev/null
+++ b/Engine/Plugins/Animation/DeformerGraph/Source/OptimusCore/Public/OptimusNodePair.h
@@ -0,0 +1,43 @@
+// Copyright Epic Games, Inc. All Rights Reserved.
+
+#pragma once
+
+#include "UObject/Object.h"
+
+#include "OptimusNodePair.generated.h"
+
+class UOptimusNode;
+
+UCLASS()
+class OPTIMUSCORE_API UOptimusNodePair : public UObject
+{
+	GENERATED_BODY()
+
+public:
+	UOptimusNodePair() = default;
+
+	// UObject overrides
+	void PostDuplicate(bool bDuplicateForPIE) override;
+	
+	/** Returns the output pin on the node this link connects from. */
+	UOptimusNode* GetFirst() const { return First; }
+
+	/** Returns the input pin on the node that this link connects to. */
+	UOptimusNode* GetSecond() const { return Second; }
+
+	bool Contains(const UOptimusNode* InFirst, const UOptimusNode* InSecond) const;
+
+	bool Contains(const UOptimusNode* InNode) const;
+
+	UOptimusNode* GetNodeCounterpart(const UOptimusNode* InNode) const;
+
+	
+protected:
+	friend class UOptimusNodeGraph;
+
+	UPROPERTY()
+	TObjectPtr<UOptimusNode> First;
+
+	UPROPERTY()
+	TObjectPtr<UOptimusNode> Second;
+};
diff --git a/Engine/Plugins/Animation/DeformerGraph/Source/OptimusEditor/Private/OptimusEditorClipboard.cpp b/Engine/Plugins/Animation/DeformerGraph/Source/OptimusEditor/Private/OptimusEditorClipboard.cpp
--- a/Engine/Plugins/Animation/DeformerGraph/Source/OptimusEditor/Private/OptimusEditorClipboard.cpp
+++ b/Engine/Plugins/Animation/DeformerGraph/Source/OptimusEditor/Private/OptimusEditorClipboard.cpp
@@ -4,6 +4,7 @@
 
 #include "OptimusNode.h"
 #include "OptimusNodeGraph.h"
+#include "OptimusNodePair.h"
 #include "OptimusNodeLink.h"
 #include "OptimusNodePin.h"
 
@@ -21,9 +22,14 @@ void FOptimusEditorClipboard::SetClipboardFromNodes(const TArray<UOptimusNode*>&
 	{
 		return;
 	}
+
+	TArray<UOptimusNode*> NodesToConsider;
+	TArray<UOptimusNodePair*> NodePairsToConsider;
 	
 	UOptimusNodeGraph* Graph = InNodes[0]->GetOwningGraph();
 
+	Graph->AddPairNodesToArray(InNodes, NodesToConsider, NodePairsToConsider);
+
 	// Export the clipboard to text.
 	const FExportObjectInnerContext Context;
 	FStringOutputDevice Archive;
@@ -33,8 +39,8 @@ void FOptimusEditorClipboard::SetClipboardFromNodes(const TArray<UOptimusNode*>&
 
 	// Export the nodes first.
 	TSet<UOptimusNode*> NodesToCopy;
-	NodesToCopy.Reserve(InNodes.Num());
-	for (UOptimusNode* Node: InNodes)
+	NodesToCopy.Reserve(NodesToConsider.Num());
+	for (UOptimusNode* Node: NodesToConsider)
 	{
 		if (!ensure(Node->GetOwningGraph() == Graph))
 		{
@@ -65,6 +71,20 @@ void FOptimusEditorClipboard::SetClipboardFromNodes(const TArray<UOptimusNode*>&
 		NodesToCopy.Add(Node);
 	}
 
+	for (UOptimusNodePair* NodePair : NodePairsToConsider)
+	{
+		UExporter::ExportToOutputDevice(
+			&Context,
+			NodePair,
+			nullptr,
+			Archive,
+			TEXT("copy"),
+			0,
+			PPF_ExportsNotFullyQualified | PPF_Copy | PPF_Delimited,
+			false,
+			Graph);	
+	}
+
 	// Export all internal links as well.
 	for (UOptimusNodeLink* Link: Graph->GetAllLinks())
 	{
@@ -321,7 +341,8 @@ bool FOptimusEditorClipboard::CanCreateClass(const UClass* InClass)
 {
 	return InClass->IsChildOf(UOptimusNode::StaticClass()) ||
 		   InClass == UOptimusNodePin::StaticClass() ||
-		   InClass == UOptimusNodeLink::StaticClass();
+		   InClass == UOptimusNodeLink::StaticClass() ||
+		   InClass == UOptimusNodePair::StaticClass();
 }
 
 
@@ -341,6 +362,11 @@ bool FOptimusEditorClipboard::ProcessPostCreateObject(UObject* InRootOuter, UObj
 			const UOptimusNodeLink* Link = Cast<UOptimusNodeLink>(InNewObject);
 
 			return Graph->AddLinkDirect(Link->GetNodeOutputPin(), Link->GetNodeInputPin());
+		}else if (InNewObject->IsA<UOptimusNodePair>())
+		{
+			const UOptimusNodePair* NodePair = Cast<UOptimusNodePair>(InNewObject);
+
+			return Graph->AddNodePairDirect(NodePair->GetFirst(), NodePair->GetSecond());
 		}
 	}
 
diff --git a/Engine/Plugins/Animation/DeformerGraph/Source/OptimusEditor/Private/OptimusEditorGraphNode.cpp b/Engine/Plugins/Animation/DeformerGraph/Source/OptimusEditor/Private/OptimusEditorGraphNode.cpp
--- a/Engine/Plugins/Animation/DeformerGraph/Source/OptimusEditor/Private/OptimusEditorGraphNode.cpp
+++ b/Engine/Plugins/Animation/DeformerGraph/Source/OptimusEditor/Private/OptimusEditorGraphNode.cpp
@@ -314,32 +314,35 @@ void UOptimusEditorGraphNode::GetNodeContextMenuActions(
 
 FLinearColor UOptimusEditorGraphNode::GetNodeTitleColor() const
 {
-	const FLinearColor ImmutableColor = FLinearColor::Green;
-	if (IOptimusComputeKernelProvider* KernelProvider = Cast<IOptimusComputeKernelProvider>(ModelNode))
+	if (ModelNode)
 	{
+		const FLinearColor ImmutableColor = FLinearColor::Green;
+
 		if (ModelNode->GetOwningGraph()->GetGraphType() == EOptimusNodeGraphType::Update)
 		{
-			if (!KernelProvider->HasMutableInput())
+			if (ModelNode->GetOwningGraph()->DoesNodeHaveMutableInput(ModelNode))
 			{
-				return ImmutableColor;
-			}	
-		}
-	}
+				return Super::GetNodeTitleColor();
+			}
 
-	if (IOptimusPinMutabilityDefiner* PinMutabilityDefiner = Cast<IOptimusPinMutabilityDefiner>(ModelNode))
-	{
-		for (UOptimusNodePin* Pin : ModelNode->GetPins())
-		{
-			if (Pin->GetDirection() == EOptimusNodePinDirection::Output)
+			if (IOptimusPinMutabilityDefiner* PinMutabilityDefiner = Cast<IOptimusPinMutabilityDefiner>(ModelNode))
 			{
-				if (PinMutabilityDefiner->IsOutputPinMutable(Pin))
+				for (UOptimusNodePin* Pin : ModelNode->GetPins())
 				{
-					return Super::GetNodeTitleColor();
+					if (Pin->GetDirection() == EOptimusNodePinDirection::Output)
+					{
+						const EOptimusPinMutability Mutability = PinMutabilityDefiner->GetOutputPinMutability(Pin);
+						if (Mutability == EOptimusPinMutability::Mutable)
+						{
+							return Super::GetNodeTitleColor();
+						}
+					}
 				}
 			}
+			
+			// No mutable output, no mutable input
+			return ImmutableColor;
 		}
-
-		return ImmutableColor;
 	}
 	
 	return Super::GetNodeTitleColor();
diff --git a/Engine/Plugins/Animation/DeformerGraph/Source/OptimusEditor/Private/OptimusEditorGraphSchema.cpp b/Engine/Plugins/Animation/DeformerGraph/Source/OptimusEditor/Private/OptimusEditorGraphSchema.cpp
--- a/Engine/Plugins/Animation/DeformerGraph/Source/OptimusEditor/Private/OptimusEditorGraphSchema.cpp
+++ b/Engine/Plugins/Animation/DeformerGraph/Source/OptimusEditor/Private/OptimusEditorGraphSchema.cpp
@@ -12,6 +12,7 @@
 #include "OptimusNodeGraph.h"
 #include "OptimusNodePin.h"
 #include "IOptimusNodeAdderPinProvider.h"
+#include "IOptimusUnnamedNodePinProvider.h"
 
 #include "EdGraphSchema_K2.h"
 #include "Editor.h"
@@ -161,6 +162,17 @@ void UOptimusEditorGraphSchema::GetGraphActions(
 
 		IoActionBuilder.AddAction(Action);
 	}
+
+	{
+		TSharedPtr< FOptimusSchemaAction_NewLoopTerminalNodes> Action(
+			new FOptimusSchemaAction_NewLoopTerminalNodes(
+				{},
+				FText::FromName(TEXT("Loop")),
+				/* Tooltip */{}, 0, /* Keywords */{}
+		));
+		
+		IoActionBuilder.AddAction(Action);
+	}
 	
 }
 
@@ -464,6 +476,22 @@ FLinearColor UOptimusEditorGraphSchema::GetColorFromPinType(const FEdGraphPinTyp
 	return GetDefault<UEdGraphSchema_K2>()->GetPinTypeColor(InPinType);
 }
 
+FText UOptimusEditorGraphSchema::GetPinDisplayName(const UEdGraphPin* Pin) const
+{
+	if (UOptimusNodePin* ModelPin = OptimusEditor::GetModelPinFromGraphPin(Pin))
+	{
+		if (const IOptimusUnnamedNodePinProvider* UnnamedNodePinProvider = Cast<const IOptimusUnnamedNodePinProvider>(ModelPin->GetOwningNode()))
+		{
+			if (UnnamedNodePinProvider->IsPinNameHidden(ModelPin))
+			{
+				return FText::GetEmpty();
+			}
+		}
+	}
+	
+	return Super::GetPinDisplayName(Pin);
+}
+
 
 void UOptimusEditorGraphSchema::TrySetDefaultValue(
 	UEdGraphPin& Pin, 
diff --git a/Engine/Plugins/Animation/DeformerGraph/Source/OptimusEditor/Private/OptimusEditorGraphSchema.h b/Engine/Plugins/Animation/DeformerGraph/Source/OptimusEditor/Private/OptimusEditorGraphSchema.h
--- a/Engine/Plugins/Animation/DeformerGraph/Source/OptimusEditor/Private/OptimusEditorGraphSchema.h
+++ b/Engine/Plugins/Animation/DeformerGraph/Source/OptimusEditor/Private/OptimusEditorGraphSchema.h
@@ -34,6 +34,7 @@ public:
 	static FLinearColor GetColorFromPinType(const FEdGraphPinType& InPinType);
 
 	// UEdGraphSchema overrides
+	FText GetPinDisplayName(const UEdGraphPin* Pin) const override;
 	bool TryCreateConnection(UEdGraphPin* PinA, UEdGraphPin* PinB) const override;
 	const FPinConnectionResponse CanCreateConnection(const UEdGraphPin* A, const UEdGraphPin* B) const override;
 	void BreakPinLinks(UEdGraphPin& TargetPin, bool bSendsNodeNotifcation) const override;
diff --git a/Engine/Plugins/Animation/DeformerGraph/Source/OptimusEditor/Private/OptimusEditorGraphSchemaActions.cpp b/Engine/Plugins/Animation/DeformerGraph/Source/OptimusEditor/Private/OptimusEditorGraphSchemaActions.cpp
--- a/Engine/Plugins/Animation/DeformerGraph/Source/OptimusEditor/Private/OptimusEditorGraphSchemaActions.cpp
+++ b/Engine/Plugins/Animation/DeformerGraph/Source/OptimusEditor/Private/OptimusEditorGraphSchemaActions.cpp
@@ -94,6 +94,29 @@ UEdGraphNode* FOptimusGraphSchemaAction_NewDataInterfaceNode::PerformAction(
 	return nullptr;
 }
 
+UEdGraphNode* FOptimusSchemaAction_NewLoopTerminalNodes::PerformAction(UEdGraph* InParentGraph, UEdGraphPin* FromPin,
+	const FVector2D Location, bool bInSelectNewNode)
+{
+	UOptimusEditorGraph* Graph = Cast<UOptimusEditorGraph>(InParentGraph);
+	
+	if (ensure(Graph != nullptr))
+	{
+		TArray<UOptimusNode*> Nodes = Graph->GetModelGraph()->AddLoopTerminalNodes(Location);
+
+		if (ensure(Nodes.Num() == 2))
+		{
+			UOptimusEditorGraphNode* GraphNode = Graph->FindGraphNodeFromModelNode(Nodes[0]);
+			if (GraphNode && bInSelectNewNode)
+			{
+				Graph->SelectNodeSet({GraphNode});
+			}
+			return GraphNode;
+		}
+	}
+
+	return nullptr;
+}
+
 
 static FText GetGraphTooltip(UOptimusNodeGraph* InGraph)
 {
diff --git a/Engine/Plugins/Animation/DeformerGraph/Source/OptimusEditor/Private/OptimusEditorGraphSchemaActions.h b/Engine/Plugins/Animation/DeformerGraph/Source/OptimusEditor/Private/OptimusEditorGraphSchemaActions.h
--- a/Engine/Plugins/Animation/DeformerGraph/Source/OptimusEditor/Private/OptimusEditorGraphSchemaActions.h
+++ b/Engine/Plugins/Animation/DeformerGraph/Source/OptimusEditor/Private/OptimusEditorGraphSchemaActions.h
@@ -84,10 +84,24 @@ struct FOptimusGraphSchemaAction_NewDataInterfaceNode :
 	FName GetTypeId() const override { return StaticGetTypeId(); }
 
 	// FEdGraphSchemaAction overrides
-	UEdGraphNode* PerformAction(class UEdGraph* ParentGraph, UEdGraphPin* FromPin, const FVector2D Location, bool bSelectNewNode = true) override;
+	UEdGraphNode* PerformAction(class UEdGraph* ParentGraph, UEdGraphPin* FromPin, const FVector2D Location, bool bInSelectNewNode = true) override;
 };
 
 
+USTRUCT()
+struct FOptimusSchemaAction_NewLoopTerminalNodes : public FEdGraphSchemaAction
+{
+	GENERATED_BODY()
+
+	// Inherit the base class's constructors
+	using FEdGraphSchemaAction::FEdGraphSchemaAction;
+	
+	static FName StaticGetTypeId() { static FName Type("FOptimusSchemaAction_LoopTerminal"); return Type; }
+	FName GetTypeId() const override { return StaticGetTypeId(); }
+
+	// FEdGraphSchemaAction overrides
+	UEdGraphNode* PerformAction(class UEdGraph* InParentGraph, UEdGraphPin* FromPin, const FVector2D Location, bool bInSelectNewNode = true) override;
+};
 
 
 /// Reference to a UOptimusNodeGraph.
diff --git a/Engine/Plugins/Animation/DeformerGraph/Source/OptimusEditor/Private/Widgets/SOptimusEditorGraphNode.cpp b/Engine/Plugins/Animation/DeformerGraph/Source/OptimusEditor/Private/Widgets/SOptimusEditorGraphNode.cpp
--- a/Engine/Plugins/Animation/DeformerGraph/Source/OptimusEditor/Private/Widgets/SOptimusEditorGraphNode.cpp
+++ b/Engine/Plugins/Animation/DeformerGraph/Source/OptimusEditor/Private/Widgets/SOptimusEditorGraphNode.cpp
@@ -1065,7 +1065,8 @@ FText SOptimusEditorGraphNode::GetPinLabel(TWeakPtr<SGraphPin> InWeakGraphPin) c
 	UOptimusEditorGraphNode* EditorGraphNode = GetEditorGraphNode();
 	TSharedPtr<SGraphPin> GraphPin = InWeakGraphPin.Pin();
 
-	if (GraphPin.IsValid() && EditorGraphNode)
+	// Absence of ModelNode indicates the node has been deleted
+	if (GraphPin.IsValid() && EditorGraphNode && EditorGraphNode->ModelNode)
 	{
 		return EditorGraphNode->GetPinDisplayName(GraphPin->GetPinObj());
 	}
