Optimus: Fix compile error on macarm64

diff --git a/Engine/Plugins/Animation/DeformerGraph/Source/OptimusCore/Private/OptimusDeformer.cpp b/Engine/Plugins/Animation/DeformerGraph/Source/OptimusCore/Private/OptimusDeformer.cpp
--- a/Engine/Plugins/Animation/DeformerGraph/Source/OptimusCore/Private/OptimusDeformer.cpp
+++ b/Engine/Plugins/Animation/DeformerGraph/Source/OptimusCore/Private/OptimusDeformer.cpp
@@ -1519,8 +1519,14 @@ TArray<UOptimusNode*> UOptimusDeformer::GetAllNodesOfClass(UClass* InNodeClass)
 struct FOptimusInstancedNode
 {
 	FOptimusRoutedConstNode	RoutedNode;
-	int32 LoopIndex;
-		
+	int32 LoopIndex = 0;
+
+	FOptimusInstancedNode() = default;
+	FOptimusInstancedNode(const FOptimusRoutedConstNode& InRoutedNode, const int32 InLoopIndex) :
+		RoutedNode(InRoutedNode), LoopIndex(InLoopIndex)
+	{
+	}
+	
 	friend uint32 GetTypeHash(const FOptimusInstancedNode& InNode)
 	{
 		return HashCombineFast(GetTypeHash(InNode.RoutedNode), InNode.LoopIndex);
@@ -1535,7 +1541,7 @@ struct FOptimusInstancedNode
 struct FOptimusInstancedPin
 {
 	FOptimusInstancedNode InstancedNode;
-	const UOptimusNodePin* Pin;
+	const UOptimusNodePin* Pin = nullptr;
 
 	friend uint32 GetTypeHash(const FOptimusInstancedPin& InPin)
 	{
