DeformerGraph: fixed trivial crash when compiling kernels with no input

#rnx

diff --git a/Engine/Plugins/Animation/DeformerGraph/Source/OptimusCore/Private/OptimusDeformer.cpp b/Engine/Plugins/Animation/DeformerGraph/Source/OptimusCore/Private/OptimusDeformer.cpp
--- a/Engine/Plugins/Animation/DeformerGraph/Source/OptimusCore/Private/OptimusDeformer.cpp
+++ b/Engine/Plugins/Animation/DeformerGraph/Source/OptimusCore/Private/OptimusDeformer.cpp
@@ -1962,6 +1962,9 @@ TArray<FOptimusComputeGraphInfo> UOptimusDeformer::CompileNodeGraphToComputeGrap
 
 			KernelDataInterfaceMap.Add(InstancedNode, KernelDataInterface);
 			DataInterfaceToBindingIndexMap.Add(KernelDataInterface) = PrimaryBindingIndex;	
+
+			KernelInputMap.Add(InstancedNode);
+			KernelOutputMap.Add(InstancedNode);
 			
 			for (const UOptimusNodePin* Pin: ConnectedNode.Node->GetPinsByDirection(EOptimusNodePinDirection::Input, true))
 			{
@@ -2028,19 +2031,19 @@ TArray<FOptimusComputeGraphInfo> UOptimusDeformer::CompileNodeGraphToComputeGrap
 				{
 					if (const IOptimusValueProvider* ValueProvider = Cast<const IOptimusValueProvider>(SourcePin->GetOwningNode()))
 					{
-						KernelInputMap.FindOrAdd(InstancedNode).Add(Pin) = {GraphDataInterface, SourcePin};
+						KernelInputMap[InstancedNode].Add(Pin) = {GraphDataInterface, SourcePin};
 					}
 					else if (UOptimusComputeDataInterface** NodeDataInterface = NodeDataInterfaceMap.Find(SourcePin->GetOwningNode()))
 					{
-						KernelInputMap.FindOrAdd(InstancedNode).Add(Pin) = {*NodeDataInterface, SourcePin};
+						KernelInputMap[InstancedNode].Add(Pin) = {*NodeDataInterface, SourcePin};
 					}
 					else if (UOptimusComputeDataInterface** KernelOutputDataInterface = KernelOutputDataInterfaceMap.Find(SourceInstancedPin))
 					{
-						KernelInputMap.FindOrAdd(InstancedNode).Add(Pin) = {*KernelOutputDataInterface, SourcePin};
+						KernelInputMap[InstancedNode].Add(Pin) = {*KernelOutputDataInterface, SourcePin};
 					}
 					else if (Cast<const UOptimusNode_LoopTerminal>(SourcePin->GetOwningNode()) && Pin->GetDataDomain().IsSingleton())
 					{
-						KernelInputMap.FindOrAdd(InstancedNode).Add(Pin) = {LoopEntryToLoopDataInterfaces[SourceRoutedNode][InstancedNode.LoopIndex] , SourcePin};
+						KernelInputMap[InstancedNode].Add(Pin) = {LoopEntryToLoopDataInterfaces[SourceRoutedNode][InstancedNode.LoopIndex] , SourcePin};
 					}
 				}
 			}
@@ -2115,7 +2118,7 @@ TArray<FOptimusComputeGraphInfo> UOptimusDeformer::CompileNodeGraphToComputeGrap
 								
 							KernelOutputDataInterfaceMap.Add(InstancedPin) = TransientBufferDI;
 
-							KernelOutputMap.FindOrAdd(InstancedNode).FindOrAdd(Pin).Add({TransientBufferDI, nullptr});	
+							KernelOutputMap[InstancedNode].FindOrAdd(Pin).Add({TransientBufferDI, nullptr});	
 						}
 					}
 					else
@@ -2131,7 +2134,7 @@ TArray<FOptimusComputeGraphInfo> UOptimusDeformer::CompileNodeGraphToComputeGrap
 							
 							if (UOptimusComputeDataInterface** NodeDataInterface = NodeDataInterfaceMap.Find(TargetNode))
 							{
-								KernelOutputMap.FindOrAdd(InstancedNode).FindOrAdd(Pin).Add({*NodeDataInterface, TargetPin});
+								KernelOutputMap[InstancedNode].FindOrAdd(Pin).Add({*NodeDataInterface, TargetPin});
 							}
 							else if (const IOptimusComputeKernelProvider* ConnectedKernel = Cast<const IOptimusComputeKernelProvider>(TargetNode))
 							{
@@ -2174,7 +2177,7 @@ TArray<FOptimusComputeGraphInfo> UOptimusDeformer::CompileNodeGraphToComputeGrap
 							// All connected kernels share the same raw buffer data interface
 							for (const UOptimusNodePin* TargetKernelPin : TargetKernelPins)
 							{
-								KernelOutputMap.FindOrAdd(InstancedNode).FindOrAdd(Pin).Add({RawBufferDI, TargetKernelPin});		
+								KernelOutputMap[InstancedNode].FindOrAdd(Pin).Add({RawBufferDI, TargetKernelPin});		
 							}
 						}
 					}
@@ -2235,8 +2238,8 @@ TArray<FOptimusComputeGraphInfo> UOptimusDeformer::CompileNodeGraphToComputeGrap
 					continue;
 				}
 				
-				const FOptimus_KernelInputMap KernelInputs = KernelInputMap[InstancedNode];
-				const FOptimus_KernelOutputMap KernelOutputs = KernelOutputMap[InstancedNode];
+				const FOptimus_KernelInputMap& KernelInputs = KernelInputMap[InstancedNode];
+				const FOptimus_KernelOutputMap& KernelOutputs = KernelOutputMap[InstancedNode];
 
 				for (const TPair<const UOptimusNodePin*, FOptimus_KernelConnection>& Item : KernelInputs)
 				{
