Optimus: fixed a crash when a kernel outputs to a kernel that is not connected to any terminal node

#rnx

diff --git a/Engine/Plugins/Animation/DeformerGraph/Source/OptimusCore/Private/OptimusDeformer.cpp b/Engine/Plugins/Animation/DeformerGraph/Source/OptimusCore/Private/OptimusDeformer.cpp
--- a/Engine/Plugins/Animation/DeformerGraph/Source/OptimusCore/Private/OptimusDeformer.cpp
+++ b/Engine/Plugins/Animation/DeformerGraph/Source/OptimusCore/Private/OptimusDeformer.cpp
@@ -15,7 +15,6 @@
 #include "IOptimusComputeKernelProvider.h"
 #include "IOptimusDataInterfaceProvider.h"
 #include "IOptimusValueProvider.h"
-#include "Misc/UObjectToken.h"
 #include "OptimusActionStack.h"
 #include "OptimusComputeGraph.h"
 #include "OptimusDataTypeRegistry.h"
@@ -48,7 +47,6 @@
 #include "Nodes/OptimusNode_LoopTerminal.h"
 
 #include "IOptimusDeprecatedExecutionDataInterface.h"
-#include "OptimusNodeLink.h"
 #include "Nodes/OptimusNode_CustomComputeKernel.h"
 
 
@@ -1671,7 +1669,7 @@ TArray<FOptimusComputeGraphInfo> UOptimusDeformer::CompileNodeGraphToComputeGrap
 
 	for (const FOptimusRoutedConstNode& ConnectedNode : ConnectedNodes )
 	{
-		if (const IOptimusComputeKernelProvider *KernelProvider = Cast<const IOptimusComputeKernelProvider>(ConnectedNode.Node))
+		if (Cast<const IOptimusComputeKernelProvider>(ConnectedNode.Node))
 		{
 			TSet<FOptimusRoutedConstNode> LoopTerminals = ConnectedNode.Node->GetOwningGraph()->GetLoopEntryTerminalForNode(ConnectedNode.Node, ConnectedNode.TraversalContext);
 			if (LoopTerminals.Num() > 0)
@@ -1694,7 +1692,7 @@ TArray<FOptimusComputeGraphInfo> UOptimusDeformer::CompileNodeGraphToComputeGrap
 	{
 		const FOptimusRoutedConstNode& ConnectedNode = ConnectedNodes[Index];
 
-		if (const IOptimusComputeKernelProvider *KernelProvider = Cast<const IOptimusComputeKernelProvider>(ConnectedNode.Node))
+		if (Cast<const IOptimusComputeKernelProvider>(ConnectedNode.Node))
 		{
 			bool bShouldSkip = true;
 			for (const FOptimusRoutedConstNode& OutputNode : NodeToOutputNodes[ConnectedNode])
@@ -1713,78 +1711,83 @@ TArray<FOptimusComputeGraphInfo> UOptimusDeformer::CompileNodeGraphToComputeGrap
 		}
 	}
 
-	// Used to establish kernel - kernel , kernel - data interface dependencies so that we know what data interface to assign to the output pin at loop boundaries
-	TMap<FOptimusRoutedConstNodePin, TArray<FOptimusRoutedConstNodePin>> TargetResourcePinConnections;
-	TMap<FOptimusRoutedConstNodePin, FOptimusRoutedConstNodePin> SourceResourcePinConnections;
+	TMap<FOptimusRoutedConstNodePin, FOptimusRoutedConstNodePin> LoopTerminalInputPinToSource;
 	
 	for (FOptimusRoutedConstNode ConnectedNode : ConnectedNodes)
 	{
-		if (Cast<const IOptimusDataInterfaceProvider>(ConnectedNode.Node) || Cast<const IOptimusComputeKernelProvider>(ConnectedNode.Node))
+		if (const UOptimusNode_LoopTerminal* LoopTerminal = Cast<const UOptimusNode_LoopTerminal>(ConnectedNode.Node))
 		{
-			if (!KernelToSkip.Contains(ConnectedNode))
+			if (LoopTerminal->GetTerminalType() == EOptimusTerminalType::Return && LoopTerminalToSkip.Contains(ConnectedNode))
+			{
+				continue;
+			}
+			
+			for (UOptimusNodePin* InputPin : LoopTerminal->GetPinsByDirection(EOptimusNodePinDirection::Input, true))
 			{
-				for (const UOptimusNodePin* Pin: ConnectedNode.Node->GetPinsByDirection(EOptimusNodePinDirection::Input, true))
+				FOptimusRoutedConstNodePin InputRoutedPin = {InputPin, ConnectedNode.TraversalContext};
+				TOptional<FOptimusRoutedConstNodePin> SourcePin;
+				
+				FOptimusRoutedConstNodePin WorkPin = InputRoutedPin;	
+				TQueue<FOptimusRoutedConstNodePin> PinQueue;
+				TSet<FOptimusRoutedConstNodePin> VisitedPin;
+				PinQueue.Enqueue(WorkPin);
+				while(PinQueue.Dequeue(WorkPin))
 				{
-					if (!Pin->GetDataDomain().IsSingleton())
+					if (!ensure(!VisitedPin.Contains(WorkPin)))
 					{
-						TArray<FOptimusRoutedConstNodePin> SourceResourcePins;
+						// Should not hit a cycle
+						continue;
+					}
+					VisitedPin.Add(WorkPin);
 
-						FOptimusRoutedConstNodePin WorkPin;
-						TQueue<FOptimusRoutedConstNodePin> PinQueue;
-						TSet<FOptimusRoutedConstNodePin> VisitedPin;
-						PinQueue.Enqueue({Pin, ConnectedNode.TraversalContext});
-						while(PinQueue.Dequeue(WorkPin))
-						{
-							if (VisitedPin.Contains(WorkPin))
-							{
-								continue;
-							}
-							VisitedPin.Add({Pin, ConnectedNode.TraversalContext});
-							
-							TArray<FOptimusRoutedNodePin> NextRoutedPin = WorkPin.NodePin->GetConnectedPinsWithRouting(WorkPin.TraversalContext);
+					TArray<FOptimusRoutedNodePin> NextRoutedPin = WorkPin.NodePin->GetConnectedPinsWithRouting(WorkPin.TraversalContext);
 
-							if (NextRoutedPin.Num() == 1)
+					if (NextRoutedPin.Num() == 1)
+					{
+						UOptimusNodePin* NextPin = NextRoutedPin[0].NodePin;
+						FOptimusRoutedConstNode NextRoutedNode = {NextPin->GetOwningNode(), NextRoutedPin[0].TraversalContext};
+						if (const UOptimusNode_LoopTerminal* NextLoopTerminal = Cast<UOptimusNode_LoopTerminal>(NextRoutedNode.Node))
+						{
+							// Entry hitting a return
+							if (LoopTerminal->GetTerminalType() == EOptimusTerminalType::Entry)
 							{
-								UOptimusNodePin* NextPin = NextRoutedPin[0].NodePin;
-								FOptimusRoutedConstNode NextRoutedNode = {NextPin->GetOwningNode(), NextRoutedPin[0].TraversalContext};
-								if (const UOptimusNode_LoopTerminal* LoopTerminal = Cast<UOptimusNode_LoopTerminal>(NextRoutedNode.Node))
+								if (ensure(NextLoopTerminal->GetTerminalType() == EOptimusTerminalType::Return))
 								{
-									if (LoopTerminal->GetTerminalType() == EOptimusTerminalType::Return)
+									if (LoopTerminalToSkip.Contains(NextRoutedNode))
 									{
-										if (!LoopTerminalToSkip.Contains(NextRoutedNode))
-										{
-											PinQueue.Enqueue({LoopTerminal->GetPinCounterpart(NextPin, EOptimusTerminalType::Return), NextRoutedPin[0].TraversalContext});
-										}
-										else
-										{
-											PinQueue.Enqueue({LoopTerminal->GetPinCounterpart(NextPin, EOptimusTerminalType::Entry), NextRoutedPin[0].TraversalContext});
-										}
+										NextPin = NextLoopTerminal->GetPinCounterpart(NextPin, EOptimusTerminalType::Entry);
 									}
-									else if (LoopTerminal->GetTerminalType() == EOptimusTerminalType::Entry)
+									else
 									{
-										if (ensure(!LoopTerminalToSkip.Contains(NextRoutedNode)))
-										{
-											PinQueue.Enqueue( {LoopTerminal->GetPinCounterpart(NextPin, EOptimusTerminalType::Entry), NextRoutedPin[0].TraversalContext});
-										}
+										NextPin = NextLoopTerminal->GetPinCounterpart(NextPin, EOptimusTerminalType::Return);
 									}
 								}
-								else
+							}
+							// Return hitting an entry
+							else
+							{
+								if (ensure(NextLoopTerminal->GetTerminalType() == EOptimusTerminalType::Entry))
 								{
-									if (!KernelToSkip.Contains(ConnectedNode))
+									if (ensure(!LoopTerminalToSkip.Contains(NextRoutedNode)))
 									{
-										SourceResourcePins.Add({NextRoutedPin[0].NodePin, NextRoutedPin[0].TraversalContext});	
+										NextPin = NextLoopTerminal->GetPinCounterpart(NextPin, EOptimusTerminalType::Entry);		
 									}
 								}
 							}
-						}
 
-						if (SourceResourcePins.Num() == 1)
+							PinQueue.Enqueue({NextPin,NextRoutedPin[0].TraversalContext});
+						}
+						else
 						{
-							TargetResourcePinConnections.FindOrAdd(SourceResourcePins[0]).Add({Pin, ConnectedNode.TraversalContext});
-							SourceResourcePinConnections.Add({Pin, ConnectedNode.TraversalContext}) = SourceResourcePins[0];
+							SourcePin = {NextRoutedPin[0].NodePin,NextRoutedPin[0].TraversalContext};
 						}
 					}
-				}	
+				}
+
+				if (SourcePin)
+				{
+					LoopTerminalInputPinToSource.Add(InputRoutedPin, *SourcePin);
+				}
 			}
 		}
 	}
@@ -1794,7 +1797,7 @@ TArray<FOptimusComputeGraphInfo> UOptimusDeformer::CompileNodeGraphToComputeGrap
 	
 	for (FOptimusRoutedConstNode ConnectedNode : ConnectedNodes)
 	{
-		if (const IOptimusComputeKernelProvider *KernelProvider = Cast<const IOptimusComputeKernelProvider>(ConnectedNode.Node))
+		if (Cast<const IOptimusComputeKernelProvider>(ConnectedNode.Node))
 		{
 			if (InNodeGraph->GraphType != EOptimusNodeGraphType::Update)
 			{
@@ -1817,8 +1820,6 @@ TArray<FOptimusComputeGraphInfo> UOptimusDeformer::CompileNodeGraphToComputeGrap
 		}
 	}
 
-	
-
 	// Instance looped nodes
 	TArray<FOptimusInstancedNode> InstancedNodes;
 	TMap<FOptimusRoutedConstNode, int32> NodeToMaxLoopIndex;
@@ -1838,17 +1839,21 @@ TArray<FOptimusComputeGraphInfo> UOptimusDeformer::CompileNodeGraphToComputeGrap
 				{
 					FOptimusRoutedConstNode LoopEntry = {LoopTerminal->GetOtherTerminal(), ConnectedNode.TraversalContext};
 
-					for (int32 Index = 1; Index < LoopTerminal->GetLoopCount(); Index++)
+					// When the entry is disconnected from the return, finding looped kernels using LoopEntryToKernelNodes would fail
+					if (const TArray<FOptimusRoutedConstNode>* LoopedKernelNodes = LoopEntryToKernelNodes.Find(LoopEntry))
 					{
-						for (const FOptimusRoutedConstNode& KernelNode : LoopEntryToKernelNodes[LoopEntry])
+						for (int32 Index = 1; Index < LoopTerminal->GetLoopCount(); Index++)
 						{
-							InstancedNodes.Add({ KernelNode , Index });
+							for (const FOptimusRoutedConstNode& KernelNode : *LoopedKernelNodes)
+							{
+								InstancedNodes.Add({ KernelNode , Index });
+							}
+						}
+						
+						for (const FOptimusRoutedConstNode& KernelNode : *LoopedKernelNodes)
+						{
+							NodeToMaxLoopIndex.FindOrAdd(KernelNode) = LoopTerminal->GetLoopCount() - 1;
 						}
-					}
-					
-					for (const FOptimusRoutedConstNode& KernelNode : LoopEntryToKernelNodes[LoopEntry])
-					{
-						NodeToMaxLoopIndex.FindOrAdd(KernelNode) = LoopTerminal->GetLoopCount() - 1;
 					}
 				}
 			}
@@ -1858,7 +1863,149 @@ TArray<FOptimusComputeGraphInfo> UOptimusDeformer::CompileNodeGraphToComputeGrap
 		NodeToMaxLoopIndex.FindOrAdd(ConnectedNode) = 0;
 	}
 
+	// Create instanced links
+	TMap<FOptimusInstancedPin, FOptimusInstancedPin> TargetPinToSourcePin;
+	
+	for (const FOptimusInstancedNode& InstancedNode : InstancedNodes )
+	{
+		const FOptimusRoutedConstNode& ThisRoutedNode = InstancedNode.RoutedNode;
+		const UOptimusNode* ThisNode = ThisRoutedNode.Node;
 
+		if (Cast<const IOptimusComputeKernelProvider>(ThisNode) || Cast<const IOptimusDataInterfaceProvider>(ThisNode))
+		{
+			for (const UOptimusNodePin* Pin: ThisRoutedNode.Node->GetPinsByDirection(EOptimusNodePinDirection::Input, true))
+			{
+				if (Pin->IsGroupingPin())
+				{
+					continue;
+				}
+
+				FOptimusInstancedPin InstancedTargetPin = {InstancedNode, Pin};
+				TArray<FOptimusRoutedNodePin> OtherPins = Pin->GetConnectedPinsWithRouting(ThisRoutedNode.TraversalContext);
+				
+				if (OtherPins.Num() == 1)
+				{
+					const UOptimusNodePin* OtherPin = OtherPins[0].NodePin;
+					const UOptimusNode* OtherNode = OtherPin->GetOwningNode();
+
+					if (const UOptimusNode_LoopTerminal* LoopTerminal = Cast<const UOptimusNode_LoopTerminal>(OtherNode);
+						LoopTerminal && !OtherPin->GetDataDomain().IsSingleton())
+					{
+						// Looped resource pins require additional routing
+						
+						enum class EAddType
+						{
+							LastInstance,
+							PreviousInstance,
+							AllButLastInstance,
+						};
+						
+						auto AddConnections = [&LoopTerminalInputPinToSource, &NodeToMaxLoopIndex, &TargetPinToSourcePin, InstancedTargetPin, &TraversalContext = OtherPins[0].TraversalContext] (UOptimusNodePin* InLoopTerminalInputPin, EAddType InType)
+						{
+							if (const FOptimusRoutedConstNodePin* SourcePin = LoopTerminalInputPinToSource.Find({InLoopTerminalInputPin, TraversalContext}))
+							{
+								const FOptimusRoutedConstNode SourceRoutedNode = {SourcePin->NodePin->GetOwningNode(), SourcePin->TraversalContext};
+
+								if (InType != EAddType::AllButLastInstance)
+								{
+									int32 SourceLoopIndex = NodeToMaxLoopIndex[SourceRoutedNode];
+									if (InType == EAddType::LastInstance)
+									{
+										SourceLoopIndex = NodeToMaxLoopIndex[SourceRoutedNode];
+									}
+									else if (InType == EAddType::PreviousInstance)
+									{
+										SourceLoopIndex = FMath::Clamp(InstancedTargetPin.InstancedNode.LoopIndex - 1, 0, NodeToMaxLoopIndex[SourceRoutedNode]);
+									}
+									const FOptimusInstancedPin InstancedSourcePin = {FOptimusInstancedNode(SourceRoutedNode, SourceLoopIndex), SourcePin->NodePin};
+									TargetPinToSourcePin.Add(InstancedTargetPin) = InstancedSourcePin;
+								}
+								else if (InType == EAddType::AllButLastInstance) 
+								{
+									// The last Instance is excluded since something out of the loop links to it 
+									for (int32 SourceLoopIndex = 0; SourceLoopIndex < NodeToMaxLoopIndex[SourceRoutedNode]; SourceLoopIndex++)
+									{
+										const FOptimusInstancedPin InstancedSourcePin = {FOptimusInstancedNode(SourceRoutedNode, SourceLoopIndex), SourcePin->NodePin};
+										TargetPinToSourcePin.Add(InstancedTargetPin) = InstancedSourcePin;		
+									} 
+								}
+								
+							}	
+						};
+
+						FOptimusRoutedConstNode OtherRoutedNode = {OtherNode, OtherPins[0].TraversalContext};
+
+						UOptimusNodePin* EntryInputPin = LoopTerminal->GetPinCounterpart(OtherPin, EOptimusTerminalType::Entry);
+						UOptimusNodePin* ReturnInputPin = LoopTerminal->GetPinCounterpart(OtherPin, EOptimusTerminalType::Return);
+						
+						if (LoopTerminal->GetTerminalType() == EOptimusTerminalType::Entry)
+						{
+							if (!LoopTerminalToSkip.Contains(OtherRoutedNode))
+							{
+								if (Cast<const IOptimusDataInterfaceProvider>(ThisNode))
+								{
+									AddConnections(EntryInputPin, EAddType::LastInstance);
+
+									if (LoopTerminal->GetLoopCount() > 1)
+									{
+										AddConnections(ReturnInputPin, EAddType::AllButLastInstance);	
+									}
+								
+								}
+								else if (Cast<const IOptimusComputeKernelProvider>(ThisNode))
+								{
+									if (InstancedNode.LoopIndex == 0)
+									{
+										AddConnections(EntryInputPin, EAddType::LastInstance);	
+									}
+									else
+									{
+										AddConnections(ReturnInputPin, EAddType::PreviousInstance);		
+									}
+								}	
+							}
+						}
+						else
+						{
+							check(InstancedNode.LoopIndex == 0);
+							
+							UOptimusNodePin* LoopTerminalInputPin = nullptr;
+							if (LoopTerminalToSkip.Contains(OtherRoutedNode))
+							{
+								LoopTerminalInputPin = EntryInputPin;
+							}
+							else
+							{
+								LoopTerminalInputPin = ReturnInputPin;
+							}
+
+							AddConnections(LoopTerminalInputPin, EAddType::LastInstance);
+						}	
+					}
+					else
+					{
+						// Plain connections
+						// 1. Kernel <-> Kernel
+						// 2. Kernel <-> Data Interface
+						// 3. Kernel -> index/count pin on Loop Terminals 
+						const FOptimusRoutedConstNode SourceRoutedNode = {OtherNode, OtherPins[0].TraversalContext};
+						const int32 SourceLoopIndex = NodeToMaxLoopIndex[SourceRoutedNode];	
+						if (ensure(SourceLoopIndex == 0))
+						{
+							const FOptimusInstancedPin InstancedSourcePin = {FOptimusInstancedNode(SourceRoutedNode, SourceLoopIndex), OtherPin};
+							TargetPinToSourcePin.Add(InstancedTargetPin) = InstancedSourcePin;	
+						}	
+					}
+				}
+			}
+		}
+	}
+	
+	TMap<FOptimusInstancedPin, TArray<FOptimusInstancedPin>> SourcePinToTargetPins;
+	for (TPair<FOptimusInstancedPin, FOptimusInstancedPin> LinkedPins : TargetPinToSourcePin)
+	{
+		SourcePinToTargetPins.FindOrAdd(LinkedPins.Value).Add(LinkedPins.Key);
+	}
 	
 	// Create all the data interfaces: node, graph, kernel outputs, loop terminal data
 	
@@ -1941,9 +2088,10 @@ TArray<FOptimusComputeGraphInfo> UOptimusDeformer::CompileNodeGraphToComputeGrap
 
 	for (const FOptimusInstancedNode& InstancedNode : InstancedNodes )
 	{
-		const FOptimusRoutedConstNode& ConnectedNode = InstancedNode.RoutedNode;
-
-		if (const IOptimusComputeKernelProvider* KernelProvider = Cast<const IOptimusComputeKernelProvider>(ConnectedNode.Node))
+		const FOptimusRoutedConstNode& RoutedNode = InstancedNode.RoutedNode;
+		const UOptimusNode* Node = RoutedNode.Node;
+		
+		if (const IOptimusComputeKernelProvider* KernelProvider = Cast<const IOptimusComputeKernelProvider>(Node))
 		{
 			UComputeDataInterface* KernelDataInterface = KernelProvider->MakeKernelDataInterface(this);
 			TSet<UOptimusComponentSourceBinding*> KernelPrimaryBindings = KernelProvider->GetPrimaryGroupPin()->GetComponentSourceBindingsRecursively();
@@ -1952,8 +2100,8 @@ TArray<FOptimusComputeGraphInfo> UOptimusDeformer::CompileNodeGraphToComputeGrap
 			{
 				AddDiagnostic(EOptimusDiagnosticLevel::Error,
 					FText::Format(LOCTEXT("InvalidComponentBindingForKernel", "Missing or multiple component bindings found in primary group of a kernel ({0}). Compilation aborted."),
-					ConnectedNode.Node->GetDisplayName()),
-					ConnectedNode.Node);
+					Node->GetDisplayName()),
+					Node);
 				return {};
 			}
 			
@@ -1966,7 +2114,7 @@ TArray<FOptimusComputeGraphInfo> UOptimusDeformer::CompileNodeGraphToComputeGrap
 			KernelInputMap.Add(InstancedNode);
 			KernelOutputMap.Add(InstancedNode);
 			
-			for (const UOptimusNodePin* Pin: ConnectedNode.Node->GetPinsByDirection(EOptimusNodePinDirection::Input, true))
+			for (const UOptimusNodePin* Pin: Node->GetPinsByDirection(EOptimusNodePinDirection::Input, true))
 			{
 				if (Pin->IsGroupingPin())
 				{
@@ -1974,62 +2122,12 @@ TArray<FOptimusComputeGraphInfo> UOptimusDeformer::CompileNodeGraphToComputeGrap
 				}
 				
 				FOptimusInstancedPin InstancedPin = {InstancedNode, Pin};
-				FOptimusInstancedPin SourceInstancedPin;
-
-				if (!Pin->GetDataDomain().IsSingleton())
-				{
-					for (const FOptimusRoutedNodePin& ConnectedPin: Pin->GetConnectedPinsWithRouting(ConnectedNode.TraversalContext))
-					{
-						// By default Read from the connected node in the same iteration
-						FOptimusRoutedConstNode SourceRoutedNode = {ConnectedPin.NodePin->GetOwningNode(), ConnectedPin.TraversalContext};
-						int32 SourceLoopIndex = FMath::Clamp(InstancedNode.LoopIndex, 0, NodeToMaxLoopIndex[SourceRoutedNode]);
-						SourceInstancedPin = {FOptimusInstancedNode(SourceRoutedNode, SourceLoopIndex), ConnectedPin.NodePin};
-
-						// If connecting to a terminal Read from last node from the previous iteration
-						if (const UOptimusNode_LoopTerminal* LoopTerminal = Cast<const UOptimusNode_LoopTerminal>(ConnectedPin.NodePin->GetOwningNode()))
-						{
-							if (InstancedNode.LoopIndex == 0)
-							{
-								// 0th instance reads from upstream as if there isn't a loop terminal
-								FOptimusRoutedConstNodePin RoutedPin = {Pin, ConnectedNode.TraversalContext};
-								if (const FOptimusRoutedConstNodePin* SourceRoutedPin = SourceResourcePinConnections.Find(RoutedPin))
-								{
-									SourceRoutedNode = {SourceRoutedPin->NodePin->GetOwningNode(), SourceRoutedPin->TraversalContext};
-									int32 MaxLoopIndex = NodeToMaxLoopIndex[SourceRoutedNode];
-									SourceInstancedPin = {FOptimusInstancedNode(SourceRoutedNode, MaxLoopIndex), SourceRoutedPin->NodePin};
-								}	
-							}
-							else
-							{
-								if (ensure(LoopTerminal->GetTerminalType() == EOptimusTerminalType::Entry))
-								{
-									UOptimusNodePin* PinCounterpart = LoopTerminal->GetPinCounterpart(ConnectedPin.NodePin, EOptimusTerminalType::Return);
-									for (const FOptimusRoutedNodePin& PinFromLastIteration: PinCounterpart->GetConnectedPinsWithRouting(ConnectedPin.TraversalContext))
-									{
-										SourceRoutedNode = {PinFromLastIteration.NodePin->GetOwningNode(), PinFromLastIteration.TraversalContext};
-										int32 SourceLoopIndexFromLastIteration = FMath::Clamp(InstancedNode.LoopIndex - 1, 0, NodeToMaxLoopIndex[SourceRoutedNode]);
-										SourceInstancedPin = {FOptimusInstancedNode(SourceRoutedNode, SourceLoopIndexFromLastIteration), PinFromLastIteration.NodePin};
-									}
-								}	
-							}
-						}
-					}
-				}
-				else
+				if (FOptimusInstancedPin* SourceInstancedPin = TargetPinToSourcePin.Find(InstancedPin))
 				{
-					for (const FOptimusRoutedNodePin& ConnectedPin: Pin->GetConnectedPinsWithRouting(ConnectedNode.TraversalContext))
-					{
-						FOptimusRoutedConstNode SourceRoutedNode = {ConnectedPin.NodePin->GetOwningNode(), ConnectedPin.TraversalContext};
-						SourceInstancedPin = {FOptimusInstancedNode(SourceRoutedNode, 0), ConnectedPin.NodePin};
-					}
-				}
-
-				const UOptimusNodePin* SourcePin = SourceInstancedPin.Pin;
-				const FOptimusRoutedConstNode& SourceRoutedNode = SourceInstancedPin.InstancedNode.RoutedNode;
-
-				if (SourcePin)
-				{
-					if (const IOptimusValueProvider* ValueProvider = Cast<const IOptimusValueProvider>(SourcePin->GetOwningNode()))
+					const UOptimusNodePin* SourcePin = SourceInstancedPin->Pin;
+					const FOptimusRoutedConstNode& SourceRoutedNode = SourceInstancedPin->InstancedNode.RoutedNode;
+					
+					if (Cast<const IOptimusValueProvider>(SourcePin->GetOwningNode()))
 					{
 						KernelInputMap[InstancedNode].Add(Pin) = {GraphDataInterface, SourcePin};
 					}
@@ -2037,7 +2135,7 @@ TArray<FOptimusComputeGraphInfo> UOptimusDeformer::CompileNodeGraphToComputeGrap
 					{
 						KernelInputMap[InstancedNode].Add(Pin) = {*NodeDataInterface, SourcePin};
 					}
-					else if (UOptimusComputeDataInterface** KernelOutputDataInterface = KernelOutputDataInterfaceMap.Find(SourceInstancedPin))
+					else if (UOptimusComputeDataInterface** KernelOutputDataInterface = KernelOutputDataInterfaceMap.Find(*SourceInstancedPin))
 					{
 						KernelInputMap[InstancedNode].Add(Pin) = {*KernelOutputDataInterface, SourcePin};
 					}
@@ -2048,85 +2146,18 @@ TArray<FOptimusComputeGraphInfo> UOptimusDeformer::CompileNodeGraphToComputeGrap
 				}
 			}
 			
-			for (const UOptimusNodePin* Pin: ConnectedNode.Node->GetPinsByDirection(EOptimusNodePinDirection::Output, true))
+			for (const UOptimusNodePin* Pin: Node->GetPinsByDirection(EOptimusNodePinDirection::Output, true))
 			{
 				if (ensure(!Pin->GetDataDomain().IsSingleton()))
 				{
 					FOptimusInstancedPin InstancedPin = {InstancedNode, Pin};
 
-					TArray<FOptimusInstancedPin> TargetInstancedPins;
-
-					// Writing to other looped kernel instances or data interface node
-					for (const FOptimusRoutedNodePin& ConnectedPin: Pin->GetConnectedPinsWithRouting(ConnectedNode.TraversalContext))
-					{
-						// By default write to the connected node in the same iteration
-						UOptimusNode* TargetNode = ConnectedPin.NodePin->GetOwningNode();
-						
-						if (const UOptimusNode_LoopTerminal* LoopTerminal = Cast<const UOptimusNode_LoopTerminal>(TargetNode))
-						{
-							// Last instance of the kernel writes as if there isn't a loop terminal
-							if (InstancedNode.LoopIndex == NodeToMaxLoopIndex[ConnectedNode])
-							{
-								FOptimusRoutedConstNodePin RoutedPin = {Pin, ConnectedNode.TraversalContext};
-								if (TArray<FOptimusRoutedConstNodePin>* TargetRoutedPins = TargetResourcePinConnections.Find(RoutedPin))
-								{
-									for (FOptimusRoutedConstNodePin& TargetRoutedPin : *TargetRoutedPins)
-									{
-										FOptimusRoutedConstNode TargetRoutedNode = {TargetRoutedPin.NodePin->GetOwningNode(), TargetRoutedPin.TraversalContext};
-										TargetInstancedPins.Add({FOptimusInstancedNode(TargetRoutedNode, 0) , TargetRoutedPin.NodePin});	
-									}
-								}	
-							}
-							else
-							{
-								// For nodes with in a loop, if connecting to a terminal, write to the first node from the next iteration
-								if (ensure(LoopTerminal->GetTerminalType() == EOptimusTerminalType::Return))
-								{
-									UOptimusNodePin* PinCounterpart = LoopTerminal->GetPinCounterpart(ConnectedPin.NodePin, EOptimusTerminalType::Entry);
-									for (const FOptimusRoutedNodePin& PinFromNextIteration: PinCounterpart->GetConnectedPinsWithRouting(ConnectedPin.TraversalContext))
-									{
-										UOptimusNode* TargetNodeFromNextIteration = PinFromNextIteration.NodePin->GetOwningNode();
-										FOptimusRoutedConstNode TargetRoutedNodeFromNextIteration = {TargetNodeFromNextIteration, PinFromNextIteration.TraversalContext};
-										int32 TargetLoopIndexFromNextIteration = FMath::Clamp(InstancedNode.LoopIndex + 1, 0, NodeToMaxLoopIndex[TargetRoutedNodeFromNextIteration]);
-										TargetInstancedPins.Add({FOptimusInstancedNode(TargetRoutedNodeFromNextIteration, TargetLoopIndexFromNextIteration), PinFromNextIteration.NodePin});
-									}
-								}	
-							}
-						}
-						else
-						{
-							FOptimusRoutedConstNode TargetRoutedNode = {TargetNode, ConnectedPin.TraversalContext};
-							int32 TargetLoopIndex = FMath::Clamp(InstancedNode.LoopIndex, 0, NodeToMaxLoopIndex[TargetRoutedNode]);
-							FOptimusInstancedPin TargetInstancedPin = {FOptimusInstancedNode(TargetRoutedNode, TargetLoopIndex), ConnectedPin.NodePin};
-							TargetInstancedPins.Add(TargetInstancedPin);	
-						}
-					}	
-
-					if (TargetInstancedPins.Num() == 0)
-					{
-						if (KernelProvider->DoesOutputPinSupportAtomic(Pin) || KernelProvider->DoesOutputPinSupportRead(Pin))
-						{
-							UOptimusTransientBufferDataInterface* TransientBufferDI = NewObject<UOptimusTransientBufferDataInterface>(this);
-							if (KernelProvider->DoesOutputPinSupportAtomic(Pin))
-							{
-								CastChecked<UOptimusTransientBufferDataInterface>(TransientBufferDI)->bZeroInitForAtomicWrites = true;
-							}
-							TransientBufferDI->ValueType = Pin->GetDataType()->ShaderValueType;
-							TransientBufferDI->DataDomain = Pin->GetDataDomain();
-							TransientBufferDI->ComponentSourceBinding = KernelPrimaryBinding;
-							DataInterfaceToBindingIndexMap.Add(TransientBufferDI) = PrimaryBindingIndex;
-								
-							KernelOutputDataInterfaceMap.Add(InstancedPin) = TransientBufferDI;
-
-							KernelOutputMap[InstancedNode].FindOrAdd(Pin).Add({TransientBufferDI, nullptr});	
-						}
-					}
-					else
+					if (TArray<FOptimusInstancedPin>* TargetInstancedPins = SourcePinToTargetPins.Find(InstancedPin))
 					{
 						bool bShouldUseImplicitPersistentDI = false;
 						TArray<const UOptimusNodePin*> TargetKernelPins;
 						
-						for (const FOptimusInstancedPin& TargetInstancedPin : TargetInstancedPins)
+						for (const FOptimusInstancedPin& TargetInstancedPin : *TargetInstancedPins)
 						{
 							const FOptimusRoutedConstNode TargetRoutedNode = TargetInstancedPin.InstancedNode.RoutedNode;
 							const UOptimusNode* TargetNode = TargetRoutedNode.Node;
@@ -2136,9 +2167,9 @@ TArray<FOptimusComputeGraphInfo> UOptimusDeformer::CompileNodeGraphToComputeGrap
 							{
 								KernelOutputMap[InstancedNode].FindOrAdd(Pin).Add({*NodeDataInterface, TargetPin});
 							}
-							else if (const IOptimusComputeKernelProvider* ConnectedKernel = Cast<const IOptimusComputeKernelProvider>(TargetNode))
+							else if (Cast<const IOptimusComputeKernelProvider>(TargetNode))
 							{
-								if (KernelToGraphType[ConnectedNode] != KernelToGraphType[TargetRoutedNode])
+								if (KernelToGraphType[RoutedNode] != KernelToGraphType[TargetRoutedNode])
 								{
 									bShouldUseImplicitPersistentDI = true;
 								}
@@ -2179,7 +2210,26 @@ TArray<FOptimusComputeGraphInfo> UOptimusDeformer::CompileNodeGraphToComputeGrap
 							{
 								KernelOutputMap[InstancedNode].FindOrAdd(Pin).Add({RawBufferDI, TargetKernelPin});		
 							}
-						}
+						}	
+					}
+					else
+					{
+						if (KernelProvider->DoesOutputPinSupportAtomic(Pin) || KernelProvider->DoesOutputPinSupportRead(Pin))
+						{
+							UOptimusTransientBufferDataInterface* TransientBufferDI = NewObject<UOptimusTransientBufferDataInterface>(this);
+							if (KernelProvider->DoesOutputPinSupportAtomic(Pin))
+							{
+								CastChecked<UOptimusTransientBufferDataInterface>(TransientBufferDI)->bZeroInitForAtomicWrites = true;
+							}
+							TransientBufferDI->ValueType = Pin->GetDataType()->ShaderValueType;
+							TransientBufferDI->DataDomain = Pin->GetDataDomain();
+							TransientBufferDI->ComponentSourceBinding = KernelPrimaryBinding;
+							DataInterfaceToBindingIndexMap.Add(TransientBufferDI) = PrimaryBindingIndex;
+								
+							KernelOutputDataInterfaceMap.Add(InstancedPin) = TransientBufferDI;
+
+							KernelOutputMap[InstancedNode].FindOrAdd(Pin).Add({TransientBufferDI, nullptr});	
+						}	
 					}
 				}
 			}
@@ -2231,7 +2281,7 @@ TArray<FOptimusComputeGraphInfo> UOptimusDeformer::CompileNodeGraphToComputeGrap
 		{
 			const FOptimusRoutedConstNode& ConnectedNode = InstancedNode.RoutedNode;
 
-			if (const IOptimusComputeKernelProvider* KernelProvider = Cast<const IOptimusComputeKernelProvider>(ConnectedNode.Node))
+			if (Cast<const IOptimusComputeKernelProvider>(ConnectedNode.Node))
 			{
 				if (KernelToGraphType[ConnectedNode] != GraphInfo.GraphType)
 				{
@@ -2326,7 +2376,6 @@ TArray<FOptimusComputeGraphInfo> UOptimusDeformer::CompileNodeGraphToComputeGrap
 				bool bHasExecution = false;
 				for (const TPair<int32, FOptimus_InterfaceBinding>& DataBinding: BoundKernel.InputDataBindings)
 				{
-					const int32 KernelBindingIndex = DataBinding.Key;
 					const FOptimus_InterfaceBinding& InterfaceBinding = DataBinding.Value;
 					const UComputeDataInterface* DataInterface = InterfaceBinding.DataInterface;
 					if (DataInterface->IsExecutionInterface())
