Optimus: minor fixes related to copy/paste component source nodes, variable nodes, duplicate deformer assets

#rnx

diff --git a/Engine/Plugins/Animation/DeformerGraph/Source/OptimusCore/Private/Nodes/OptimusNode_ComponentSource.cpp b/Engine/Plugins/Animation/DeformerGraph/Source/OptimusCore/Private/Nodes/OptimusNode_ComponentSource.cpp
--- a/Engine/Plugins/Animation/DeformerGraph/Source/OptimusCore/Private/Nodes/OptimusNode_ComponentSource.cpp
+++ b/Engine/Plugins/Animation/DeformerGraph/Source/OptimusCore/Private/Nodes/OptimusNode_ComponentSource.cpp
@@ -6,6 +6,8 @@
 
 #include "OptimusDataTypeRegistry.h"
 #include "OptimusDeformer.h"
+#include "Actions/OptimusAction.h"
+#include "Actions/OptimusComponentBindingActions.h"
 
 #include UE_INLINE_GENERATED_CPP_BY_NAME(OptimusNode_ComponentSource)
 
@@ -13,18 +15,12 @@
 void UOptimusNode_ComponentSource::PostDuplicate(EDuplicateMode::Type DuplicateMode)
 {
 	Super::PostDuplicate(DuplicateMode);
-
-	if (!Binding)
-	{
-		return;
-	}
 	
 	if (!GetOwningGraph())
 	{
 		return;
 	}
 	
-	const UOptimusDeformer* OldBindingOwner = Binding->GetOwningDeformer();
 	const UOptimusDeformer* NewBindingOwner = Cast<UOptimusDeformer>(GetOwningGraph()->GetCollectionRoot());
 
 	if (!NewBindingOwner)
@@ -32,14 +28,75 @@ void UOptimusNode_ComponentSource::PostDuplicate(EDuplicateMode::Type DuplicateM
 		return;
 	}
 	
-	// No action needed if we are copying/pasting within the same deformer asset 
-	if (OldBindingOwner == NewBindingOwner)
+	if (Binding)
 	{
-		return;
+		const UOptimusDeformer* OldBindingOwner = Binding->GetOwningDeformer();
+	
+		// No action needed if we are copying/pasting within the same deformer asset 
+		if (OldBindingOwner == NewBindingOwner)
+		{
+			return;
+		}
+
+		// Refresh the binding so that we don't hold a reference to a binding in another deformer asset
+		Binding = NewBindingOwner->ResolveComponentBinding(Binding->GetFName());
+	}	
+	else if (!DuplicationInfo.BindingName.IsNone())
+	{
+		// Resolve the variable description from the duplication information that likely came from a copy/paste operation
+		// where the pointer didn't survive and the variable didn't exist and we had to set it up in PreDuplicateRequirementActions
+		Binding = NewBindingOwner->ResolveComponentBinding(DuplicationInfo.BindingName);
+
+		if (Binding)
+		{
+			// Re-create the pins. Don't call PostCreateNode, since it won't notify the graph that the pin layout has
+			// changed.
+			check(GetPins().IsEmpty());
+			ConstructNode();
+		}
+
+		// Empty the duplication info before we carry on.
+		DuplicationInfo = FOptimusNode_ComponentSource_DuplicationInfo();
+	}
+}
+
+void UOptimusNode_ComponentSource::ExportCustomProperties(FOutputDevice& Out, uint32 Indent)
+{
+	if (ensure(Binding))
+	{
+		FString ExportText;
+
+		Out.Logf(TEXT("%sCustomProperties ComponentSourceBinding BindingName=%s ComponentType=%s\n"),
+			FCString::Spc(Indent), *Binding->BindingName.ToString(), *Binding->ComponentType->GetStructPathName().ToString());
 	}
+}
 
-	// Refresh the binding so that we don't hold a reference to a binding in another deformer asset
-	Binding = NewBindingOwner->ResolveComponentBinding(Binding->GetFName());
+
+void UOptimusNode_ComponentSource::ImportCustomProperties(const TCHAR* SourceText, FFeedbackContext* Warn)
+{
+	if (FParse::Command(&SourceText, TEXT("ComponentSourceBinding")))
+	{
+		FName BindingName;
+		if (!FParse::Value(SourceText, TEXT("BindingName="), BindingName))
+		{
+			return;
+		}
+
+		FString ComponentTypeClassPath;
+		if (!FParse::Value(SourceText, TEXT("ComponentType="), ComponentTypeClassPath))
+		{
+			return;
+		}
+
+		UClass* ComponentSourceClass = Cast<UClass>(StaticFindObject(UClass::StaticClass(), FTopLevelAssetPath(ComponentTypeClassPath), false));
+		if (!ComponentSourceClass)
+		{
+			return;
+		}
+
+		DuplicationInfo.BindingName=BindingName;
+		DuplicationInfo.ComponentType = ComponentSourceClass;
+	}
 }
 
 void UOptimusNode_ComponentSource::SetComponentSourceBinding(
@@ -64,14 +121,54 @@ void UOptimusNode_ComponentSource::ConstructNode()
 	const FOptimusDataTypeRegistry& TypeRegistry = FOptimusDataTypeRegistry::Get();
 	FOptimusDataTypeRef ComponentSourceType = TypeRegistry.FindType(*UOptimusComponentSourceBinding::StaticClass());
 
+	// Binding can be empty during copy paste. It is assigned a value in post duplicate and ConstructNode is called again.
 	if (ensure(ComponentSourceType.IsValid()) &&
-		ensure(Binding) &&
+		Binding &&
 		ensure(Binding->ComponentType))
 	{
 		AddPinDirect(Binding->GetComponentSource()->GetBindingName(), EOptimusNodePinDirection::Output, {}, ComponentSourceType);
 	}
 }
 
+void UOptimusNode_ComponentSource::PreDuplicateRequirementActions(const UOptimusNodeGraph* InTargetGraph, FOptimusCompoundAction* InCompoundAction)
+{
+	if (DuplicationInfo.BindingName.IsNone())
+	{
+		return;
+	}
+
+	// Check if the deformer has a variable that matches the name and type.
+	UOptimusDeformer* Deformer = Cast<UOptimusDeformer>(InTargetGraph->GetCollectionRoot());
+	
+	for (const UOptimusComponentSourceBinding* ExistingBinding: Deformer->GetComponentBindings())
+	{
+		if (ExistingBinding->BindingName == DuplicationInfo.BindingName &&
+			ExistingBinding->ComponentType == DuplicationInfo.ComponentType)
+		{
+			// Note: We don't care about component tags in this context.
+			return;
+		}
+	}
+
+	// Note: This will fail if there are multiple overlapping names. A better system is needed.
+	bool bFoundDuplicate = false;
+	do
+	{
+		bFoundDuplicate = false;
+		for (const UOptimusComponentSourceBinding* ExistingBinding: Deformer->GetComponentBindings())
+		{
+			if (ExistingBinding->BindingName == DuplicationInfo.BindingName)
+			{
+				DuplicationInfo.BindingName.SetNumber(DuplicationInfo.BindingName.GetNumber() + 1);
+				bFoundDuplicate = true;
+			}
+		}
+	}
+	while(bFoundDuplicate);
+
+	InCompoundAction->AddSubAction<FOptimusComponentBindingAction_AddBinding>(DuplicationInfo.ComponentType->GetDefaultObject<UOptimusComponentSource>(), DuplicationInfo.BindingName);
+}
+
 UOptimusNodePin* UOptimusNode_ComponentSource::GetComponentPin() const
 {
 	if (ensure(!GetPins().IsEmpty()))
diff --git a/Engine/Plugins/Animation/DeformerGraph/Source/OptimusCore/Private/Nodes/OptimusNode_ComponentSource.h b/Engine/Plugins/Animation/DeformerGraph/Source/OptimusCore/Private/Nodes/OptimusNode_ComponentSource.h
--- a/Engine/Plugins/Animation/DeformerGraph/Source/OptimusCore/Private/Nodes/OptimusNode_ComponentSource.h
+++ b/Engine/Plugins/Animation/DeformerGraph/Source/OptimusCore/Private/Nodes/OptimusNode_ComponentSource.h
@@ -9,6 +9,17 @@
 
 #include "OptimusNode_ComponentSource.generated.h"
 
+USTRUCT()
+struct FOptimusNode_ComponentSource_DuplicationInfo
+{
+	GENERATED_BODY()
+	
+	UPROPERTY()
+	FName BindingName;
+
+	UPROPERTY()
+	TSubclassOf<UOptimusComponentSource> ComponentType;
+};
 
 UCLASS(Hidden)
 class UOptimusNode_ComponentSource :
@@ -35,6 +46,7 @@ protected:
 
 	// UOptimusNode overrides
 	void ConstructNode() override;
+	void PreDuplicateRequirementActions(const UOptimusNodeGraph* InTargetGraph, FOptimusCompoundAction* InCompoundAction) override;
 
 	/** Accessor to allow the UOptimusDeformer class to hook up data interface nodes to binding nodes automatically
 	 *  for backcomp.
@@ -43,8 +55,14 @@ protected:
 
 	// UObject overrides
 	virtual void PostDuplicate(EDuplicateMode::Type DuplicateMode) override;
+	void ExportCustomProperties(FOutputDevice& Out, uint32 Indent) override;
+	void ImportCustomProperties(const TCHAR* SourceText, FFeedbackContext* Warn) override;
 	
 private:
 	UPROPERTY()
 	TObjectPtr<UOptimusComponentSourceBinding> Binding;
+
+	// Duplication data across graphs/assets
+	UPROPERTY(DuplicateTransient)
+	FOptimusNode_ComponentSource_DuplicationInfo DuplicationInfo;	
 };
diff --git a/Engine/Plugins/Animation/DeformerGraph/Source/OptimusCore/Private/Nodes/OptimusNode_GetVariable.cpp b/Engine/Plugins/Animation/DeformerGraph/Source/OptimusCore/Private/Nodes/OptimusNode_GetVariable.cpp
--- a/Engine/Plugins/Animation/DeformerGraph/Source/OptimusCore/Private/Nodes/OptimusNode_GetVariable.cpp
+++ b/Engine/Plugins/Animation/DeformerGraph/Source/OptimusCore/Private/Nodes/OptimusNode_GetVariable.cpp
@@ -43,6 +43,7 @@ void UOptimusNode_GetVariable::PreDuplicateRequirementActions(
 	bool bFoundDuplicate = false;
 	do
 	{
+		bFoundDuplicate = false;
 		for (const UOptimusVariableDescription* ExistingVariableDesc: Deformer->GetVariables())
 		{
 			if (ExistingVariableDesc->VariableName == DuplicationInfo.VariableName)
diff --git a/Engine/Plugins/Animation/DeformerGraph/Source/OptimusCore/Private/OptimusNodeGraph.cpp b/Engine/Plugins/Animation/DeformerGraph/Source/OptimusCore/Private/OptimusNodeGraph.cpp
--- a/Engine/Plugins/Animation/DeformerGraph/Source/OptimusCore/Private/OptimusNodeGraph.cpp
+++ b/Engine/Plugins/Animation/DeformerGraph/Source/OptimusCore/Private/OptimusNodeGraph.cpp
@@ -105,7 +105,7 @@ void UOptimusNodeGraph::PostDuplicate(EDuplicateMode::Type DuplicateMode)
 				ConnectedPinsMap.Add(Pin->GetFName(), Pin->GetConnectedPins());
 			}
 
-			FName NodeName = ConstantNode->GetNodeName();
+			FName NodeName = ConstantNode->GetFName();
 			
 			RemoveNodeDirect(ConstantNode, false /* remove all links as well */);
 
