Optimus: fixed crash on DX11 when using debugdraw data interface. Added tooltip on debugdraw datainterface to guide users to change project settings

#rnx

diff --git a/Engine/Plugins/Animation/DeformerGraph/Source/OptimusCore/Private/DataInterfaces/OptimusDataInterfaceDebugDraw.cpp b/Engine/Plugins/Animation/DeformerGraph/Source/OptimusCore/Private/DataInterfaces/OptimusDataInterfaceDebugDraw.cpp
--- a/Engine/Plugins/Animation/DeformerGraph/Source/OptimusCore/Private/DataInterfaces/OptimusDataInterfaceDebugDraw.cpp
+++ b/Engine/Plugins/Animation/DeformerGraph/Source/OptimusCore/Private/DataInterfaces/OptimusDataInterfaceDebugDraw.cpp
@@ -14,6 +14,14 @@
 #include UE_INLINE_GENERATED_CPP_BY_NAME(OptimusDataInterfaceDebugDraw)
 
 
+UOptimusDebugDrawDataInterface::UOptimusDebugDrawDataInterface()
+{
+	if (ShaderPrint::IsSupported(GMaxRHIShaderPlatform))
+	{
+		bIsSupported = true;
+	}
+}
+
 FString UOptimusDebugDrawDataInterface::GetDisplayName() const
 {
 	return TEXT("Debug Draw");
@@ -151,6 +159,11 @@ bool FOptimusDebugDrawDataProviderProxy::IsValid(FValidationData const& InValida
 		return false;
 	}
 
+	if (!ShaderPrint::IsSupported(GMaxRHIShaderPlatform))
+	{
+		return false;
+	}
+	
 	return true;
 }
 
diff --git a/Engine/Plugins/Animation/DeformerGraph/Source/OptimusCore/Private/DataInterfaces/OptimusDataInterfaceDebugDraw.h b/Engine/Plugins/Animation/DeformerGraph/Source/OptimusCore/Private/DataInterfaces/OptimusDataInterfaceDebugDraw.h
--- a/Engine/Plugins/Animation/DeformerGraph/Source/OptimusCore/Private/DataInterfaces/OptimusDataInterfaceDebugDraw.h
+++ b/Engine/Plugins/Animation/DeformerGraph/Source/OptimusCore/Private/DataInterfaces/OptimusDataInterfaceDebugDraw.h
@@ -44,6 +44,8 @@ class OPTIMUSCORE_API UOptimusDebugDrawDataInterface : public UOptimusComputeDat
 	GENERATED_BODY()
 
 public:
+	UOptimusDebugDrawDataInterface();
+	
 	//~ Begin UOptimusComputeDataInterface Interface
 	FString GetDisplayName() const override;
 	FName GetCategory() const override;
@@ -63,6 +65,10 @@ public:
 	UComputeDataProvider* CreateDataProvider(TObjectPtr<UObject> InBinding, uint64 InInputMask, uint64 InOutputMask) const override;
 	//~ End UComputeDataInterface Interface
 
+	// Make sure DirectX12 and Shader Model 6 is enabled in project settings for DebugDraw to function, since DXC is required for shader compilation. 
+	UPROPERTY(Transient, VisibleAnywhere, Category = DebugDraw, meta=(EditConditionHides, EditCondition="!bIsSupported"))
+	bool bIsSupported = false;
+	
 	UPROPERTY(EditAnywhere, Category = DebugDraw, meta = (ShowOnlyInnerProperties))
 	FOptimusDebugDrawParameters DebugDrawParameters;
 
