Optimus: Insert Copy Kernels when necessary:

1. when two data interfaces are directly linked
2. when a run-once kernel writes to a data interface
3. when a readable raw buffer goes to a data interface

diff --git a/Engine/Plugins/Animation/DeformerGraph/Source/OptimusCore/Private/DataInterfaces/OptimusDataInterfaceAnimAttribute.cpp b/Engine/Plugins/Animation/DeformerGraph/Source/OptimusCore/Private/DataInterfaces/OptimusDataInterfaceAnimAttribute.cpp
--- a/Engine/Plugins/Animation/DeformerGraph/Source/OptimusCore/Private/DataInterfaces/OptimusDataInterfaceAnimAttribute.cpp
+++ b/Engine/Plugins/Animation/DeformerGraph/Source/OptimusCore/Private/DataInterfaces/OptimusDataInterfaceAnimAttribute.cpp
@@ -374,12 +374,8 @@ void UOptimusAnimAttributeDataInterface::GetStructDeclarations(TSet<FString>& Ou
 
 void UOptimusAnimAttributeDataInterface::GetShaderHash(FString& InOutKey) const
 {
-// 	FSHA1 HashState;
-// 	FString HLSL;
-// 	GetHLSL(HLSL);
-// 	
-// 	HashState.UpdateWithString(*HLSL, HLSL.Len());
-// 	HashState.Finalize().AppendString(InOutKey);
+	// UComputeGraph::BuildKernelSource hashes the result of GetHLSL()
+	// Only append additional hashes here if the HLSL contains any additional includes
 }
 
 UComputeDataProvider* UOptimusAnimAttributeDataInterface::CreateDataProvider(TObjectPtr<UObject> InBinding, uint64 InInputMask, uint64 InOutputMask) const
diff --git a/Engine/Plugins/Animation/DeformerGraph/Source/OptimusCore/Private/DataInterfaces/OptimusDataInterfaceCopyKernel.cpp b/Engine/Plugins/Animation/DeformerGraph/Source/OptimusCore/Private/DataInterfaces/OptimusDataInterfaceCopyKernel.cpp
new file mode 100644
--- /dev/null
+++ b/Engine/Plugins/Animation/DeformerGraph/Source/OptimusCore/Private/DataInterfaces/OptimusDataInterfaceCopyKernel.cpp
@@ -0,0 +1,314 @@
+// Copyright Epic Games, Inc. All Rights Reserved.
+
+#include "OptimusDataInterfaceCopyKernel.h"
+
+#include "OptimusComponentSource.h"
+#include "OptimusExpressionEvaluator.h"
+#include "ShaderParameterMetadataBuilder.h"
+#include "ComputeFramework/ComputeMetadataBuilder.h"
+#include "ComputeFramework/ShaderParameterMetadataAllocation.h"
+#include "ComputeFramework/ShaderParamTypeDefinition.h"
+
+#include UE_INLINE_GENERATED_CPP_BY_NAME(OptimusDataInterfaceCopyKernel)
+
+const FString UOptimusCopyKernelDataInterface::NumThreadsReservedName = TEXT("NumThreads");
+
+void UOptimusCopyKernelDataInterface::GetSupportedInputs(TArray<FShaderFunctionDefinition>& OutFunctions) const
+{
+	OutFunctions.AddDefaulted_GetRef()
+		.SetName(FString::Printf(TEXT("Read%s"), *NumThreadsReservedName))
+		.AddReturnType(FShaderValueType::Get(EShaderFundamentalType::Int, 3));
+}
+
+void UOptimusCopyKernelDataInterface::GetShaderParameters(TCHAR const* UID,
+	FShaderParametersMetadataBuilder& InOutBuilder, FShaderParametersMetadataAllocations& InOutAllocations) const
+{
+	FShaderParametersMetadataBuilder Builder;
+	TArray<FShaderParametersMetadata*> DummyNestedStructs;
+	ComputeFramework::AddParamForType(Builder, *NumThreadsReservedName, FShaderValueType::Get(EShaderFundamentalType::Int, 3),DummyNestedStructs);
+
+	FShaderParametersMetadata* ShaderParameterMetadata = Builder.Build(FShaderParametersMetadata::EUseCase::ShaderParameterStruct, TEXT("UCopyKernelDataInterface"));
+
+	InOutAllocations.ShaderParameterMetadatas.Add(ShaderParameterMetadata);
+	InOutAllocations.ShaderParameterMetadatas.Append(DummyNestedStructs);
+
+	// Add the generated nested struct to our builder.
+	InOutBuilder.AddNestedStruct(UID, ShaderParameterMetadata);
+
+}
+
+void UOptimusCopyKernelDataInterface::GetShaderHash(FString& InOutKey) const
+{
+	// UComputeGraph::BuildKernelSource hashes the result of GetHLSL()
+	// Only append additional hashes here if the HLSL contains any additional includes	
+}
+
+void UOptimusCopyKernelDataInterface::GetHLSL(FString& OutHLSL, FString const& InDataInterfaceName) const
+{
+	const FString TypeName = FShaderValueType::Get(EShaderFundamentalType::Int, 3)->ToString();
+
+	if (ensure(!TypeName.IsEmpty()))
+	{
+		// Add uniforms.
+		OutHLSL += FString::Printf(TEXT("%s %s_%s;\n"), 
+			*TypeName, 
+			*InDataInterfaceName, 
+			*NumThreadsReservedName);
+			
+		// Add function getters.
+		OutHLSL += FString::Printf(TEXT("%s Read%s_%s()\n{\n\treturn %s_%s;\n}\n"), 
+			*TypeName,
+			*NumThreadsReservedName,
+			*InDataInterfaceName, 
+			*InDataInterfaceName, 
+			*NumThreadsReservedName);
+	}
+}
+
+UComputeDataProvider* UOptimusCopyKernelDataInterface::CreateDataProvider(TObjectPtr<UObject> InBinding,
+	uint64 InInputMask, uint64 InOutputMask) const
+{
+	UOptimusCopyKernelDataProvider* Provider = NewObject<UOptimusCopyKernelDataProvider>();
+
+	Provider->InitFromDataInterface(this, InBinding);
+
+	return Provider;
+}
+
+void UOptimusCopyKernelDataInterface::SetExecutionDomain(const FString& InExecutionDomain)
+{
+	NumThreadsExpression = InExecutionDomain;
+}
+
+void UOptimusCopyKernelDataInterface::SetComponentBinding(const UOptimusComponentSourceBinding* InBinding)
+{
+	ComponentSourceBinding = InBinding;
+}
+
+void UOptimusCopyKernelDataProvider::InitFromDataInterface(const UOptimusCopyKernelDataInterface* InDataInterface, const UObject* InBinding)
+{
+	WeakComponent = Cast<UActorComponent>(InBinding);
+	WeakDataInterface = InDataInterface;
+}
+
+FComputeDataProviderRenderProxy* UOptimusCopyKernelDataProvider::GetRenderProxy()
+{
+	TArray<int32> InvocationCounts;
+	int32 TotalThreadCount = 0;
+	
+	if (!GetInvocationThreadCounts(InvocationCounts, TotalThreadCount))
+	{
+		InvocationCounts.Reset();
+	}
+	
+	FOptimusCopyKernelDataProviderProxy* Proxy = new FOptimusCopyKernelDataProviderProxy(MoveTemp(InvocationCounts), TotalThreadCount);
+	return Proxy;
+}
+
+bool UOptimusCopyKernelDataProvider::GetInvocationThreadCounts(
+	TArray<int32>& OutInvocationThreadCount,
+	int32& OutTotalThreadCount
+	) const
+{
+	if (!WeakDataInterface.IsValid())
+	{
+		return false;
+	}
+
+	if (!WeakDataInterface->ComponentSourceBinding.IsValid())
+	{
+		return false;
+	}
+	
+	const UOptimusComponentSource* ComponentSource = WeakDataInterface->ComponentSourceBinding->GetComponentSource();
+	const FString& NumThreadsExpression = WeakDataInterface->NumThreadsExpression;
+	
+	if (!WeakComponent.IsValid() || !ComponentSource)
+	{
+		return false;
+	}
+	
+	const UActorComponent* Component = WeakComponent.Get();
+	int32 LodIndex = ComponentSource->GetLodIndex(Component);
+	int32 NumInvocations = ComponentSource->GetDefaultNumInvocations(Component, LodIndex);
+
+	// In some cases the component can be set to a unusable state intentionally
+	// such as when the OptimusEditor shuts down,
+	// in which case we don't have data to work with so simply do nothing
+	if (NumInvocations < 1)
+	{
+		return false;
+	}
+	
+	TMap<FName, float> EngineConstants;
+	TMap<FName, TArray<int32>> ElementCountsPerDomain;
+
+	for(FName ExecutionDomain: ComponentSource->GetExecutionDomains())
+	{
+		EngineConstants.Add(ExecutionDomain, 0);
+
+		TArray<int32>& ElementCounts = ElementCountsPerDomain.FindOrAdd(ExecutionDomain);
+		if (!ComponentSource->GetComponentElementCountsForExecutionDomain(
+				ExecutionDomain,
+				Component,
+				LodIndex,
+				ElementCounts))
+		{
+			return false;
+		}
+
+		if (!ensure(NumInvocations == ElementCounts.Num()))
+		{
+			// Component source needs to provide as many values for each of its execution domains as there are invocations
+			return false;
+		}
+	}
+
+	using namespace Optimus::Expression;
+
+	FEngine Engine;
+	TVariant<FExpressionObject, FParseError> ParseResult = Engine.Parse(NumThreadsExpression, [EngineConstants](FName InName)->TOptional<float>
+	{
+		if (const float* Value = EngineConstants.Find(InName))
+		{
+			return *Value;
+		};
+
+		return {};
+	});
+	
+	if (ParseResult.IsType<FParseError>())
+	{
+		return false;
+	}
+
+	OutInvocationThreadCount.Reset(NumInvocations);
+	for (int32 Index = 0; Index < NumInvocations; Index++)
+	{
+		for (TPair<FName, float>& Constant: EngineConstants)
+		{
+			const FName ConstantName = Constant.Key;
+			float& Value = Constant.Value;
+			
+			const TArray<int32>& ElementCounts = ElementCountsPerDomain[ConstantName]; 
+			Value = ElementCounts.IsValidIndex(Index) ? ElementCounts[Index] : 1;
+		}
+		
+		const int32 Count = static_cast<int32>(Engine.Execute(ParseResult.Get<FExpressionObject>(),
+			[EngineConstants](FName InName)->TOptional<float>
+			{
+				if (const float* Value = EngineConstants.Find(InName))
+				{
+					return *Value;
+				};
+
+				return {};
+			}));
+
+		if (Count < 0)
+		{
+			return false;
+		}
+		
+		OutInvocationThreadCount.Add(Count);
+	}
+
+	OutTotalThreadCount = 0;
+	for (int32 Count : OutInvocationThreadCount)
+	{
+		OutTotalThreadCount += Count;
+	}
+
+	// Extra validation for unified dispatch
+
+	// We want to make sure the results are the same for unified and non-unified
+	// Do the sum across invocations first and then evaluate the expression
+	for (TPair<FName, float>& Constant : EngineConstants)
+	{
+		const FName ConstantName = Constant.Key;
+		float& Value = Constant.Value;
+		Value = 0;
+		const TArray<int32>& ElementCounts = ElementCountsPerDomain[ConstantName];
+		
+		for (int32 Count : ElementCounts )
+		{
+			Value += Count;
+		}
+	}
+	
+	
+	const int32 TotalElementCountForUnifiedDispatch = static_cast<int32>(Engine.Execute(ParseResult.Get<FExpressionObject>(), 
+		[EngineConstants](FName InName)->TOptional<float>
+		{
+			if (const float* Value = EngineConstants.Find(InName))
+			{
+				return *Value;
+			};
+
+			return {};
+		}));
+	
+	if (OutTotalThreadCount != TotalElementCountForUnifiedDispatch)
+	{
+		return false;
+	}
+
+	return true;	
+}
+
+FOptimusCopyKernelDataProviderProxy::FOptimusCopyKernelDataProviderProxy(
+	TArray<int32>&& InInvocationThreadCounts,
+	int32 InTotalThreadCount
+	) :
+	InvocationThreadCounts(MoveTemp(InInvocationThreadCounts)),
+	TotalThreadCount(InTotalThreadCount)
+{
+}
+
+bool FOptimusCopyKernelDataProviderProxy::IsValid(FValidationData const& InValidationData) const
+{
+	if (InvocationThreadCounts.Num() == 0)
+	{
+		return false;
+	}
+
+	if (TotalThreadCount <= 0)
+	{
+		return false;
+	}
+	
+	return true;
+}
+
+int32 FOptimusCopyKernelDataProviderProxy::GetDispatchThreadCount(TArray<FIntVector>& InOutThreadCounts) const
+{
+	InOutThreadCounts.Reset(InvocationThreadCounts.Num());
+	for (const int32 Count : InvocationThreadCounts)
+	{
+		InOutThreadCounts.Add({Count, 1, 1});
+	}
+	return InOutThreadCounts.Num();
+}
+
+void FOptimusCopyKernelDataProviderProxy::GatherDispatchData(FDispatchData const& InDispatchData)
+{
+	for (int32 InvocationIndex = 0; InvocationIndex < InDispatchData.NumInvocations; ++InvocationIndex)
+	{
+		int32 NumThreads;
+		if (InDispatchData.bUnifiedDispatch)
+		{
+			NumThreads = TotalThreadCount;
+		}
+		else
+		{
+			NumThreads = InvocationThreadCounts[InvocationIndex];
+		}
+
+		uint8* ParameterBuffer = (InDispatchData.ParameterBuffer + InDispatchData.ParameterBufferOffset + InDispatchData.ParameterBufferStride * InvocationIndex);
+		
+		const int32 ThreadCount3D[] = {NumThreads, 1, 1};
+		FMemory::Memcpy(ParameterBuffer, ThreadCount3D, 3 * sizeof(int32));
+	}
+}
+
+
diff --git a/Engine/Plugins/Animation/DeformerGraph/Source/OptimusCore/Private/DataInterfaces/OptimusDataInterfaceCopyKernel.h b/Engine/Plugins/Animation/DeformerGraph/Source/OptimusCore/Private/DataInterfaces/OptimusDataInterfaceCopyKernel.h
new file mode 100644
--- /dev/null
+++ b/Engine/Plugins/Animation/DeformerGraph/Source/OptimusCore/Private/DataInterfaces/OptimusDataInterfaceCopyKernel.h
@@ -0,0 +1,92 @@
+// Copyright Epic Games, Inc. All Rights Reserved.
+
+#pragma once
+
+#include "IOptimusComputeKernelDataInterface.h"
+#include "OptimusComponentSource.h"
+#include "OptimusComputeDataInterface.h"
+#include "OptimusDataType.h"
+#include "ComputeFramework/ComputeDataProvider.h"
+
+#include "OptimusDataInterfaceCopyKernel.generated.h"
+
+class UOptimusComponentSource;
+class UOptimusComponentSourceBinding;
+
+UCLASS(Category = ComputeFramework)
+class OPTIMUSCORE_API UOptimusCopyKernelDataInterface :
+	public UComputeDataInterface,
+	public IOptimusComputeKernelDataInterface
+{
+	GENERATED_BODY()
+
+public:
+	static const FString NumThreadsReservedName;
+	
+	//~ Begin UComputeDataInterface Interface
+	TCHAR const* GetClassName() const override { return TEXT("CopyKernelData"); }
+	bool IsExecutionInterface() const override { return true; }
+	bool CanSupportUnifiedDispatch() const override { return true; }
+	void GetSupportedInputs(TArray<FShaderFunctionDefinition>& OutFunctions) const override;
+	void GetShaderParameters(TCHAR const* UID, FShaderParametersMetadataBuilder& InOutBuilder, FShaderParametersMetadataAllocations& InOutAllocations) const override;
+	void GetShaderHash(FString& InOutKey) const override;
+	void GetHLSL(FString& OutHLSL, FString const& InDataInterfaceName) const override;
+	UComputeDataProvider* CreateDataProvider(TObjectPtr<UObject> InBinding, uint64 InInputMask, uint64 InOutputMask) const override;
+	//~ End UComputeDataInterface Interface
+
+	//~ Begin IOptimusComputeKernelDataInterface Interface
+	void SetExecutionDomain(const FString& InExecutionDomain) override;
+	void SetComponentBinding(const UOptimusComponentSourceBinding* InBinding) override;
+	//~ End IOptimusComputeKernelDataInterface Interface
+	
+	UPROPERTY()
+	TWeakObjectPtr<const UOptimusComponentSourceBinding> ComponentSourceBinding;
+	
+	UPROPERTY()
+	FString NumThreadsExpression;
+};
+
+/** Compute Framework Data Provider for each copy kernel. */
+UCLASS()
+class UOptimusCopyKernelDataProvider :
+	public UComputeDataProvider
+{
+	GENERATED_BODY()
+
+public:
+	void InitFromDataInterface(const UOptimusCopyKernelDataInterface* InDataInterface, const UObject* InBinding);
+	
+	//~ Begin UComputeDataProvider Interface
+	FComputeDataProviderRenderProxy* GetRenderProxy() override;
+	//~ End UComputeDataProvider Interface
+
+	TWeakObjectPtr<const UActorComponent> WeakComponent = nullptr;
+
+	TWeakObjectPtr<const UOptimusCopyKernelDataInterface> WeakDataInterface = nullptr;
+
+protected:
+	bool GetInvocationThreadCounts(
+		TArray<int32>& OutInvocationThreadCount,
+		int32& OutTotalThreadCount
+		) const;
+};
+
+class FOptimusCopyKernelDataProviderProxy : public FComputeDataProviderRenderProxy
+{
+public:
+	FOptimusCopyKernelDataProviderProxy(
+		TArray<int32>&& InInvocationThreadCounts,
+		int32 InTotalThreadCount
+		);
+
+	//~ Begin FComputeDataProviderRenderProxy Interface
+	bool IsValid(FValidationData const& InValidationData) const override;
+	int32 GetDispatchThreadCount(TArray<FIntVector>& InOutThreadCounts) const override;
+	void GatherDispatchData(FDispatchData const& InDispatchData) override;
+	//~ End FComputeDataProviderRenderProxy Interface
+
+	TArray<int32> InvocationThreadCounts;
+	int32 TotalThreadCount;
+};
+
+
diff --git a/Engine/Plugins/Animation/DeformerGraph/Source/OptimusCore/Private/DataInterfaces/OptimusDataInterfaceCustomComputeKernel.cpp b/Engine/Plugins/Animation/DeformerGraph/Source/OptimusCore/Private/DataInterfaces/OptimusDataInterfaceCustomComputeKernel.cpp
--- a/Engine/Plugins/Animation/DeformerGraph/Source/OptimusCore/Private/DataInterfaces/OptimusDataInterfaceCustomComputeKernel.cpp
+++ b/Engine/Plugins/Animation/DeformerGraph/Source/OptimusCore/Private/DataInterfaces/OptimusDataInterfaceCustomComputeKernel.cpp
@@ -30,7 +30,7 @@ void UOptimusCustomComputeKernelDataInterface::GetShaderParameters(TCHAR const*
 	TArray<FShaderParametersMetadata*> DummyNestedStructs;
 	ComputeFramework::AddParamForType(Builder, *NumThreadsReservedName, FShaderValueType::Get(EShaderFundamentalType::Int, 3),DummyNestedStructs);
 
-	FShaderParametersMetadata* ShaderParameterMetadata = Builder.Build(FShaderParametersMetadata::EUseCase::ShaderParameterStruct, TEXT("UAnimAttributeDataInterface"));
+	FShaderParametersMetadata* ShaderParameterMetadata = Builder.Build(FShaderParametersMetadata::EUseCase::ShaderParameterStruct, TEXT("UCustomComputeKernelDataInterface"));
 
 	InOutAllocations.ShaderParameterMetadatas.Add(ShaderParameterMetadata);
 	InOutAllocations.ShaderParameterMetadatas.Append(DummyNestedStructs);
@@ -42,7 +42,8 @@ void UOptimusCustomComputeKernelDataInterface::GetShaderParameters(TCHAR const*
 
 void UOptimusCustomComputeKernelDataInterface::GetShaderHash(FString& InOutKey) const
 {
-	
+	// UComputeGraph::BuildKernelSource hashes the result of GetHLSL()
+	// Only append additional hashes here if the HLSL contains any additional includes	
 }
 
 void UOptimusCustomComputeKernelDataInterface::GetHLSL(FString& OutHLSL, FString const& InDataInterfaceName) const
diff --git a/Engine/Plugins/Animation/DeformerGraph/Source/OptimusCore/Private/Nodes/OptimusNode_ComputeKernelBase.cpp b/Engine/Plugins/Animation/DeformerGraph/Source/OptimusCore/Private/Nodes/OptimusNode_ComputeKernelBase.cpp
--- a/Engine/Plugins/Animation/DeformerGraph/Source/OptimusCore/Private/Nodes/OptimusNode_ComputeKernelBase.cpp
+++ b/Engine/Plugins/Animation/DeformerGraph/Source/OptimusCore/Private/Nodes/OptimusNode_ComputeKernelBase.cpp
@@ -310,40 +310,6 @@ TOptional<FText> UOptimusNode_ComputeKernelBase::ValidateForCompile() const
 				}
 			}
 		}
-
-		if (Pin->GetDirection() == EOptimusNodePinDirection::Output)
-		{
-			if (DoesOutputPinSupportAtomic(Pin) || DoesOutputPinSupportRead(Pin))
-			{
-				TArray<FOptimusRoutedNodePin> ConnectedPins = Pin->GetConnectedPinsWithRouting();
-				if (ConnectedPins.Num() > 1)
-				{
-					int32 TransientBufferCount = 0;
-					int32 PersistentBufferCount = 0;
-					for (const FOptimusRoutedNodePin& ConnectedPin : ConnectedPins)
-					{
-						if (Cast<IOptimusComputeKernelProvider>(ConnectedPin.NodePin->GetOwningNode()))
-						{
-							// At most 1 transient buffer, shared among connected kernels
-							TransientBufferCount = 1;
-						}
-						else if (Cast<UOptimusNode_ResourceAccessorBase>(ConnectedPin.NodePin->GetOwningNode()))
-						{
-							PersistentBufferCount++;
-						}
-						else
-						{
-							return FText::Format(LOCTEXT("TooManyPinConnections", "Pin '{0}' supports Atomic Writes or Read Access and should be connected to either a single resource node or strictly kernel nodes"), FText::FromName(Pin->GetUniqueName()));
-						}
-					}
-
-					if (PersistentBufferCount + TransientBufferCount > 1)
-					{
-						return FText::Format(LOCTEXT("TooManyPinConnections", "Pin '{0}' supports Atomic Writes or Read Access and should be connected to either a single resource node or strictly kernel nodes"), FText::FromName(Pin->GetUniqueName()));
-					}
-				}
-			}
-		}
 	}
 
 	// We should have at least the primary group, which needs to have a unique component binding
@@ -356,64 +322,7 @@ TOptional<FText> UOptimusNode_ComputeKernelBase::ValidateForCompile() const
 }
 
 
-FString UOptimusNode_ComputeKernelBase::GetCookedKernelSource(
-	const FString& InObjectPathName,
-	const FString& InShaderSource,
-	const FString& InKernelName,
-	FIntVector InGroupSize
-	)
-{
-	// FIXME: Create source range mappings so that we can go from error location to
-	// our source.
-	FString Source = InShaderSource;
 
-#if PLATFORM_WINDOWS
-	// Remove old-school stuff.
-	Source.ReplaceInline(TEXT("\r"), TEXT(""));
-#endif
-
-	FString ShaderPathName = InObjectPathName;
-	Optimus::ConvertObjectPathToShaderFilePath(ShaderPathName);
-
-	const bool bHasKernelKeyword = Source.Contains(TEXT("KERNEL"), ESearchCase::CaseSensitive);
-
-	const FString ComputeShaderUtilsInclude = TEXT("#include \"/Engine/Private/ComputeShaderUtils.ush\"");
-	
-	const FString KernelFunc = FString::Printf(
-		TEXT("[numthreads(%d,%d,%d)]\nvoid %s(uint3 GroupId : SV_GroupID, uint GroupIndex : SV_GroupIndex)"), 
-		InGroupSize.X, InGroupSize.Y, InGroupSize.Z, *InKernelName);
-	
-	const FString UnWrappedDispatchThreadId = FString::Printf(
-	TEXT("GetUnWrappedDispatchThreadId(GroupId, GroupIndex, %d)"),
-		InGroupSize.X * InGroupSize.Y * InGroupSize.Z
-	);
-
-	if (bHasKernelKeyword)
-	{
-		Source.ReplaceInline(TEXT("KERNEL"), TEXT("void __kernel_func(uint Index)"), ESearchCase::CaseSensitive);
-
-		return FString::Printf(
-			TEXT(
-				"#line 1 \"%s\"\n"
-				"%s\n\n"
-				"%s\n\n"
-				"%s { __kernel_func(%s); }\n"
-				), *ShaderPathName, *Source, *ComputeShaderUtilsInclude,*KernelFunc, *UnWrappedDispatchThreadId);
-	}
-	else
-	{
-		return FString::Printf(
-		TEXT(
-			"%s\n"
-			"%s\n"
-			"{\n"
-			"uint Index = %s;\n"
-			"#line 1 \"%s\"\n"
-			"%s\n"
-			"}\n"
-			), *ComputeShaderUtilsInclude,*KernelFunc, *UnWrappedDispatchThreadId, *ShaderPathName, *Source);
-	}
-}
 
 
 
@@ -644,37 +553,18 @@ void UOptimusNode_ComputeKernelBase::ProcessOutputPinForComputeKernel(
 
 			if (!ConnectedNode)
 			{
-				if (const UOptimusTransientBufferDataInterface* TransientBufferDataInterface = Cast<const UOptimusTransientBufferDataInterface>(DataInterface);
-							ensure(TransientBufferDataInterface))
+				// No connected node indicates the kernel output is cached in a generated raw buffer
+				if (const UOptimusRawBufferDataInterface* RawBufferDataInterface = Cast<const UOptimusRawBufferDataInterface>(DataInterface);
+					ensure(RawBufferDataInterface))
 				{
 					const int32 OutputIndex = UOptimusRawBufferDataInterface::GetWriteValueOutputIndex(EOptimusBufferWriteType::Write);
-					WriteConnectionDefs.Add({TransientBufferDataInterface, WriteFunctions[OutputIndex].Name, TEXT("TransientBuffer")});	
+					WriteConnectionDefs.Add({RawBufferDataInterface, WriteFunctions[OutputIndex].Name, TEXT("RawBuffer")});
 				}
-				
-				check(KernelConnections->Num() == 1);
-				break;
 			}
 			else
 			{
-				if (Cast<const IOptimusComputeKernelProvider>(ConnectedNode))
-				{
-					if (!SharedRawBufferDI)
-					{
-						if (const UOptimusRawBufferDataInterface* RawBufferDataInterface = Cast<const UOptimusRawBufferDataInterface>(DataInterface);
-							ensure(RawBufferDataInterface))
-						{
-							const int32 OutputIndex = UOptimusRawBufferDataInterface::GetWriteValueOutputIndex(EOptimusBufferWriteType::Write);
-							WriteConnectionDefs.Add({RawBufferDataInterface, WriteFunctions[OutputIndex].Name, TEXT("RawBuffer")});
-							SharedRawBufferDI = DataInterface;
-						}	
-					}
-					else
-					{
-						// All connected kernels should share the same data interface
-						check(SharedRawBufferDI == DataInterface);
-					}
-				}
-				else if (const IOptimusDataInterfaceProvider* InterfaceProvider = Cast<const IOptimusDataInterfaceProvider>(ConnectedNode))
+				if (const IOptimusDataInterfaceProvider* InterfaceProvider = Cast<const IOptimusDataInterfaceProvider>(ConnectedNode);
+					ensure(InterfaceProvider))
 				{
 					int32 DataFunctionIndex = InterfaceProvider->GetDataFunctionIndexFromPin(ConnectedPin);
 					TArray<FShaderFunctionDefinition> FunctionDefinitions;
diff --git a/Engine/Plugins/Animation/DeformerGraph/Source/OptimusCore/Private/Nodes/OptimusNode_ComputeKernelBase.h b/Engine/Plugins/Animation/DeformerGraph/Source/OptimusCore/Private/Nodes/OptimusNode_ComputeKernelBase.h
--- a/Engine/Plugins/Animation/DeformerGraph/Source/OptimusCore/Private/Nodes/OptimusNode_ComputeKernelBase.h
+++ b/Engine/Plugins/Animation/DeformerGraph/Source/OptimusCore/Private/Nodes/OptimusNode_ComputeKernelBase.h
@@ -92,14 +92,7 @@ protected:
 
 		return IndexNames;
 	}
-
-	static FString GetCookedKernelSource(
-		const FString& InObjectPathName,
-		const FString& InShaderSource,
-		const FString& InKernelName,
-		FIntVector InGroupSize
-		);
-
+	
 	static FString GetAtomicWriteFunctionName(EOptimusBufferWriteType InWriteType, const FString& InBindingName);
 	static FString GetReadFunctionName(const FString& InBindingName);
 	
diff --git a/Engine/Plugins/Animation/DeformerGraph/Source/OptimusCore/Private/Nodes/OptimusNode_ComputeKernelFunction.cpp b/Engine/Plugins/Animation/DeformerGraph/Source/OptimusCore/Private/Nodes/OptimusNode_ComputeKernelFunction.cpp
--- a/Engine/Plugins/Animation/DeformerGraph/Source/OptimusCore/Private/Nodes/OptimusNode_ComputeKernelFunction.cpp
+++ b/Engine/Plugins/Animation/DeformerGraph/Source/OptimusCore/Private/Nodes/OptimusNode_ComputeKernelFunction.cpp
@@ -185,7 +185,7 @@ FIntVector UOptimusNode_ComputeKernelFunction::GetGroupSize() const
 
 FString UOptimusNode_ComputeKernelFunction::GetKernelSourceText() const
 {
-	return GetCookedKernelSource(GetPathName(), GetGeneratorClass()->ShaderSource, GetKernelName(), GetGroupSize());
+	return Optimus::GetCookedKernelSource(GetPathName(), GetGeneratorClass()->ShaderSource, GetKernelName(), GetGroupSize());
 }
 
 
diff --git a/Engine/Plugins/Animation/DeformerGraph/Source/OptimusCore/Private/Nodes/OptimusNode_CustomComputeKernel.cpp b/Engine/Plugins/Animation/DeformerGraph/Source/OptimusCore/Private/Nodes/OptimusNode_CustomComputeKernel.cpp
--- a/Engine/Plugins/Animation/DeformerGraph/Source/OptimusCore/Private/Nodes/OptimusNode_CustomComputeKernel.cpp
+++ b/Engine/Plugins/Animation/DeformerGraph/Source/OptimusCore/Private/Nodes/OptimusNode_CustomComputeKernel.cpp
@@ -103,7 +103,7 @@ UOptimusNode_CustomComputeKernel::UOptimusNode_CustomComputeKernel()
 
 FString UOptimusNode_CustomComputeKernel::GetKernelSourceText() const
 {
-	return GetCookedKernelSource(GetPathName(), ShaderSource.ShaderText, KernelName.ToString(), GroupSize);
+	return Optimus::GetCookedKernelSource(GetPathName(), ShaderSource.ShaderText, KernelName.ToString(), GroupSize);
 }
 
 FOptimusExecutionDomain UOptimusNode_CustomComputeKernel::GetExecutionDomain() const
diff --git a/Engine/Plugins/Animation/DeformerGraph/Source/OptimusCore/Private/OptimusDeformer.cpp b/Engine/Plugins/Animation/DeformerGraph/Source/OptimusCore/Private/OptimusDeformer.cpp
--- a/Engine/Plugins/Animation/DeformerGraph/Source/OptimusCore/Private/OptimusDeformer.cpp
+++ b/Engine/Plugins/Animation/DeformerGraph/Source/OptimusCore/Private/OptimusDeformer.cpp
@@ -12,6 +12,7 @@
 #include "DataInterfaces/OptimusDataInterfaceGraph.h"
 #include "DataInterfaces/OptimusDataInterfaceRawBuffer.h"
 #include "DataInterfaces/OptimusDataInterfaceLoopTerminal.h"
+#include "DataInterfaces/OptimusDataInterfaceCopyKernel.h"
 #include "IOptimusComputeKernelProvider.h"
 #include "IOptimusDataInterfaceProvider.h"
 #include "IOptimusValueProvider.h"
@@ -2006,7 +2007,27 @@ TArray<FOptimusComputeGraphInfo> UOptimusDeformer::CompileNodeGraphToComputeGrap
 	{
 		SourcePinToTargetPins.FindOrAdd(LinkedPins.Value).Add(LinkedPins.Key);
 	}
-	
+
+	TMap<FOptimusInstancedPin, TArray<FOptimusInstancedPin>> LinksToInsertCopyKernel;
+	for (TPair<FOptimusInstancedPin, TArray<FOptimusInstancedPin>> OutputLink : SourcePinToTargetPins)
+	{
+		const FOptimusInstancedPin& SourcePin = OutputLink.Key;
+		const UOptimusNode* SourceNode = SourcePin.InstancedNode.RoutedNode.Node;
+
+		for (const FOptimusInstancedPin& TargetPin : OutputLink.Value)
+		{
+			const UOptimusNode* TargetNode = TargetPin.InstancedNode.RoutedNode.Node;
+
+			if (Cast<const IOptimusDataInterfaceProvider>(TargetNode))
+			{
+				if (Cast<const IOptimusDataInterfaceProvider>(SourceNode))
+				{
+					LinksToInsertCopyKernel.FindOrAdd(SourcePin).Add(TargetPin);
+				}
+			}
+		}
+	}
+
 	// Create all the data interfaces: node, graph, kernel outputs, loop terminal data
 	
 	// The component binding for the graph data is the primary binding on the deformer.
@@ -2152,83 +2173,97 @@ TArray<FOptimusComputeGraphInfo> UOptimusDeformer::CompileNodeGraphToComputeGrap
 				{
 					FOptimusInstancedPin InstancedPin = {InstancedNode, Pin};
 
+					bool bShouldCreateRawBuffer = false;
+					bool bShouldUseImplicitPersistentDI = false;
+					bool bShouldCopyToDataInterface = false;
+
+					if (KernelProvider->DoesOutputPinSupportAtomic(Pin) || KernelProvider->DoesOutputPinSupportRead(Pin))
+					{
+						bShouldCreateRawBuffer = true;
+						bShouldCopyToDataInterface = true;
+					}
+
+					TArray<FOptimusInstancedPin> TargetDataInterfacePins;
+
 					if (TArray<FOptimusInstancedPin>* TargetInstancedPins = SourcePinToTargetPins.Find(InstancedPin))
 					{
-						bool bShouldUseImplicitPersistentDI = false;
-						TArray<const UOptimusNodePin*> TargetKernelPins;
-						
 						for (const FOptimusInstancedPin& TargetInstancedPin : *TargetInstancedPins)
 						{
 							const FOptimusRoutedConstNode TargetRoutedNode = TargetInstancedPin.InstancedNode.RoutedNode;
 							const UOptimusNode* TargetNode = TargetRoutedNode.Node;
 							const UOptimusNodePin* TargetPin = TargetInstancedPin.Pin;
-							
-							if (UOptimusComputeDataInterface** NodeDataInterface = NodeDataInterfaceMap.Find(TargetNode))
+
+							if (Cast<const IOptimusDataInterfaceProvider>(TargetNode))
 							{
-								KernelOutputMap[InstancedNode].FindOrAdd(Pin).Add({*NodeDataInterface, TargetPin});
+								if (InNodeGraph->GetGraphType() == EOptimusNodeGraphType::Update && 
+									KernelToGraphType[RoutedNode] == EOptimusNodeGraphType::Setup)
+								{
+									bShouldCreateRawBuffer = true;
+									bShouldUseImplicitPersistentDI = true;
+									bShouldCopyToDataInterface = true;
+								}
+
+								TargetDataInterfacePins.Add(TargetInstancedPin);
 							}
 							else if (Cast<const IOptimusComputeKernelProvider>(TargetNode))
 							{
+								bShouldCreateRawBuffer = true;
 								if (KernelToGraphType[RoutedNode] != KernelToGraphType[TargetRoutedNode])
 								{
 									bShouldUseImplicitPersistentDI = true;
 								}
-
-								TargetKernelPins.Add(TargetPin);
 							}
 						}
+					}
 
-						if (TargetKernelPins.Num() > 0)
+					if (bShouldCreateRawBuffer)
+					{
+						UOptimusRawBufferDataInterface* RawBufferDI = nullptr;
+						if (bShouldUseImplicitPersistentDI)
 						{
-							UOptimusRawBufferDataInterface* RawBufferDI = nullptr;
-							if (bShouldUseImplicitPersistentDI)
+							RawBufferDI = NewObject<UOptimusImplicitPersistentBufferDataInterface>(this);
+							if (KernelProvider->DoesOutputPinSupportAtomic(Pin))
 							{
-								RawBufferDI = NewObject<UOptimusImplicitPersistentBufferDataInterface>(this);
-								if (KernelProvider->DoesOutputPinSupportAtomic(Pin))
-								{
-									CastChecked<UOptimusImplicitPersistentBufferDataInterface>(RawBufferDI)->bZeroInitForAtomicWrites = true;
-								}
+								CastChecked<UOptimusImplicitPersistentBufferDataInterface>(RawBufferDI)->bZeroInitForAtomicWrites = true;
 							}
-							else
+						}
+						else
+						{
+							RawBufferDI = NewObject<UOptimusTransientBufferDataInterface>(this);
+							if (KernelProvider->DoesOutputPinSupportAtomic(Pin))
 							{
-								RawBufferDI = NewObject<UOptimusTransientBufferDataInterface>(this);
-								if (KernelProvider->DoesOutputPinSupportAtomic(Pin))
-								{
-									CastChecked<UOptimusTransientBufferDataInterface>(RawBufferDI)->bZeroInitForAtomicWrites = true;
-								}
+								CastChecked<UOptimusTransientBufferDataInterface>(RawBufferDI)->bZeroInitForAtomicWrites = true;
 							}
-								
-							RawBufferDI->ValueType = Pin->GetDataType()->ShaderValueType;
-							RawBufferDI->DataDomain = Pin->GetDataDomain();
-							RawBufferDI->ComponentSourceBinding = KernelPrimaryBinding;
-							DataInterfaceToBindingIndexMap.Add(RawBufferDI) = PrimaryBindingIndex;
-								
-							KernelOutputDataInterfaceMap.Add(InstancedPin) = RawBufferDI;
+						}
 
-							// All connected kernels share the same raw buffer data interface
-							for (const UOptimusNodePin* TargetKernelPin : TargetKernelPins)
-							{
-								KernelOutputMap[InstancedNode].FindOrAdd(Pin).Add({RawBufferDI, TargetKernelPin});		
-							}
-						}	
+						RawBufferDI->ValueType = Pin->GetDataType()->ShaderValueType;
+						RawBufferDI->DataDomain = Pin->GetDataDomain();
+						RawBufferDI->ComponentSourceBinding = KernelPrimaryBinding;
+
+						KernelOutputDataInterfaceMap.Add(InstancedPin) = RawBufferDI;
+						DataInterfaceToBindingIndexMap.Add(RawBufferDI) = PrimaryBindingIndex;
+
+						// All connected kernels share the same raw buffer data interface
+						KernelOutputMap[InstancedNode].FindOrAdd(Pin).Add({RawBufferDI, nullptr});
 					}
-					else
+
+					for (const FOptimusInstancedPin& TargetInstancedPin : TargetDataInterfacePins)
 					{
-						if (KernelProvider->DoesOutputPinSupportAtomic(Pin) || KernelProvider->DoesOutputPinSupportRead(Pin))
+						const FOptimusRoutedConstNode TargetRoutedNode = TargetInstancedPin.InstancedNode.RoutedNode;
+						const UOptimusNode* TargetNode = TargetRoutedNode.Node;
+						const UOptimusNodePin* TargetPin = TargetInstancedPin.Pin;
+
+						if (bShouldCopyToDataInterface)
 						{
-							UOptimusTransientBufferDataInterface* TransientBufferDI = NewObject<UOptimusTransientBufferDataInterface>(this);
-							if (KernelProvider->DoesOutputPinSupportAtomic(Pin))
+							check(bShouldCreateRawBuffer);
+							LinksToInsertCopyKernel.FindOrAdd(InstancedPin).Add(TargetInstancedPin);
+						}
+						else
+						{
+							if (UOptimusComputeDataInterface** NodeDataInterface = NodeDataInterfaceMap.Find(TargetNode))
 							{
-								CastChecked<UOptimusTransientBufferDataInterface>(TransientBufferDI)->bZeroInitForAtomicWrites = true;
+								KernelOutputMap[InstancedNode].FindOrAdd(Pin).Add({*NodeDataInterface, TargetPin});
 							}
-							TransientBufferDI->ValueType = Pin->GetDataType()->ShaderValueType;
-							TransientBufferDI->DataDomain = Pin->GetDataDomain();
-							TransientBufferDI->ComponentSourceBinding = KernelPrimaryBinding;
-							DataInterfaceToBindingIndexMap.Add(TransientBufferDI) = PrimaryBindingIndex;
-								
-							KernelOutputDataInterfaceMap.Add(InstancedPin) = TransientBufferDI;
-
-							KernelOutputMap[InstancedNode].FindOrAdd(Pin).Add({TransientBufferDI, nullptr});	
 						}	
 					}
 				}
@@ -2236,6 +2271,65 @@ TArray<FOptimusComputeGraphInfo> UOptimusDeformer::CompileNodeGraphToComputeGrap
 		}
 	}
 
+	if (!LinksToInsertCopyKernel.IsEmpty())
+	{
+		GraphTypes.Add(InNodeGraph->GetGraphType());
+	}
+	
+	struct FDataInterfaceFunctionBinding
+	{
+		UOptimusComputeDataInterface* DataInterface;
+		int32 FunctionIndex;
+	};
+	
+	TMap<FOptimusInstancedPin, UComputeDataInterface*> CopyKernelDataInterfaceMap;
+	TMap<FOptimusInstancedPin, FDataInterfaceFunctionBinding> CopyFromDataInterfaceMap;
+	TMap<FOptimusInstancedPin, FDataInterfaceFunctionBinding> CopyToDataInterfaceMap;
+
+	for (TPair<FOptimusInstancedPin, TArray<FOptimusInstancedPin>> OutputLink : LinksToInsertCopyKernel)
+	{
+		const FOptimusInstancedPin& SourceInstancedPin = OutputLink.Key;
+		const UOptimusNode* SourceNode = SourceInstancedPin.Pin->GetOwningNode();
+		if (const IOptimusDataInterfaceProvider* InterfaceProvider = Cast<const IOptimusDataInterfaceProvider>(SourceNode))
+		{
+			FDataInterfaceFunctionBinding DataInterfaceBinding;
+			DataInterfaceBinding.DataInterface = NodeDataInterfaceMap[SourceNode];
+			DataInterfaceBinding.FunctionIndex = InterfaceProvider->GetDataFunctionIndexFromPin(SourceInstancedPin.Pin);
+			CopyFromDataInterfaceMap.Add(SourceInstancedPin) = DataInterfaceBinding;
+		}
+		else if (Cast<const IOptimusComputeKernelProvider>(SourceNode))
+		{
+			FDataInterfaceFunctionBinding DataInterfaceBinding;
+			DataInterfaceBinding.DataInterface = KernelOutputDataInterfaceMap[SourceInstancedPin];
+			DataInterfaceBinding.FunctionIndex = UOptimusRawBufferDataInterface::GetReadValueInputIndex(EOptimusBufferReadType::Default);
+
+			CopyFromDataInterfaceMap.Add(SourceInstancedPin) = DataInterfaceBinding;
+		}
+
+		UOptimusCopyKernelDataInterface* CopyKernelDataInterface = NewObject<UOptimusCopyKernelDataInterface>(this);
+
+		CopyKernelDataInterface->SetExecutionDomain(*SourceInstancedPin.Pin->GetDataDomain().AsExpression());
+		UOptimusComponentSourceBinding* Binding = *SourceInstancedPin.Pin->GetComponentSourceBindings().CreateConstIterator();
+		CopyKernelDataInterface->SetComponentBinding(Binding);
+
+		CopyKernelDataInterfaceMap.Add(SourceInstancedPin) = CopyKernelDataInterface;
+		DataInterfaceToBindingIndexMap.Add(CopyKernelDataInterface) = Binding->GetIndex();
+
+		for (const FOptimusInstancedPin& TargetInstancedPin : OutputLink.Value)
+		{
+			const UOptimusNode* TargetNode = TargetInstancedPin.Pin->GetOwningNode();
+
+			if (const IOptimusDataInterfaceProvider* InterfaceProvider = Cast<const IOptimusDataInterfaceProvider>(TargetNode);
+				ensure(InterfaceProvider))
+			{
+				FDataInterfaceFunctionBinding DataInterfaceBinding;
+				DataInterfaceBinding.DataInterface = NodeDataInterfaceMap[TargetNode];
+				DataInterfaceBinding.FunctionIndex = InterfaceProvider->GetDataFunctionIndexFromPin(TargetInstancedPin.Pin);
+				CopyToDataInterfaceMap.Add(TargetInstancedPin) = DataInterfaceBinding;
+			}
+		}
+	}
+
 
 	TArray<FOptimusComputeGraphInfo> GraphInfos;
 	for (EOptimusNodeGraphType GraphType : GraphTypes)
@@ -2320,7 +2414,41 @@ TArray<FOptimusComputeGraphInfo> UOptimusDeformer::CompileNodeGraphToComputeGrap
 				ComputeGraph->DataInterfaceToBinding.Add(DataInterfaceToBindingIndexMap[KernelDataInterface]);
 			}
 		}
-			
+
+		if (InNodeGraph->GetGraphType() == GraphInfo.GraphType)
+		{
+			for (TPair<FOptimusInstancedPin, TArray<FOptimusInstancedPin>> OutputLink : LinksToInsertCopyKernel)
+			{
+				const FOptimusInstancedPin& SourceInstancedPin = OutputLink.Key;	
+				UComputeDataInterface* CopyKernelDataInterface = CopyKernelDataInterfaceMap[SourceInstancedPin];
+
+				if (!ComputeGraph->DataInterfaces.Contains(CopyKernelDataInterface))
+				{
+					ComputeGraph->DataInterfaces.Add(CopyKernelDataInterface);
+					ComputeGraph->DataInterfaceToBinding.Add(DataInterfaceToBindingIndexMap[CopyKernelDataInterface]);
+				}	
+
+				const FDataInterfaceFunctionBinding& CopyFromBinding = CopyFromDataInterfaceMap[SourceInstancedPin];
+				if (!ComputeGraph->DataInterfaces.Contains(CopyFromBinding.DataInterface))
+				{
+					ComputeGraph->DataInterfaces.Add(CopyFromBinding.DataInterface);
+					ComputeGraph->DataInterfaceToBinding.Add(DataInterfaceToBindingIndexMap[CopyFromBinding.DataInterface]);
+				}
+
+				for (int32 TargetIndex = 0 ; TargetIndex < OutputLink.Value.Num(); TargetIndex++)
+				{
+					const FOptimusInstancedPin& TargetInstancedPin = OutputLink.Value[TargetIndex];
+					const FDataInterfaceFunctionBinding& CopyToBinding = CopyToDataInterfaceMap[TargetInstancedPin];
+
+					if (!ComputeGraph->DataInterfaces.Contains(CopyToBinding.DataInterface))
+					{
+						ComputeGraph->DataInterfaces.Add(CopyToBinding.DataInterface);
+						ComputeGraph->DataInterfaceToBinding.Add(DataInterfaceToBindingIndexMap[CopyToBinding.DataInterface]);
+					}	
+				}
+			}
+		}
+
 		// Create bound kernels
 		struct FKernelWithDataBindings
 		{
@@ -2400,7 +2528,120 @@ TArray<FOptimusComputeGraphInfo> UOptimusDeformer::CompileNodeGraphToComputeGrap
 				ComputeGraph->KernelToNode.Add(ConnectedNode.Node);
 			}
 		}
+
+		
+		if (InNodeGraph->GetGraphType() == GraphInfo.GraphType)
+		{
+			for (TPair<FOptimusInstancedPin, TArray<FOptimusInstancedPin>> OutputLink : LinksToInsertCopyKernel)
+			{
+				// Create a copy kernel per source pin, that copies from 1 source pin to multiple target pins
+				const FOptimusInstancedPin& SourceInstancedPin = OutputLink.Key;
+				const FShaderValueTypeHandle ValueType = SourceInstancedPin.Pin->GetDataType()->ShaderValueType;
+				
+				FKernelWithDataBindings BoundKernel;
+
+				BoundKernel.Kernel = NewObject<UComputeKernel>(this);
+
+				FOptimus_InterfaceBindingMap& InputDataBindings = BoundKernel.InputDataBindings;
+				FOptimus_InterfaceBindingMap& OutputDataBindings = BoundKernel.OutputDataBindings;
+
+				UOptimusKernelSource* KernelSource = NewObject<UOptimusKernelSource>(BoundKernel.Kernel);
+				FString SourceText;
+				SourceText = TEXT("if (Index >= ReadNumThreads().x) return;\n");
+
+				{
+					UComputeDataInterface* CopyKernelDataInterface = CopyKernelDataInterfaceMap[SourceInstancedPin];
+					TArray<FShaderFunctionDefinition> Functions;
+					CopyKernelDataInterface->GetSupportedInputs(Functions);
+					// Simply grab everything the kernel data interface has to offer
+					for (int32 FuncIndex = 0; FuncIndex < Functions.Num(); FuncIndex++)
+					{
+						FShaderFunctionDefinition FuncDef = Functions[FuncIndex];
+						for (FShaderParamTypeDefinition& ParamType : FuncDef.ParamTypes)
+						{
+							// Making sure parameter has type declaration generated
+							ParamType.ResetTypeDeclaration();
+						}
+
+						FOptimus_InterfaceBinding InterfaceBinding;
+						InterfaceBinding.DataInterface = CopyKernelDataInterface;
+						InterfaceBinding.DataInterfaceBindingIndex = FuncIndex;
+						InterfaceBinding.BindingFunctionName = FuncDef.Name;
+						InterfaceBinding.BindingFunctionNamespace = FString();
+				
+						InputDataBindings.Add(KernelSource->ExternalInputs.Num(), InterfaceBinding);  
+				
+						KernelSource->ExternalInputs.Emplace(FuncDef);	
+					}	
+				}
+				
+
+				{
+					const FDataInterfaceFunctionBinding& CopyFromBinding = CopyFromDataInterfaceMap[SourceInstancedPin];
+				
+					TArray<FShaderFunctionDefinition> Functions;
+					CopyFromBinding.DataInterface->GetSupportedInputs(Functions);
+					FShaderFunctionDefinition FuncDef = Functions[CopyFromBinding.FunctionIndex];
+
+					for (FShaderParamTypeDefinition& ParamType : FuncDef.ParamTypes)
+					{
+						ParamType.ResetTypeDeclaration();
+					}
+				
+					FOptimus_InterfaceBinding InterfaceBinding;
+					InterfaceBinding.DataInterface = CopyFromBinding.DataInterface;
+					InterfaceBinding.DataInterfaceBindingIndex = CopyFromBinding.FunctionIndex;
+					InterfaceBinding.BindingFunctionName = FString::Printf(TEXT("Read%s"), *SourceInstancedPin.Pin->GetName());
+					InterfaceBinding.BindingFunctionNamespace = FString();
+		
+					InputDataBindings.Add(KernelSource->ExternalInputs.Num(), InterfaceBinding);  
 		
+					KernelSource->ExternalInputs.Emplace(FuncDef);
+
+					SourceText += FString::Printf(TEXT("%s Value = %s(Index);\n"), *ValueType->ToString(), *InterfaceBinding.BindingFunctionName);
+				}
+				
+				for (int32 TargetIndex = 0 ; TargetIndex < OutputLink.Value.Num(); TargetIndex++)
+				{
+					const FOptimusInstancedPin& TargetInstancedPin = OutputLink.Value[TargetIndex];
+					const FDataInterfaceFunctionBinding& CopyToBinding = CopyToDataInterfaceMap[TargetInstancedPin];
+					
+					TArray<FShaderFunctionDefinition> Functions;
+					CopyToBinding.DataInterface->GetSupportedOutputs(Functions);
+					FShaderFunctionDefinition FuncDef = Functions[CopyToBinding.FunctionIndex];
+
+					for (FShaderParamTypeDefinition& ParamType : FuncDef.ParamTypes)
+					{
+						ParamType.ResetTypeDeclaration();
+					}
+				
+					FOptimus_InterfaceBinding InterfaceBinding;
+					InterfaceBinding.DataInterface = CopyToBinding.DataInterface;
+					InterfaceBinding.DataInterfaceBindingIndex = CopyToBinding.FunctionIndex;
+					InterfaceBinding.BindingFunctionName = FString::Printf(TEXT("Write_%d_%s"), TargetIndex, *TargetInstancedPin.Pin->GetName());
+					InterfaceBinding.BindingFunctionNamespace = FString();
+		
+					OutputDataBindings.Add(KernelSource->ExternalOutputs.Num(), InterfaceBinding);  
+		
+					KernelSource->ExternalOutputs.Emplace(FuncDef);
+
+					SourceText += FString::Printf(TEXT("%s(Index, Value);\n"), *InterfaceBinding.BindingFunctionName);
+				}
+
+				static const FString CopyKernelName = TEXT("CopyKernel");
+				static const FIntVector GroupSize = FIntVector(64, 1, 1);
+				FString CookedSource = Optimus::GetCookedKernelSource(BoundKernel.Kernel->GetPathName(), SourceText, CopyKernelName, GroupSize);
+				KernelSource->SetSource(CookedSource);
+				KernelSource->EntryPoint = CopyKernelName;
+				KernelSource->GroupSize = GroupSize;
+				BoundKernel.Kernel->KernelSource = KernelSource;
+				
+				BoundKernels.Add(BoundKernel);
+				ComputeGraph->KernelInvocations.Add(BoundKernel.Kernel);
+				ComputeGraph->KernelToNode.Add(nullptr);
+			}
+		}
+
 		check(ComputeGraph->KernelInvocations.Num() == BoundKernels.Num());
 
 		// Create the graph edges.
diff --git a/Engine/Plugins/Animation/DeformerGraph/Source/OptimusCore/Private/OptimusHelpers.cpp b/Engine/Plugins/Animation/DeformerGraph/Source/OptimusCore/Private/OptimusHelpers.cpp
--- a/Engine/Plugins/Animation/DeformerGraph/Source/OptimusCore/Private/OptimusHelpers.cpp
+++ b/Engine/Plugins/Animation/DeformerGraph/Source/OptimusCore/Private/OptimusHelpers.cpp
@@ -136,3 +136,62 @@ bool Optimus::ConvertShaderFilePathToObjectPath(FString& InOutPath)
 	InOutPath.RemoveFromEnd(TEXT(".ush"));
 	return true;
 }
+
+FString Optimus::GetCookedKernelSource(
+	const FString& InObjectPathName,
+	const FString& InShaderSource,
+	const FString& InKernelName,
+	FIntVector InGroupSize
+	)
+{
+	// FIXME: Create source range mappings so that we can go from error location to
+	// our source.
+	FString Source = InShaderSource;
+
+#if PLATFORM_WINDOWS
+	// Remove old-school stuff.
+	Source.ReplaceInline(TEXT("\r"), TEXT(""));
+#endif
+
+	FString ShaderPathName = InObjectPathName;
+	Optimus::ConvertObjectPathToShaderFilePath(ShaderPathName);
+
+	const bool bHasKernelKeyword = Source.Contains(TEXT("KERNEL"), ESearchCase::CaseSensitive);
+
+	const FString ComputeShaderUtilsInclude = TEXT("#include \"/Engine/Private/ComputeShaderUtils.ush\"");
+	
+	const FString KernelFunc = FString::Printf(
+		TEXT("[numthreads(%d,%d,%d)]\nvoid %s(uint3 GroupId : SV_GroupID, uint GroupIndex : SV_GroupIndex)"), 
+		InGroupSize.X, InGroupSize.Y, InGroupSize.Z, *InKernelName);
+	
+	const FString UnWrappedDispatchThreadId = FString::Printf(
+	TEXT("GetUnWrappedDispatchThreadId(GroupId, GroupIndex, %d)"),
+		InGroupSize.X * InGroupSize.Y * InGroupSize.Z
+	);
+
+	if (bHasKernelKeyword)
+	{
+		Source.ReplaceInline(TEXT("KERNEL"), TEXT("void __kernel_func(uint Index)"), ESearchCase::CaseSensitive);
+
+		return FString::Printf(
+			TEXT(
+				"#line 1 \"%s\"\n"
+				"%s\n\n"
+				"%s\n\n"
+				"%s { __kernel_func(%s); }\n"
+				), *ShaderPathName, *Source, *ComputeShaderUtilsInclude,*KernelFunc, *UnWrappedDispatchThreadId);
+	}
+	else
+	{
+		return FString::Printf(
+		TEXT(
+			"%s\n"
+			"%s\n"
+			"{\n"
+			"uint Index = %s;\n"
+			"#line 1 \"%s\"\n"
+			"%s\n"
+			"}\n"
+			), *ComputeShaderUtilsInclude,*KernelFunc, *UnWrappedDispatchThreadId, *ShaderPathName, *Source);
+	}
+}
diff --git a/Engine/Plugins/Animation/DeformerGraph/Source/OptimusCore/Public/IOptimusDataInterfaceProvider.h b/Engine/Plugins/Animation/DeformerGraph/Source/OptimusCore/Public/IOptimusDataInterfaceProvider.h
--- a/Engine/Plugins/Animation/DeformerGraph/Source/OptimusCore/Public/IOptimusDataInterfaceProvider.h
+++ b/Engine/Plugins/Animation/DeformerGraph/Source/OptimusCore/Public/IOptimusDataInterfaceProvider.h
@@ -13,7 +13,7 @@ class UOptimusComponentSourceBinding;
 
 
 UINTERFACE()
-class UOptimusDataInterfaceProvider :
+class OPTIMUSCORE_API UOptimusDataInterfaceProvider :
 	public UInterface
 {
 	GENERATED_BODY()
diff --git a/Engine/Plugins/Animation/DeformerGraph/Source/OptimusCore/Public/OptimusHelpers.h b/Engine/Plugins/Animation/DeformerGraph/Source/OptimusCore/Public/OptimusHelpers.h
--- a/Engine/Plugins/Animation/DeformerGraph/Source/OptimusCore/Public/OptimusHelpers.h
+++ b/Engine/Plugins/Animation/DeformerGraph/Source/OptimusCore/Public/OptimusHelpers.h
@@ -100,4 +100,11 @@ namespace Optimus
 	void ConvertObjectPathToShaderFilePath(FString& InOutPath);
 	/** Helpers to convert a virtual shader path back to a UObject path. Returns false if the virtual shader path wasn't recognized as a UObject path. */
 	bool ConvertShaderFilePathToObjectPath(FString& InOutPath);
+
+	FString GetCookedKernelSource(
+		const FString& InObjectPathName,
+		const FString& InShaderSource,
+		const FString& InKernelName,
+		FIntVector InGroupSize
+	);
 }
diff --git a/Engine/Plugins/Animation/DeformerGraph/Source/OptimusEditor/Private/OptimusEditorGraphNode.cpp b/Engine/Plugins/Animation/DeformerGraph/Source/OptimusEditor/Private/OptimusEditorGraphNode.cpp
--- a/Engine/Plugins/Animation/DeformerGraph/Source/OptimusEditor/Private/OptimusEditorGraphNode.cpp
+++ b/Engine/Plugins/Animation/DeformerGraph/Source/OptimusEditor/Private/OptimusEditorGraphNode.cpp
@@ -3,6 +3,7 @@
 #include "OptimusEditorGraphNode.h"
 
 #include "IOptimusComputeKernelProvider.h"
+#include "IOptimusDataInterfaceProvider.h"
 #include "OptimusEditorHelpers.h"
 #include "OptimusEditorGraph.h"
 #include "OptimusEditorGraphSchema.h"
@@ -14,6 +15,8 @@
 #include "OptimusNodePin.h"
 #include "IOptimusNodeAdderPinProvider.h"
 #include "IOptimusPinMutabilityDefiner.h"
+#include "OptimusComponentSource.h"
+#include "OptimusDataTypeRegistry.h"
 #include "OptimusNodeGraph.h"
 
 #include "Framework/Commands/GenericCommands.h"
@@ -320,6 +323,21 @@ FLinearColor UOptimusEditorGraphNode::GetNodeTitleColor() const
 
 		if (ModelNode->GetOwningGraph()->GetGraphType() == EOptimusNodeGraphType::Update)
 		{
+			if (Cast<IOptimusDataInterfaceProvider>(ModelNode))
+			{
+				const FOptimusDataTypeRegistry& TypeRegistry = FOptimusDataTypeRegistry::Get();
+				FOptimusDataTypeHandle ComponentSourceType = TypeRegistry.FindType(*UOptimusComponentSourceBinding::StaticClass());
+				
+				// Any writes to data interface provider need to take place every tick, thus they are considered mutable
+				for (const UOptimusNodePin* Pin : ModelNode->GetPinsByDirection(EOptimusNodePinDirection::Input, true))
+				{
+					if (Pin->GetDataType() != ComponentSourceType)
+					{
+						return Super::GetNodeTitleColor();
+					}
+				}
+			}
+			
 			if (ModelNode->GetOwningGraph()->DoesNodeHaveMutableInput(ModelNode))
 			{
 				return Super::GetNodeTitleColor();
