Optimus: minor bug fixes related to loops

#rnx

diff --git a/Engine/Plugins/Animation/DeformerGraph/Source/OptimusCore/Private/OptimusDeformer.cpp b/Engine/Plugins/Animation/DeformerGraph/Source/OptimusCore/Private/OptimusDeformer.cpp
--- a/Engine/Plugins/Animation/DeformerGraph/Source/OptimusCore/Private/OptimusDeformer.cpp
+++ b/Engine/Plugins/Animation/DeformerGraph/Source/OptimusCore/Private/OptimusDeformer.cpp
@@ -1986,16 +1986,14 @@ TArray<FOptimusComputeGraphInfo> UOptimusDeformer::CompileNodeGraphToComputeGrap
 					else
 					{
 						// Plain connections
-						// 1. Kernel <-> Kernel
+						// 1. Kernel <-> Kernel, validation should make sure the source kernel is not looped or both kernel belong to the same loop
 						// 2. Kernel <-> Data Interface
 						// 3. Kernel -> index/count pin on Loop Terminals 
 						const FOptimusRoutedConstNode SourceRoutedNode = {OtherNode, OtherPins[0].TraversalContext};
-						const int32 SourceLoopIndex = NodeToMaxLoopIndex[SourceRoutedNode];	
-						if (ensure(SourceLoopIndex == 0))
-						{
-							const FOptimusInstancedPin InstancedSourcePin = {FOptimusInstancedNode(SourceRoutedNode, SourceLoopIndex), OtherPin};
-							TargetPinToSourcePin.Add(InstancedTargetPin) = InstancedSourcePin;	
-						}	
+						const int32 SourceLoopIndex = FMath::Clamp(InstancedNode.LoopIndex, 0, NodeToMaxLoopIndex[SourceRoutedNode]);
+						
+						const FOptimusInstancedPin InstancedSourcePin = {FOptimusInstancedNode(SourceRoutedNode, SourceLoopIndex), OtherPin};
+						TargetPinToSourcePin.Add(InstancedTargetPin) = InstancedSourcePin;	
 					}
 				}
 			}
@@ -2352,7 +2350,17 @@ TArray<FOptimusComputeGraphInfo> UOptimusDeformer::CompileNodeGraphToComputeGrap
 		GraphInfo.GraphType = GraphType;
 		GraphInfo.ComputeGraph = NewObject<UOptimusComputeGraph>(this, GraphInfo.GraphName);
 
-		GraphInfos.Add(GraphInfo);
+		if (GraphType != InNodeGraph->GraphType)
+		{
+			// Make sure generated graphs run before the user created graph
+			check(GraphType == EOptimusNodeGraphType::Setup);
+			GraphInfos.Insert(GraphInfo, 0);
+		}
+		else
+		{
+			GraphInfos.Add(GraphInfo);
+		}
+		
 	}
 
 	for (int32 GraphIndex = 0 ; GraphIndex < GraphInfos.Num(); GraphIndex++)
diff --git a/Engine/Plugins/Animation/DeformerGraph/Source/OptimusCore/Private/OptimusNodeGraph.cpp b/Engine/Plugins/Animation/DeformerGraph/Source/OptimusCore/Private/OptimusNodeGraph.cpp
--- a/Engine/Plugins/Animation/DeformerGraph/Source/OptimusCore/Private/OptimusNodeGraph.cpp
+++ b/Engine/Plugins/Animation/DeformerGraph/Source/OptimusCore/Private/OptimusNodeGraph.cpp
@@ -1741,7 +1741,7 @@ TSet<UOptimusComponentSourceBinding*> UOptimusNodeGraph::GetComponentSourceBindi
 
 bool UOptimusNodeGraph::IsPinMutable(const UOptimusNodePin* InNodePin) const
 {
-	TSet<FOptimusRoutedConstNode> VisitedNodes;
+	TSet<FOptimusRoutedConstNodePin> VisitedPins;
 	TQueue<FOptimusRoutedConstNodePin> WorkingSet;
 
 	// If given an input pin, find the other side, since output pin provides mutability definition
@@ -1797,12 +1797,10 @@ bool UOptimusNodeGraph::IsPinMutable(const UOptimusNodePin* InNodePin) const
 				if (ensure(ConnectedPin.NodePin != nullptr))
 				{
 					FOptimusRoutedConstNodePin CollectedNodePin{ConnectedPin.NodePin, ConnectedPin.TraversalContext};
-					const UOptimusNode *NextNode = ConnectedPin.NodePin->GetOwningNode();
-					FOptimusRoutedConstNode CollectedNode{NextNode, ConnectedPin.TraversalContext};
 
-					if (!VisitedNodes.Contains(CollectedNode))
+					if (!VisitedPins.Contains(CollectedNodePin))
 					{
-						VisitedNodes.Add(CollectedNode);
+						VisitedPins.Add(CollectedNodePin);
 						WorkingSet.Enqueue(CollectedNodePin);
 					}
 				}
@@ -1815,12 +1813,10 @@ bool UOptimusNodeGraph::IsPinMutable(const UOptimusNodePin* InNodePin) const
 			{
 				const UOptimusNodePin* NextPin = LoopTerminal->GetPinCounterpart(NodePin, EOptimusTerminalType::Return, EOptimusNodePinDirection::Output);
 				
-				const UOptimusNode*	NextNode = LoopTerminal->GetOtherTerminal();
-				FOptimusRoutedConstNode CollectedNode{NextNode, WorkItem.TraversalContext};
-
-				if (!VisitedNodes.Contains(CollectedNode))
+				FOptimusRoutedConstNodePin CollectedNodePin{NextPin, WorkItem.TraversalContext};	
+				if (!VisitedPins.Contains(CollectedNodePin))
 				{
-					VisitedNodes.Add(CollectedNode);
+					VisitedPins.Add(CollectedNodePin);
 					WorkingSet.Enqueue({NextPin, WorkItem.TraversalContext});
 				}
 			}
