Silence PVS warnings

diff --git a/Engine/Plugins/Animation/DeformerGraph/Source/OptimusCore/Private/Nodes/OptimusNode_ComputeKernelFunction.cpp b/Engine/Plugins/Animation/DeformerGraph/Source/OptimusCore/Private/Nodes/OptimusNode_ComputeKernelFunction.cpp
--- a/Engine/Plugins/Animation/DeformerGraph/Source/OptimusCore/Private/Nodes/OptimusNode_ComputeKernelFunction.cpp
+++ b/Engine/Plugins/Animation/DeformerGraph/Source/OptimusCore/Private/Nodes/OptimusNode_ComputeKernelFunction.cpp
@@ -74,7 +74,7 @@ UClass* UOptimusNode_ComputeKernelFunctionGeneratorClass::CreateNodeClass(
 		NextProperty = &(*NextProperty)->Next;
 	}
 
-	TMap<const FProperty*, const TArray<uint8>*> PropertyValues;
+	//TMap<const FProperty*, const TArray<uint8>*> PropertyValues;
 	for (const FOptimusParameterBinding& InputBinding: InInputBindings)
 	{
 		if (InputBinding.DataDomain.IsSingleton())
@@ -106,7 +106,7 @@ UClass* UOptimusNode_ComputeKernelFunctionGeneratorClass::CreateNodeClass(
 	KernelClass->AddToRoot();
 
 	// Grab the CDO and update the default values based on the raw values in the value bindings.
-	UOptimusNode_ComputeKernelFunction *KernelCDO = Cast<UOptimusNode_ComputeKernelFunction>(KernelClass->GetDefaultObject());
+	/*UOptimusNode_ComputeKernelFunction *KernelCDO = Cast<UOptimusNode_ComputeKernelFunction>(KernelClass->GetDefaultObject());
 
 	// Copy the default values from the incoming properties.
 	for (const TPair<const FProperty*, const TArray<uint8>*>& Item: PropertyValues)
@@ -115,7 +115,7 @@ UClass* UOptimusNode_ComputeKernelFunctionGeneratorClass::CreateNodeClass(
 		const TArray<uint8>& RawValue = *Item.Value;
 
 		Property->CopyCompleteValue(Property->ContainerPtrToValuePtr<UOptimusNode_ComputeKernelFunction>(KernelCDO), RawValue.GetData());
-	}
+	}*/
 	
 	return KernelClass;
 }
diff --git a/Engine/Source/Editor/UnrealEd/Private/SEditorViewportToolBarButton.cpp b/Engine/Source/Editor/UnrealEd/Private/SEditorViewportToolBarButton.cpp
--- a/Engine/Source/Editor/UnrealEd/Private/SEditorViewportToolBarButton.cpp
+++ b/Engine/Source/Editor/UnrealEd/Private/SEditorViewportToolBarButton.cpp
@@ -102,11 +102,11 @@ void SEditorViewportToolBarButton::OnCheckStateChanged( ECheckBoxState NewChecke
 
 const FSlateBrush* SEditorViewportToolBarButton::OnGetButtonImage() const
 {
-	return IsChecked.Get() == true ? CheckedBrush : NormalBrush;
+	return IsChecked.Get() ? CheckedBrush : NormalBrush;
 }
 
 ECheckBoxState SEditorViewportToolBarButton::OnIsChecked() const
 {
-	return IsChecked.Get() == true ? ECheckBoxState::Checked : ECheckBoxState::Unchecked;
+	return IsChecked.Get() ? ECheckBoxState::Checked : ECheckBoxState::Unchecked;
 }
 
diff --git a/Engine/Source/Runtime/Core/Private/UObject/UnrealNames.cpp b/Engine/Source/Runtime/Core/Private/UObject/UnrealNames.cpp
--- a/Engine/Source/Runtime/Core/Private/UObject/UnrealNames.cpp
+++ b/Engine/Source/Runtime/Core/Private/UObject/UnrealNames.cpp
@@ -4934,7 +4934,7 @@ struct FNameBatchAsyncLoader : FNameBatchLoader
 
 	bool ShouldLoadAsync(uint32 MaxWorkers) const
 	{
-		return MaxWorkers && Hashes.Num() >= FNamePoolShards && Hashes.Num() > 30000;
+		return MaxWorkers && Hashes.Num() >= FNamePoolShards && Hashes.Num() > 30000; //-V590
 	}
 
 	void PrepareWork()
diff --git a/Engine/Source/Runtime/Core/Public/Math/Vector.h b/Engine/Source/Runtime/Core/Public/Math/Vector.h
--- a/Engine/Source/Runtime/Core/Public/Math/Vector.h
+++ b/Engine/Source/Runtime/Core/Public/Math/Vector.h
@@ -2332,7 +2332,7 @@ FORCEINLINE bool TVector<T>::InitFromCompactString(const FString& InSourceString
 		return true;
 	}
 
-	const bool bSuccessful = FParse::Value(*InSourceString, TEXT("X="), X) | FParse::Value(*InSourceString, TEXT("Y="), Y) | FParse::Value(*InSourceString, TEXT("Z="), Z);
+	const bool bSuccessful = FParse::Value(*InSourceString, TEXT("X="), X) | FParse::Value(*InSourceString, TEXT("Y="), Y) | FParse::Value(*InSourceString, TEXT("Z="), Z); //-V792
 
 	return bSuccessful;
 }
diff --git a/Engine/Source/Runtime/Engine/Private/SparseVolumeTexture/SparseVolumeTextureStreamingManager.cpp b/Engine/Source/Runtime/Engine/Private/SparseVolumeTexture/SparseVolumeTextureStreamingManager.cpp
--- a/Engine/Source/Runtime/Engine/Private/SparseVolumeTexture/SparseVolumeTextureStreamingManager.cpp
+++ b/Engine/Source/Runtime/Engine/Private/SparseVolumeTexture/SparseVolumeTextureStreamingManager.cpp
@@ -1802,7 +1802,7 @@ void FStreamingManager::IssueRequests(int32 MaxSelectedRequests)
 #if DO_CHECK
 		for (auto& Pending : PendingMipLevels)
 		{
-			check(Pending.SparseVolumeTexture != SelectedKey.SVT || Pending.FrameIndex != SelectedKey.FrameIndex || Pending.MipLevelIndex != SelectedKey.MipLevelIndex);
+			check(Pending.SparseVolumeTexture != SelectedKey.SVT || Pending.FrameIndex != SelectedKey.FrameIndex || Pending.MipLevelIndex != SelectedKey.MipLevelIndex); //-V1013
 		}
 #endif
 
diff --git a/Engine/Source/Runtime/Engine/Private/VT/VirtualTextureTranscodeCache.cpp b/Engine/Source/Runtime/Engine/Private/VT/VirtualTextureTranscodeCache.cpp
--- a/Engine/Source/Runtime/Engine/Private/VT/VirtualTextureTranscodeCache.cpp
+++ b/Engine/Source/Runtime/Engine/Private/VT/VirtualTextureTranscodeCache.cpp
@@ -259,7 +259,7 @@ FVTTranscodeKey FVirtualTextureTranscodeCache::GetKey(const FVirtualTextureProdu
 	Result.vAddress = vAddress;
 	Result.vLevel = vLevel;
 	Result.LayerMask = LayerMask;
-	Result.Hash = (uint16)MurmurFinalize64(Result.Key);
+	Result.Hash = (uint16)MurmurFinalize64(Result.Key); //-V614
 	return Result;
 }
 
diff --git a/Engine/Source/Runtime/Experimental/Iris/Core/Private/Iris/Serialization/PackedVectorNetSerializers.cpp b/Engine/Source/Runtime/Experimental/Iris/Core/Private/Iris/Serialization/PackedVectorNetSerializers.cpp
--- a/Engine/Source/Runtime/Experimental/Iris/Core/Private/Iris/Serialization/PackedVectorNetSerializers.cpp
+++ b/Engine/Source/Runtime/Experimental/Iris/Core/Private/Iris/Serialization/PackedVectorNetSerializers.cpp
@@ -698,12 +698,12 @@ void FPackedVectorNetSerializerBase::Dequantize(uint32 ScaleBitCount, FNetSerial
 
 			if constexpr (std::is_same<SourceType, FVector3f>::value)
 			{
-				memcpy(&Target.X, &Components[0], 3U*sizeof(float));
+				memcpy(&Target.X, &Components[0], 3U*sizeof(float)); //-V512
 			}
 			else
 			{
 				FVector3f Vector;
-				memcpy(&Vector.X, &Components[0], 3U*sizeof(float));
+				memcpy(&Vector.X, &Components[0], 3U*sizeof(float)); //-V512
 				Target = SourceType(Vector);
 			}
 		}
diff --git a/Engine/Source/Runtime/GeometryCore/Public/Sampling/GridSampler.h b/Engine/Source/Runtime/GeometryCore/Public/Sampling/GridSampler.h
--- a/Engine/Source/Runtime/GeometryCore/Public/Sampling/GridSampler.h
+++ b/Engine/Source/Runtime/GeometryCore/Public/Sampling/GridSampler.h
@@ -41,8 +41,8 @@ public:
 	UE::Math::TVector2<RealType> Sample(int32 Index) const
 	{
 		checkSlow(Index >= 0 && Index < Num());
-		RealType X = (static_cast<RealType>(Index % Size) + RealType{ 1 } / RealType{ 2 }) / static_cast<RealType>(Size);
-		RealType Y = (static_cast<RealType>(Index / Size) + RealType{ 1 } / RealType{ 2 }) / static_cast<RealType>(Size);
+		RealType X = (static_cast<RealType>(Index % Size) + RealType{ 1 } / RealType{ 2 }) / static_cast<RealType>(Size); //-V1064
+		RealType Y = (static_cast<RealType>(Index / Size) + RealType{ 1 } / RealType{ 2 }) / static_cast<RealType>(Size); //-V1064
 		return UE::Math::TVector2<RealType>(X, Y);
 	}
 };
