Disable mesh deformers on mobile since they're not supported.

diff --git a/Engine/Plugins/Animation/DeformerGraph/DeformerGraph.uplugin b/Engine/Plugins/Animation/DeformerGraph/DeformerGraph.uplugin
--- a/Engine/Plugins/Animation/DeformerGraph/DeformerGraph.uplugin
+++ b/Engine/Plugins/Animation/DeformerGraph/DeformerGraph.uplugin
@@ -18,17 +18,29 @@
 		{
 			"Name": "OptimusSettings",
 			"Type": "Runtime",
-			"LoadingPhase": "PostConfigInit"
+			"LoadingPhase": "PostConfigInit",
+			"PlatformDenyList": [
+				"Android",
+				"IOS"
+			]
 		},
 		{
 			"Name": "OptimusCore",
 			"Type": "Runtime",
-			"LoadingPhase": "PreDefault"
+			"LoadingPhase": "PreDefault",
+			"PlatformDenyList": [
+				"Android",
+				"IOS"
+			]
 		},
 		{
 			"Name": "OptimusDeveloper",
 			"Type": "UncookedOnly",
-			"LoadingPhase": "PreDefault"
+			"LoadingPhase": "PreDefault",
+			"PlatformDenyList": [
+				"Android",
+				"IOS"
+			]
 		},
 		{
 			"Name": "OptimusEditor",
