[Deformer Graph & SkeletalMesh] added "Build Half Edge Buffers" checkbox to Skeletal mesh such that half edge data can be cooked per lod to avoid run time generation of the half edge data in deformer graph.
+also improved half edge DI in deformer graph to handle Lod change better

diff --git a/Engine/Plugins/Animation/DeformerGraph/Shaders/Private/DataInterfaceHalfEdge.ush b/Engine/Plugins/Animation/DeformerGraph/Shaders/Private/DataInterfaceHalfEdge.ush
--- a/Engine/Plugins/Animation/DeformerGraph/Shaders/Private/DataInterfaceHalfEdge.ush
+++ b/Engine/Plugins/Animation/DeformerGraph/Shaders/Private/DataInterfaceHalfEdge.ush
@@ -7,6 +7,10 @@ uint {DataInterfaceName}_InputStreamStart;
 StructuredBuffer<int> {DataInterfaceName}_VertexToEdgeBuffer;
 StructuredBuffer<int> {DataInterfaceName}_EdgeToTwinEdgeBuffer;
 
+uint {DataInterfaceName}_bUseBufferFromRenderData;
+Buffer<int> {DataInterfaceName}_RenderDataVertexToEdgeBuffer;
+Buffer<int> {DataInterfaceName}_RenderDataEdgeToTwinEdgeBuffer;
+
 uint ReadNumVertices_{DataInterfaceName}()
 {
 	return {DataInterfaceName}_NumVertices;
@@ -19,11 +23,16 @@ uint ReadNumTriangles_{DataInterfaceName}()
 
 int ReadEdge_{DataInterfaceName}(uint VertexIndex)
 {
-	return {DataInterfaceName}_VertexToEdgeBuffer[VertexIndex + {DataInterfaceName}_InputStreamStart] - {DataInterfaceName}_IndexBufferStart;
+	return {DataInterfaceName}_bUseBufferFromRenderData == 0 ?
+		{DataInterfaceName}_VertexToEdgeBuffer[VertexIndex + {DataInterfaceName}_InputStreamStart] - {DataInterfaceName}_IndexBufferStart :
+		{DataInterfaceName}_RenderDataVertexToEdgeBuffer[VertexIndex + {DataInterfaceName}_InputStreamStart] - {DataInterfaceName}_IndexBufferStart;
 }
 
 int ReadTwinEdge_{DataInterfaceName}(uint EdgeIndex)
 {
-	int TwinEdge = {DataInterfaceName}_EdgeToTwinEdgeBuffer[EdgeIndex + {DataInterfaceName}_IndexBufferStart];
+	int TwinEdge = {DataInterfaceName}_bUseBufferFromRenderData == 0 ?
+		{DataInterfaceName}_EdgeToTwinEdgeBuffer[EdgeIndex + {DataInterfaceName}_IndexBufferStart] :
+		{DataInterfaceName}_RenderDataEdgeToTwinEdgeBuffer[EdgeIndex + {DataInterfaceName}_IndexBufferStart];
+	
 	return TwinEdge == -1 ? -1 : TwinEdge - {DataInterfaceName}_IndexBufferStart;
 }
diff --git a/Engine/Plugins/Animation/DeformerGraph/Source/OptimusCore/Private/DataInterfaces/OptimusDataInterfaceHalfEdge.cpp b/Engine/Plugins/Animation/DeformerGraph/Source/OptimusCore/Private/DataInterfaces/OptimusDataInterfaceHalfEdge.cpp
--- a/Engine/Plugins/Animation/DeformerGraph/Source/OptimusCore/Private/DataInterfaces/OptimusDataInterfaceHalfEdge.cpp
+++ b/Engine/Plugins/Animation/DeformerGraph/Source/OptimusCore/Private/DataInterfaces/OptimusDataInterfaceHalfEdge.cpp
@@ -14,276 +14,54 @@
 #include "SkeletalRenderPublic.h"
 
 #include UE_INLINE_GENERATED_CPP_BY_NAME(OptimusDataInterfaceHalfEdge)
-
-struct FEdgeKey
-{
-	FVector3f V0;
-	FVector3f V1;
-
-	bool operator==(FEdgeKey const& Rhs) const { return V0 == Rhs.V0 && V1 == Rhs.V1; }
-};
-
-static uint32 GetTypeHash(FEdgeKey const& Key)
-{
-	return HashCombine(GetTypeHash(Key.V0), GetTypeHash(Key.V1));
-}
-
-namespace
-{
-	/** Builds lookup from edge to twin edge. */
-	class FTwinEdgeBuilder
-	{
-	public:
-		FTwinEdgeBuilder(int32 InNumTriangles)
-		{
-			TwinEdges.SetNum(InNumTriangles * 3);
-			for (int32& TwinEdge : TwinEdges)
-			{
-				TwinEdge = -1;
-			}
-		}
-
-		int32 SmallestVector(FVector3f V0, FVector3f V1)
-		{
-			if (V0.X != V1.X) return V1.X > V0.X;
-			if (V0.Y != V1.Y) return V1.Y > V0.Y;
-			return V1.Z > V0.Z;
-		}
-
-		void Add(int32 InTriangleIndex, TStaticArray<FVector3f, 3>& InVertexPositions)
-		{
-			for (uint32 TriangleEdgeIndex = 0; TriangleEdgeIndex < 3; ++TriangleEdgeIndex)
-			{
-				FVector3f V[2] = { 
-					InVertexPositions[TriangleEdgeIndex],
-					InVertexPositions[(TriangleEdgeIndex + 1) % 3] };
-				int32 VMin = SmallestVector(V[0], V[1]);
-				const FEdgeKey Key({V[VMin], V[1 - VMin]});
-
-				if (FEdgeDescription* EdgeDesc = EdgeMap.Find(Key))
-				{
-					if (EdgeDesc->E1 == -1)
-					{
-						// The second time that we've seen this edge.
-						EdgeDesc->E1 = InTriangleIndex * 3 + TriangleEdgeIndex;
-						TwinEdges[EdgeDesc->E0] = EdgeDesc->E1;
-						TwinEdges[EdgeDesc->E1] = EdgeDesc->E0;
-					}
-					else
-					{
-						// Non-manifold geometry can end up here. Two options: 
-						// (i) Don't store twins for this edge the first pair.
-						// (ii) Reset the edge map entry and hope that good twins are met consecutively when walking the index buffer.
- 						EdgeDesc->E0 = InTriangleIndex * 3 + TriangleEdgeIndex;
- 						EdgeDesc->E1 = -1;
-					}
-				}
-				else
-				{
-					// First time that we've seen this edge.
-					EdgeMap.Add(Key).E0 = InTriangleIndex * 3 + TriangleEdgeIndex;
-				}
-			}
-		}
-
-		TArray<int32>& GetTwinEdges()
-		{
-			return TwinEdges;
-		}
-
-	private:
-		struct FEdgeDescription
-		{
-			int32 E0 = -1;
-			int32 E1 = -1;
-		};
-
-		TMap<FEdgeKey, FEdgeDescription> EdgeMap;
-
-		TArray<int32> TwinEdges;
-	};
-
-	/** Builds lookup from (one) vertex to (many) edges. */
-	class FVertexToEdgeBuilder
-	{
-	public:
-		FVertexToEdgeBuilder(int32 InNumVertices, int32 InNumTriangles)
-		{
-			VertexToEdgeItems.Reserve(InNumTriangles * 3);
-			VertexToEdgeItems.AddDefaulted(InNumVertices);
-		}
-
-		void Add(int32 InVertexIndex, int32 InEdgeIndex)
-		{
-			int32 ItemIndex = InVertexIndex;
-			if (VertexToEdgeItems[ItemIndex].EdgeIndex == -1)
-			{
-				VertexToEdgeItems[ItemIndex].EdgeIndex = InEdgeIndex;
-				return;
-			}
-			while (VertexToEdgeItems[ItemIndex].NextItem != -1)
-			{
-				ItemIndex = VertexToEdgeItems[ItemIndex].NextItem;
-			}
-			int32 NewItemIndex = VertexToEdgeItems.AddDefaulted();
-			VertexToEdgeItems[ItemIndex].NextItem = NewItemIndex;
-			VertexToEdgeItems[NewItemIndex].EdgeIndex = InEdgeIndex;
-		}
-
-		using FEdgeInlineArray = TArray<int32, TInlineAllocator<16>>;
-		void Get(int32 InVertexIndex, FEdgeInlineArray& OutEdgeIndices) const
-		{
-			int32 ItemIndex = InVertexIndex;
-			while (ItemIndex != -1)
-			{
-				OutEdgeIndices.Add(VertexToEdgeItems[ItemIndex].EdgeIndex);
-				ItemIndex = VertexToEdgeItems[ItemIndex].NextItem;
-			}
-		}
-
-	private:
-		struct FVertexToEdgeItem
-		{
-			int32 EdgeIndex = -1;
-			int32 NextItem = -1;
-		};
-
-		/** First NumVertices entries are allocated up front as start of per vertex linked lists. */
-		TArray<FVertexToEdgeItem> VertexToEdgeItems;
-	};
-
-	static int32 GetHeadEdgeForVertex(int32 InVertexIndex, FVertexToEdgeBuilder const& InVertexToEdge, TArray<int32> const& InEdgeToTwinEdge)
-	{
-		FVertexToEdgeBuilder::FEdgeInlineArray EdgeIndicesScratchBuffer;
-		InVertexToEdge.Get(InVertexIndex, EdgeIndicesScratchBuffer);
-
-		// Iterate backwards around the edges until we find a border, or a non recognized edge (which may lead to a border).
-		const int32 StartEdgeIndex = EdgeIndicesScratchBuffer[0];
-		int32 EdgeIndex = StartEdgeIndex;
-		while(1)
-		{
-			const int32 LastEdgeIndex = EdgeIndex;
-			EdgeIndex = ((EdgeIndex / 3) * 3) + ((EdgeIndex + 2) % 3);
-			EdgeIndex = InEdgeToTwinEdge[EdgeIndex];
-			
-			if (EdgeIndex == StartEdgeIndex)
-			{
-				break;
-			}
-			if (EdgeIndex == -1 || EdgeIndicesScratchBuffer.Find(EdgeIndex) == INDEX_NONE)
-			{
-				EdgeIndex = LastEdgeIndex;
-				break;
-			}
-		}
-
-		return EdgeIndex;
-	}
-
-	static void BuildHalfEdgeBuffers(const FSkeletalMeshLODRenderData& InLodRenderData, TArray<int32>& OutVertexToEdge, TArray<int32>& OutEdgeToTwinEdge)
-	{
-		const FRawStaticIndexBuffer16or32Interface* IndexBuffer = InLodRenderData.MultiSizeIndexContainer.GetIndexBuffer();
-		const int32 IndexCount = IndexBuffer->Num();
-		const int32 TriangleCount = IndexCount / 3;
-
-		const FPositionVertexBuffer& VertexBuffer = InLodRenderData.StaticVertexBuffers.PositionVertexBuffer;
-		const int32 VertexCount = InLodRenderData.StaticVertexBuffers.PositionVertexBuffer.GetNumVertices();
-
-		// Build edge to twin edge map.
-		{
-			FTwinEdgeBuilder TwinEdgeBuilder(TriangleCount);
-			for (int32 TriangleIndex = 0; TriangleIndex < TriangleCount; ++TriangleIndex)
-			{
-				TStaticArray<FVector3f, 3> VertexPositions;
-				VertexPositions[0] = VertexBuffer.VertexPosition(IndexBuffer->Get(TriangleIndex * 3 + 0));
-				VertexPositions[1] = VertexBuffer.VertexPosition(IndexBuffer->Get(TriangleIndex * 3 + 1));
-				VertexPositions[2] = VertexBuffer.VertexPosition(IndexBuffer->Get(TriangleIndex * 3 + 2));
-				TwinEdgeBuilder.Add(TriangleIndex, VertexPositions);
-			}
-			OutEdgeToTwinEdge = MoveTemp(TwinEdgeBuilder.GetTwinEdges());
-		}
-
-		// Build vertex to edge index map.
-		FVertexToEdgeBuilder VertexToEdgeBuilder(VertexCount, TriangleCount);
-		for (int32 Index = 0; Index < TriangleCount * 3; ++Index)
-		{
-			const int32 VertexIndex = IndexBuffer->Get(Index);
-			const int32 EdgeIndex = Index;
-			VertexToEdgeBuilder.Add(VertexIndex, EdgeIndex);
-		}
-
-		// VertexToEdgeBuilder gives us multiple edges per vertex.
-		// These should form a linked list of edges. 
-		// Usually that list will be a loop, but at border geometry there will be a head and end to the list.
-		// For loops it doesn't matter where we start iteration. But at borders we want to use the head of the linked list to start iteration.
-		// Non-manifold geometry will be strange, and potentially have multiple edge independent loops. We just do the best we can in that case.
-		OutVertexToEdge.SetNumUninitialized(VertexCount);
-		for (int32 VertexIndex = 0; VertexIndex < VertexCount; ++VertexIndex)
-		{
-			OutVertexToEdge[VertexIndex] = GetHeadEdgeForVertex(VertexIndex, VertexToEdgeBuilder, OutEdgeToTwinEdge);
-		}
-	}
-}
-
 /*
  * Render resource containing the half edge buffers. 
- * todo[CF]: We need to move this to the skeletal mesh cooked mesh data to avoid the perf and memory hit everytime we instantiate a deformer.
  */ 
-class FHalfEdgeBufferResources : public FRenderResource
+class FHalfEdgeBuffers
 {
 public:
-	FHalfEdgeBufferResources(TArray< TArray<int32> >&& InVertexToEdgeData, TArray< TArray<int32> >&& InEdgeToTwinEdgeData)
-	{
-		VertexToEdgeData = MoveTemp(InVertexToEdgeData);
-		EdgeToTwinEdgeData = MoveTemp(InEdgeToTwinEdgeData);
-	}
+	FHalfEdgeBuffers() = default;
 
-	void InitRHI(FRHICommandListBase&) override
+	void Init(FRDGBuilder& GraphBuilder, const FSkeletalMeshLODRenderData& InLodRenderData)
 	{
-		FRDGBuilder GraphBuilder(FRHICommandListImmediate::Get());
+		TResourceArray<int32> VertexToEdgeData;
+		TResourceArray<int32> EdgeToTwinEdgeData;
 
-		const int32 NumLods = VertexToEdgeData.Num();
-		VertexToEdgeBuffers.SetNum(NumLods);
-		EdgeToTwinEdgeBuffers.SetNum(NumLods);
+		SkeletalMeshHalfEdgeUtility::BuildHalfEdgeBuffers(InLodRenderData, VertexToEdgeData, EdgeToTwinEdgeData);	
 		
-		for (int32 LodIndex = 0; LodIndex < NumLods; ++LodIndex)
-		{
-			FRDGBuffer* VertexToEdgeBuffer = GraphBuilder.CreateBuffer(FRDGBufferDesc::CreateStructuredDesc(sizeof(uint32), VertexToEdgeData[LodIndex].Num()), TEXT("Optimus.VertexToEdge"));
-			GraphBuilder.QueueBufferUpload(VertexToEdgeBuffer, VertexToEdgeData[LodIndex].GetData(), VertexToEdgeData[LodIndex].Num() * sizeof(int32));
-			VertexToEdgeBuffers[LodIndex] = GraphBuilder.ConvertToExternalBuffer(VertexToEdgeBuffer);
+		FRDGBuffer* VertexToEdgeBuffer = GraphBuilder.CreateBuffer(FRDGBufferDesc::CreateStructuredDesc(sizeof(uint32), VertexToEdgeData.Num()), TEXT("Optimus.VertexToEdge"));
+		GraphBuilder.QueueBufferUpload(VertexToEdgeBuffer, VertexToEdgeData.GetData(), VertexToEdgeData.Num() * sizeof(int32));
+		VertexToEdgePooledBuffer = GraphBuilder.ConvertToExternalBuffer(VertexToEdgeBuffer);
 			
-			FRDGBuffer* EdgeToTwinEdgeBuffer = GraphBuilder.CreateBuffer(FRDGBufferDesc::CreateStructuredDesc(sizeof(uint32), EdgeToTwinEdgeData[LodIndex].Num()), TEXT("Optimus.EdgeToTwinEdge"));
-			GraphBuilder.QueueBufferUpload(EdgeToTwinEdgeBuffer, EdgeToTwinEdgeData[LodIndex].GetData(), EdgeToTwinEdgeData[LodIndex].Num() * sizeof(int32));
-			EdgeToTwinEdgeBuffers[LodIndex] = GraphBuilder.ConvertToExternalBuffer(EdgeToTwinEdgeBuffer);
-		}
+		FRDGBuffer* EdgeToTwinEdgeBuffer = GraphBuilder.CreateBuffer(FRDGBufferDesc::CreateStructuredDesc(sizeof(uint32), EdgeToTwinEdgeData.Num()), TEXT("Optimus.EdgeToTwinEdge"));
+		GraphBuilder.QueueBufferUpload(EdgeToTwinEdgeBuffer, EdgeToTwinEdgeData.GetData(), EdgeToTwinEdgeData.Num() * sizeof(int32));
+		EdgeToTwinEdgePooledBuffer = GraphBuilder.ConvertToExternalBuffer(EdgeToTwinEdgeBuffer);
 
-		GraphBuilder.Execute();
-
-		VertexToEdgeData.Reset();
-		EdgeToTwinEdgeData.Reset();
+		bInitialized = true;
 	}
 
-	void ReleaseRHI() override
+	bool IsInitialized() const { return bInitialized; }
+	
+	~FHalfEdgeBuffers()
 	{
-		VertexToEdgeBuffers.Reset();
-		EdgeToTwinEdgeBuffers.Reset();
+		VertexToEdgePooledBuffer.SafeRelease();
+		EdgeToTwinEdgePooledBuffer.SafeRelease();
 	}
 
-	TRefCountPtr<FRDGPooledBuffer> GetVertexToEdgeBuffer(int32 InLodIndex)
+	TRefCountPtr<FRDGPooledBuffer> GetVertexToEdgeBuffer()
 	{
-		return VertexToEdgeBuffers[InLodIndex];
+		return VertexToEdgePooledBuffer;
 	}
 
-	TRefCountPtr<FRDGPooledBuffer> GetEdgeToTwinBuffer(int32 InLodIndex)
+	TRefCountPtr<FRDGPooledBuffer> GetEdgeToTwinBuffer()
 	{
-		return EdgeToTwinEdgeBuffers[InLodIndex];
+		return EdgeToTwinEdgePooledBuffer;
 	}
 
 private:
-	TArray< TArray<int32> > VertexToEdgeData;
-	TArray< TArray<int32> > EdgeToTwinEdgeData;
-	TArray< TRefCountPtr<FRDGPooledBuffer> > VertexToEdgeBuffers;
-	TArray< TRefCountPtr<FRDGPooledBuffer> > EdgeToTwinEdgeBuffers;
+	bool bInitialized = false;
+	TRefCountPtr<FRDGPooledBuffer> VertexToEdgePooledBuffer;
+	TRefCountPtr<FRDGPooledBuffer> EdgeToTwinEdgePooledBuffer;
 };
 
 
@@ -336,6 +114,10 @@ BEGIN_SHADER_PARAMETER_STRUCT(FHalfEdgeDataInterfaceParameters, )
 	SHADER_PARAMETER(uint32, InputStreamStart)
 	SHADER_PARAMETER_RDG_BUFFER_SRV(Buffer<int>, VertexToEdgeBuffer)
 	SHADER_PARAMETER_RDG_BUFFER_SRV(Buffer<int>, EdgeToTwinEdgeBuffer)
+
+	SHADER_PARAMETER(uint32, bUseBufferFromRenderData)
+	SHADER_PARAMETER_SRV(Buffer<int>, RenderDataVertexToEdgeBuffer)
+	SHADER_PARAMETER_SRV(Buffer<int>, RenderDataEdgeToTwinEdgeBuffer)
 END_SHADER_PARAMETER_STRUCT()
 
 void UOptimusHalfEdgeDataInterface::GetShaderParameters(TCHAR const* UID, FShaderParametersMetadataBuilder& InOutBuilder, FShaderParametersMetadataAllocations& InOutAllocations) const
@@ -378,22 +160,8 @@ UComputeDataProvider* UOptimusHalfEdgeDataInterface::CreateDataProvider(TObjectP
 		FSkeletalMeshRenderData const* SkeletalMeshRenderData = Provider->SkinnedMesh->GetSkeletalMeshRenderData();
 		if (SkeletalMeshRenderData != nullptr)
 		{
-			// Build half edge data.
-			TArray< TArray<int32> > VertexToEdgePerLod;
-			VertexToEdgePerLod.SetNum(SkeletalMeshRenderData->NumInlinedLODs);
-			TArray< TArray<int32> > EdgeToTwinEdgePerLod;
-			EdgeToTwinEdgePerLod.SetNum(SkeletalMeshRenderData->NumInlinedLODs);
-
-			for (int32 LodIndex = 0; LodIndex < SkeletalMeshRenderData->NumInlinedLODs; ++LodIndex)
-			{
-				FSkeletalMeshLODRenderData const* LodRenderData = &SkeletalMeshRenderData->LODRenderData[LodIndex];
-				BuildHalfEdgeBuffers(*LodRenderData, VertexToEdgePerLod[LodIndex], EdgeToTwinEdgePerLod[LodIndex]);
-			}
-
-			// Upload data to a resource stored on the provider.
-			Provider->HalfEdgeBuffers = MakeShared<FHalfEdgeBufferResources>(MoveTemp(VertexToEdgePerLod), MoveTemp(EdgeToTwinEdgePerLod));
-
-			BeginInitResource(Provider->HalfEdgeBuffers.Get());
+			// Items in this array are initialized on demand by the render proxy
+			Provider->OnDemandHalfEdgeBuffers.SetNum(SkeletalMeshRenderData->LODRenderData.Num());
 		}
 	}
 
@@ -404,10 +172,7 @@ UComputeDataProvider* UOptimusHalfEdgeDataInterface::CreateDataProvider(TObjectP
 void UOptimusHalfEdgeDataProvider::BeginDestroy()
 {
 	Super::BeginDestroy();
-	if (HalfEdgeBuffers)
-	{
-		BeginReleaseResource(HalfEdgeBuffers.Get());
-	}
+	OnDemandHalfEdgeBuffers.Reset();
 	DestroyFence.BeginFence();
 }
 
@@ -418,15 +183,15 @@ bool UOptimusHalfEdgeDataProvider::IsReadyForFinishDestroy()
 
 FComputeDataProviderRenderProxy* UOptimusHalfEdgeDataProvider::GetRenderProxy()
 {
-	return new FOptimusHalfEdgeDataProviderProxy(SkinnedMesh, HalfEdgeBuffers);
+	return new FOptimusHalfEdgeDataProviderProxy(SkinnedMesh, OnDemandHalfEdgeBuffers);
 }
 
 
 FOptimusHalfEdgeDataProviderProxy::FOptimusHalfEdgeDataProviderProxy(
 	USkinnedMeshComponent* InSkinnedMeshComponent, 
-	TSharedPtr<FHalfEdgeBufferResources>& InHalfEdgeBuffers)
+	TArray<FHalfEdgeBuffers>& InOnDemandHalfEdgeBuffers)
 	: SkeletalMeshObject(InSkinnedMeshComponent != nullptr ? InSkinnedMeshComponent->MeshObject : nullptr)
-	, HalfEdgeBuffers(InHalfEdgeBuffers)
+	, OnDemandHalfEdgeBuffers(InOnDemandHalfEdgeBuffers)
 {
 }
 
@@ -436,7 +201,11 @@ bool FOptimusHalfEdgeDataProviderProxy::IsValid(FValidationData const& InValidat
 	{
 		return false;
 	}
-	if (SkeletalMeshObject == nullptr || !HalfEdgeBuffers.IsValid())
+	if (SkeletalMeshObject == nullptr)
+	{
+		return false;
+	}
+	if (!ensure(SkeletalMeshObject->GetSkeletalMeshRenderData().LODRenderData.Num() == OnDemandHalfEdgeBuffers.Num()))
 	{
 		return false;
 	}
@@ -450,13 +219,35 @@ bool FOptimusHalfEdgeDataProviderProxy::IsValid(FValidationData const& InValidat
 
 void FOptimusHalfEdgeDataProviderProxy::AllocateResources(FRDGBuilder& GraphBuilder)
 {
+	FallbackSRV = GraphBuilder.CreateSRV(GraphBuilder.RegisterExternalBuffer(GWhiteVertexBufferWithRDG->Buffer), PF_R16G16B16A16_UINT);
+	
 	const int32 LodIndex = SkeletalMeshObject->GetLOD();
+	FSkeletalMeshRenderData const& SkeletalMeshRenderData = SkeletalMeshObject->GetSkeletalMeshRenderData();
+	FSkeletalMeshLODRenderData const* LodRenderData = &SkeletalMeshRenderData.LODRenderData[LodIndex];
 
-	FRDGBuffer* VertexToEdgeBuffer = GraphBuilder.RegisterExternalBuffer(HalfEdgeBuffers->GetVertexToEdgeBuffer(LodIndex));
-	VertexToEdgeBufferSRV = GraphBuilder.CreateSRV(VertexToEdgeBuffer);
+	if (LodRenderData->HalfEdgeBuffer.IsInitialized())
+	{
+		bUseBufferFromRenderData = true;
+
+		VertexToEdgeBufferSRV = FallbackSRV;
+		EdgeToTwinEdgeBufferSRV = FallbackSRV;
+	}
+	else
+	{
+		bUseBufferFromRenderData = false;
+		
+		if (!OnDemandHalfEdgeBuffers[LodIndex].IsInitialized())
+		{
+			// fall back to compute half edge on demand
+			OnDemandHalfEdgeBuffers[LodIndex].Init(GraphBuilder, *LodRenderData);
+		}
+
+		FRDGBuffer* VertexToEdgeBuffer = GraphBuilder.RegisterExternalBuffer(OnDemandHalfEdgeBuffers[LodIndex].GetVertexToEdgeBuffer());
+		VertexToEdgeBufferSRV = GraphBuilder.CreateSRV(VertexToEdgeBuffer);
 	
-	FRDGBuffer* EdgeToTwinEdgeBuffer = GraphBuilder.RegisterExternalBuffer(HalfEdgeBuffers->GetEdgeToTwinBuffer(LodIndex));
-	EdgeToTwinEdgeBufferSRV = GraphBuilder.CreateSRV(EdgeToTwinEdgeBuffer);
+		FRDGBuffer* EdgeToTwinEdgeBuffer = GraphBuilder.RegisterExternalBuffer(OnDemandHalfEdgeBuffers[LodIndex].GetEdgeToTwinBuffer());
+		EdgeToTwinEdgeBufferSRV = GraphBuilder.CreateSRV(EdgeToTwinEdgeBuffer);	
+	}
 }
 
 void FOptimusHalfEdgeDataProviderProxy::GatherDispatchData(FDispatchData const& InDispatchData)
@@ -465,6 +256,9 @@ void FOptimusHalfEdgeDataProviderProxy::GatherDispatchData(FDispatchData const&
 	FSkeletalMeshRenderData const& SkeletalMeshRenderData = SkeletalMeshObject->GetSkeletalMeshRenderData();
 	FSkeletalMeshLODRenderData const* LodRenderData = &SkeletalMeshRenderData.LODRenderData[LodIndex];
 
+	FRHIShaderResourceView* RenderDataVertexToEdgeBufferSRV = LodRenderData->HalfEdgeBuffer.GetVertexToEdgeBufferSRV();
+	FRHIShaderResourceView* RenderDataEdgeToTwinEdgeBufferSRV = LodRenderData->HalfEdgeBuffer.GetEdgeToTwinEdgeBufferSRV();
+	
 	FRHIShaderResourceView* NullSRVBinding = GWhiteVertexBufferWithSRV->ShaderResourceViewRHI.GetReference();
 
 	const TStridedView<FParameters> ParameterArray = MakeStridedParameterView<FParameters>(InDispatchData);
@@ -479,5 +273,9 @@ void FOptimusHalfEdgeDataProviderProxy::GatherDispatchData(FDispatchData const&
 		Parameters.InputStreamStart = InDispatchData.bUnifiedDispatch ? 0 : RenderSection.BaseVertexIndex;
 		Parameters.VertexToEdgeBuffer = VertexToEdgeBufferSRV;
 		Parameters.EdgeToTwinEdgeBuffer = EdgeToTwinEdgeBufferSRV;
+
+		Parameters.bUseBufferFromRenderData = bUseBufferFromRenderData;
+		Parameters.RenderDataVertexToEdgeBuffer = RenderDataVertexToEdgeBufferSRV != nullptr ? RenderDataVertexToEdgeBufferSRV : NullSRVBinding;
+		Parameters.RenderDataEdgeToTwinEdgeBuffer = RenderDataEdgeToTwinEdgeBufferSRV != nullptr ? RenderDataEdgeToTwinEdgeBufferSRV : NullSRVBinding;
 	}
 }
diff --git a/Engine/Plugins/Animation/DeformerGraph/Source/OptimusCore/Private/DataInterfaces/OptimusDataInterfaceHalfEdge.h b/Engine/Plugins/Animation/DeformerGraph/Source/OptimusCore/Private/DataInterfaces/OptimusDataInterfaceHalfEdge.h
--- a/Engine/Plugins/Animation/DeformerGraph/Source/OptimusCore/Private/DataInterfaces/OptimusDataInterfaceHalfEdge.h
+++ b/Engine/Plugins/Animation/DeformerGraph/Source/OptimusCore/Private/DataInterfaces/OptimusDataInterfaceHalfEdge.h
@@ -2,13 +2,14 @@
 
 #pragma once
 
+#include "RenderGraphDefinitions.h"
 #include "OptimusComputeDataInterface.h"
 #include "ComputeFramework/ComputeDataProvider.h"
 #include "RenderCommandFence.h"
 #include "OptimusDataInterfaceHalfEdge.generated.h"
 
 class FHalfEdgeDataInterfaceParameters;
-class FHalfEdgeBufferResources;
+class FHalfEdgeBuffers;
 class FRDGBuffer;
 class FRDGBufferSRV;
 class FSkeletalMeshObject;
@@ -62,7 +63,7 @@ public:
 	FComputeDataProviderRenderProxy* GetRenderProxy() override;
 	//~ End UComputeDataProvider Interface
 
-	TSharedPtr<FHalfEdgeBufferResources> HalfEdgeBuffers;
+	TArray<FHalfEdgeBuffers> OnDemandHalfEdgeBuffers;
 	FRenderCommandFence DestroyFence;
 };
 
@@ -71,7 +72,7 @@ class FOptimusHalfEdgeDataProviderProxy : public FComputeDataProviderRenderProxy
 public:
 	FOptimusHalfEdgeDataProviderProxy(
 		USkinnedMeshComponent* InSkinnedMeshComponent, 
-		TSharedPtr<FHalfEdgeBufferResources>& InHalfEdgeBuffers);
+		TArray<FHalfEdgeBuffers>& InOnDemandHalfEdgeBuffers);
 
 	//~ Begin FComputeDataProviderRenderProxy Interface
 	bool IsValid(FValidationData const& InValidationData) const override;
@@ -83,8 +84,12 @@ private:
 	using FParameters = FHalfEdgeDataInterfaceParameters;
 
 	FSkeletalMeshObject* SkeletalMeshObject = nullptr;
-	TSharedPtr<FHalfEdgeBufferResources> HalfEdgeBuffers;
-
+	
+	TArray<FHalfEdgeBuffers>& OnDemandHalfEdgeBuffers;
+	
+	bool bUseBufferFromRenderData = false;
 	FRDGBufferSRV* VertexToEdgeBufferSRV = nullptr;
 	FRDGBufferSRV* EdgeToTwinEdgeBufferSRV = nullptr;
+	
+	FRDGBufferSRVRef FallbackSRV = nullptr;
 };
diff --git a/Engine/Source/Runtime/Core/Private/UObject/FortniteMainBranchObjectVersion.cpp b/Engine/Source/Runtime/Core/Private/UObject/FortniteMainBranchObjectVersion.cpp
--- a/Engine/Source/Runtime/Core/Private/UObject/FortniteMainBranchObjectVersion.cpp
+++ b/Engine/Source/Runtime/Core/Private/UObject/FortniteMainBranchObjectVersion.cpp
@@ -12,7 +12,7 @@ TMap<FGuid, FGuid> FFortniteMainBranchObjectVersion::GetSystemGuids()
 	SystemGuids.Add(DevGuids.NANITE_DERIVEDDATA_VER, FGuid("D28C003F825F4C178F698938FB0E47F7"));
 	SystemGuids.Add(DevGuids.NIAGARASHADERMAP_DERIVEDDATA_VER, FGuid("D0DF49D408F34E38A98FA097E0280EC7"));
 	SystemGuids.Add(DevGuids.Niagara_LatestScriptCompileVersion, FGuid("CE33E29C6268469E8E18F1D96A55DFDA"));
-	SystemGuids.Add(DevGuids.SkeletalMeshDerivedDataVersion, FGuid("41EFB33F1B72B1F49F9AA7D295871F36"));
+	SystemGuids.Add(DevGuids.SkeletalMeshDerivedDataVersion, FGuid("C9CF008376FC4F7D9996CB449719502B"));
 	SystemGuids.Add(DevGuids.STATICMESH_DERIVEDDATA_VER, FGuid("543843B103794E0E9BA4BE80FB602F79"));
 	SystemGuids.Add(DevGuids.MaterialTranslationDDCVersion, FGuid("BBD650804D314CE7A44ADF9BB4F3F5F9"));
 	return SystemGuids;
diff --git a/Engine/Source/Runtime/Core/Public/UObject/FortniteMainBranchObjectVersion.h b/Engine/Source/Runtime/Core/Public/UObject/FortniteMainBranchObjectVersion.h
--- a/Engine/Source/Runtime/Core/Public/UObject/FortniteMainBranchObjectVersion.h
+++ b/Engine/Source/Runtime/Core/Public/UObject/FortniteMainBranchObjectVersion.h
@@ -413,6 +413,9 @@ struct FFortniteMainBranchObjectVersion
 
 		// Skeletal Mesh uses Mesh Description to store mesh bulk data.
 		MeshDescriptionForSkeletalMesh,
+
+		// Skeletal Mesh optionally cooks half edge data per lod
+		SkeletalHalfEdgeData,
 		
 		// -----<new versions can be added above this line>-------------------------------------------------
 		VersionPlusOne,
diff --git a/Engine/Source/Runtime/Engine/Classes/Engine/SkinnedAssetCommon.h b/Engine/Source/Runtime/Engine/Classes/Engine/SkinnedAssetCommon.h
--- a/Engine/Source/Runtime/Engine/Classes/Engine/SkinnedAssetCommon.h
+++ b/Engine/Source/Runtime/Engine/Classes/Engine/SkinnedAssetCommon.h
@@ -183,7 +183,14 @@ struct FSkeletalMeshLODInfo
 	/** Keeps this LODs data on the CPU so it can be used for things such as sampling in FX. */
 	UPROPERTY(EditAnywhere, Category = SkeletalMeshLODInfo)
 	uint8 bAllowCPUAccess : 1;
-
+	
+	/**
+	 * If true, we will cache/cook half edge data that provides vertex connectivity information across material sections, which
+ 	 * may be useful for other systems like Mesh Deformer.
+	 */
+	UPROPERTY(EditAnywhere, Category = SkeletalMeshLODInfo)
+	uint8 bBuildHalfEdgeBuffers : 1;
+	
 	/** List of vertex attributes to include for rendering and what type they should be */
 	UPROPERTY(EditAnywhere, Category = SkeletalMeshLODInfo, AdvancedDisplay, EditFixedSize, Meta=(NoResetToDefault))
 	TArray<FSkeletalMeshVertexAttributeInfo> VertexAttributes;	
@@ -220,6 +227,7 @@ struct FSkeletalMeshLODInfo
 		, bHasBeenSimplified(false)
 		, bHasPerLODVertexColors(false)
 		, bAllowCPUAccess(false)
+		, bBuildHalfEdgeBuffers(false)
 		, bSupportUniformlyDistributedSampling(false)
 #if WITH_EDITORONLY_DATA
 		, bImportWithBaseMesh(false)
diff --git a/Engine/Source/Runtime/Engine/Private/Rendering/SkeletalMeshHalfEdgeBuffer.cpp b/Engine/Source/Runtime/Engine/Private/Rendering/SkeletalMeshHalfEdgeBuffer.cpp
new file mode 100644
--- /dev/null
+++ b/Engine/Source/Runtime/Engine/Private/Rendering/SkeletalMeshHalfEdgeBuffer.cpp
@@ -0,0 +1,291 @@
+﻿// Copyright Epic Games, Inc. All Rights Reserved.
+
+#include "Rendering/SkeletalMeshHalfEdgeBuffer.h"
+#include "Rendering/SkeletalMeshLODRenderData.h"
+
+struct FEdgeKey
+{
+	FVector3f V0;
+	FVector3f V1;
+
+	bool operator==(FEdgeKey const& Rhs) const { return V0 == Rhs.V0 && V1 == Rhs.V1; }
+};
+
+static uint32 GetTypeHash(FEdgeKey const& Key)
+{
+	return HashCombine(GetTypeHash(Key.V0), GetTypeHash(Key.V1));
+}
+
+namespace
+{
+	/** Builds lookup from edge to twin edge. */
+	class FTwinEdgeBuilder
+	{
+	public:
+		FTwinEdgeBuilder(int32 InNumTriangles)
+		{
+			TwinEdges.SetNum(InNumTriangles * 3);
+			for (int32& TwinEdge : TwinEdges)
+			{
+				TwinEdge = -1;
+			}
+		}
+
+		int32 SmallestVector(FVector3f V0, FVector3f V1)
+		{
+			if (V0.X != V1.X) return V1.X > V0.X;
+			if (V0.Y != V1.Y) return V1.Y > V0.Y;
+			return V1.Z > V0.Z;
+		}
+
+		void Add(int32 InTriangleIndex, TStaticArray<FVector3f, 3>& InVertexPositions)
+		{
+			for (uint32 TriangleEdgeIndex = 0; TriangleEdgeIndex < 3; ++TriangleEdgeIndex)
+			{
+				FVector3f V[2] = { 
+					InVertexPositions[TriangleEdgeIndex],
+					InVertexPositions[(TriangleEdgeIndex + 1) % 3] };
+				int32 VMin = SmallestVector(V[0], V[1]);
+				const FEdgeKey Key({V[VMin], V[1 - VMin]});
+
+				if (FEdgeDescription* EdgeDesc = EdgeMap.Find(Key))
+				{
+					if (EdgeDesc->E1 == -1)
+					{
+						// The second time that we've seen this edge.
+						EdgeDesc->E1 = InTriangleIndex * 3 + TriangleEdgeIndex;
+						TwinEdges[EdgeDesc->E0] = EdgeDesc->E1;
+						TwinEdges[EdgeDesc->E1] = EdgeDesc->E0;
+					}
+					else
+					{
+						// Non-manifold geometry can end up here. Two options: 
+						// (i) Don't store twins for this edge the first pair.
+						// (ii) Reset the edge map entry and hope that good twins are met consecutively when walking the index buffer.
+ 						EdgeDesc->E0 = InTriangleIndex * 3 + TriangleEdgeIndex;
+ 						EdgeDesc->E1 = -1;
+					}
+				}
+				else
+				{
+					// First time that we've seen this edge.
+					EdgeMap.Add(Key).E0 = InTriangleIndex * 3 + TriangleEdgeIndex;
+				}
+			}
+		}
+
+		TResourceArray<int32>& GetTwinEdges()
+		{
+			return TwinEdges;
+		}
+
+	private:
+		struct FEdgeDescription
+		{
+			int32 E0 = -1;
+			int32 E1 = -1;
+		};
+
+		TMap<FEdgeKey, FEdgeDescription> EdgeMap;
+
+		TResourceArray<int32> TwinEdges;
+	};
+
+	/** Builds lookup from (one) vertex to (many) edges. */
+	class FVertexToEdgeBuilder
+	{
+	public:
+		FVertexToEdgeBuilder(int32 InNumVertices, int32 InNumTriangles)
+		{
+			VertexToEdgeItems.Reserve(InNumTriangles * 3);
+			VertexToEdgeItems.AddDefaulted(InNumVertices);
+		}
+
+		void Add(int32 InVertexIndex, int32 InEdgeIndex)
+		{
+			int32 ItemIndex = InVertexIndex;
+			if (VertexToEdgeItems[ItemIndex].EdgeIndex == -1)
+			{
+				VertexToEdgeItems[ItemIndex].EdgeIndex = InEdgeIndex;
+				return;
+			}
+			while (VertexToEdgeItems[ItemIndex].NextItem != -1)
+			{
+				ItemIndex = VertexToEdgeItems[ItemIndex].NextItem;
+			}
+			int32 NewItemIndex = VertexToEdgeItems.AddDefaulted();
+			VertexToEdgeItems[ItemIndex].NextItem = NewItemIndex;
+			VertexToEdgeItems[NewItemIndex].EdgeIndex = InEdgeIndex;
+		}
+
+		using FEdgeInlineArray = TArray<int32, TInlineAllocator<16>>;
+		void Get(int32 InVertexIndex, FEdgeInlineArray& OutEdgeIndices) const
+		{
+			int32 ItemIndex = InVertexIndex;
+			while (ItemIndex != -1)
+			{
+				OutEdgeIndices.Add(VertexToEdgeItems[ItemIndex].EdgeIndex);
+				ItemIndex = VertexToEdgeItems[ItemIndex].NextItem;
+			}
+		}
+
+	private:
+		struct FVertexToEdgeItem
+		{
+			int32 EdgeIndex = -1;
+			int32 NextItem = -1;
+		};
+
+		/** First NumVertices entries are allocated up front as start of per vertex linked lists. */
+		TArray<FVertexToEdgeItem> VertexToEdgeItems;
+	};
+
+	static int32 GetHeadEdgeForVertex(int32 InVertexIndex, FVertexToEdgeBuilder const& InVertexToEdge, TResourceArray<int32> const& InEdgeToTwinEdge)
+	{
+		FVertexToEdgeBuilder::FEdgeInlineArray EdgeIndicesScratchBuffer;
+		InVertexToEdge.Get(InVertexIndex, EdgeIndicesScratchBuffer);
+
+		// Iterate backwards around the edges until we find a border, or a non recognized edge (which may lead to a border).
+		const int32 StartEdgeIndex = EdgeIndicesScratchBuffer[0];
+		int32 EdgeIndex = StartEdgeIndex;
+		while(1)
+		{
+			const int32 LastEdgeIndex = EdgeIndex;
+			EdgeIndex = ((EdgeIndex / 3) * 3) + ((EdgeIndex + 2) % 3);
+			EdgeIndex = InEdgeToTwinEdge[EdgeIndex];
+			
+			if (EdgeIndex == StartEdgeIndex)
+			{
+				break;
+			}
+			if (EdgeIndex == -1 || EdgeIndicesScratchBuffer.Find(EdgeIndex) == INDEX_NONE)
+			{
+				EdgeIndex = LastEdgeIndex;
+				break;
+			}
+		}
+
+		return EdgeIndex;
+	}
+}
+
+void SkeletalMeshHalfEdgeUtility::BuildHalfEdgeBuffers(const FSkeletalMeshLODRenderData& InLodRenderData, TResourceArray<int32>& OutVertexToEdge, TResourceArray<int32>& OutEdgeToTwinEdge)
+{
+	const FRawStaticIndexBuffer16or32Interface* IndexBuffer = InLodRenderData.MultiSizeIndexContainer.GetIndexBuffer();
+	const int32 IndexCount = IndexBuffer->Num();
+	const int32 TriangleCount = IndexCount / 3;
+
+	const FPositionVertexBuffer& VertexBuffer = InLodRenderData.StaticVertexBuffers.PositionVertexBuffer;
+	const int32 VertexCount = InLodRenderData.StaticVertexBuffers.PositionVertexBuffer.GetNumVertices();
+
+	// Build edge to twin edge map.
+	{
+		FTwinEdgeBuilder TwinEdgeBuilder(TriangleCount);
+		for (int32 TriangleIndex = 0; TriangleIndex < TriangleCount; ++TriangleIndex)
+		{
+			TStaticArray<FVector3f, 3> VertexPositions;
+			VertexPositions[0] = VertexBuffer.VertexPosition(IndexBuffer->Get(TriangleIndex * 3 + 0));
+			VertexPositions[1] = VertexBuffer.VertexPosition(IndexBuffer->Get(TriangleIndex * 3 + 1));
+			VertexPositions[2] = VertexBuffer.VertexPosition(IndexBuffer->Get(TriangleIndex * 3 + 2));
+			TwinEdgeBuilder.Add(TriangleIndex, VertexPositions);
+		}
+		OutEdgeToTwinEdge = MoveTemp(TwinEdgeBuilder.GetTwinEdges());
+	}
+
+	// Build vertex to edge index map.
+	FVertexToEdgeBuilder VertexToEdgeBuilder(VertexCount, TriangleCount);
+	for (int32 Index = 0; Index < TriangleCount * 3; ++Index)
+	{
+		const int32 VertexIndex = IndexBuffer->Get(Index);
+		const int32 EdgeIndex = Index;
+		VertexToEdgeBuilder.Add(VertexIndex, EdgeIndex);
+	}
+
+	// VertexToEdgeBuilder gives us multiple edges per vertex.
+	// These should form a linked list of edges. 
+	// Usually that list will be a loop, but at border geometry there will be a head and end to the list.
+	// For loops it doesn't matter where we start iteration. But at borders we want to use the head of the linked list to start iteration.
+	// Non-manifold geometry will be strange, and potentially have multiple edge independent loops. We just do the best we can in that case.
+	OutVertexToEdge.SetNumUninitialized(VertexCount);
+	for (int32 VertexIndex = 0; VertexIndex < VertexCount; ++VertexIndex)
+	{
+		OutVertexToEdge[VertexIndex] = GetHeadEdgeForVertex(VertexIndex, VertexToEdgeBuilder, OutEdgeToTwinEdge);
+	}	
+}
+
+void FSkeletalMeshHalfEdgeBuffer::Init(const FSkeletalMeshLODRenderData& InLodRenderData)
+{
+	check(VertexToEdgeData.Num() == 0);
+	SkeletalMeshHalfEdgeUtility::BuildHalfEdgeBuffers(InLodRenderData, VertexToEdgeData, EdgeToTwinEdgeData);	
+}
+
+void FSkeletalMeshHalfEdgeBuffer::InitRHI(FRHICommandListBase& RHICmdList)
+{
+	const static FLazyName ClassName(TEXT("FSkeletalMeshHalfEdgeBuffer"));
+
+	// InitRHI only called if CPU data is valid
+	check(VertexToEdgeData.Num() > 0 && EdgeToTwinEdgeData.Num() > 0)
+	{
+		const uint32 BufferSize = VertexToEdgeData.GetResourceDataSize();
+		FRHIResourceCreateInfo CreateInfo(TEXT("VertexToEdgeData"));
+		CreateInfo.ClassName = ClassName;
+		CreateInfo.OwnerName = GetOwnerName();
+		VertexToEdgeBufferRHI = RHICmdList.CreateStructuredBuffer(sizeof(int32), BufferSize, BUF_Static | BUF_ShaderResource, ERHIAccess::SRVMask, CreateInfo);
+		VertexToEdgeBufferRHI->SetOwnerName(GetOwnerName());
+
+		void* BufferPtr = RHICmdList.LockBuffer(VertexToEdgeBufferRHI, 0, BufferSize, RLM_WriteOnly);
+		FMemory::ParallelMemcpy(BufferPtr, VertexToEdgeData.GetData(), BufferSize, EMemcpyCachePolicy::StoreUncached);
+		RHICmdList.UnlockBuffer(VertexToEdgeBufferRHI);
+		VertexToEdgeBufferSRV = RHICmdList.CreateShaderResourceView(VertexToEdgeBufferRHI);
+	}
+
+	{
+		const uint32 BufferSize = EdgeToTwinEdgeData.GetResourceDataSize();
+		FRHIResourceCreateInfo CreateInfo(TEXT("EdgeToTwinEdgeData"));
+		CreateInfo.ClassName = ClassName;
+		CreateInfo.OwnerName = GetOwnerName();
+		EdgeToTwinEdgeBufferRHI = RHICmdList.CreateStructuredBuffer(sizeof(int32), BufferSize, BUF_Static | BUF_ShaderResource, ERHIAccess::SRVMask, CreateInfo);
+		EdgeToTwinEdgeBufferRHI->SetOwnerName(GetOwnerName());
+
+		void* BufferPtr = RHICmdList.LockBuffer(EdgeToTwinEdgeBufferRHI, 0, BufferSize, RLM_WriteOnly);
+		FMemory::ParallelMemcpy(BufferPtr, EdgeToTwinEdgeData.GetData(), BufferSize, EMemcpyCachePolicy::StoreUncached);
+		RHICmdList.UnlockBuffer(EdgeToTwinEdgeBufferRHI);
+		EdgeToTwinEdgeBufferSRV = RHICmdList.CreateShaderResourceView(EdgeToTwinEdgeBufferRHI);
+	}
+}
+
+void FSkeletalMeshHalfEdgeBuffer::ReleaseRHI()
+{
+	VertexToEdgeBufferRHI.SafeRelease();
+	VertexToEdgeBufferSRV.SafeRelease();
+	EdgeToTwinEdgeBufferRHI.SafeRelease();
+	EdgeToTwinEdgeBufferSRV.SafeRelease();
+}
+
+bool FSkeletalMeshHalfEdgeBuffer::IsCPUDataValid() const
+{
+	return VertexToEdgeData.Num() > 0 && EdgeToTwinEdgeData.Num() > 0;
+}
+
+void FSkeletalMeshHalfEdgeBuffer::CleanUp()
+{
+	VertexToEdgeData.Reset();
+	EdgeToTwinEdgeData.Reset();	
+}
+
+int32 FSkeletalMeshHalfEdgeBuffer::GetResourceSize() const
+{
+	return VertexToEdgeData.GetResourceDataSize() + EdgeToTwinEdgeData.GetResourceDataSize();
+}
+
+void FSkeletalMeshHalfEdgeBuffer::Serialize(FArchive& Ar)
+{
+	Ar << VertexToEdgeData;
+	Ar << EdgeToTwinEdgeData;
+}
+
+FArchive& operator<<(FArchive& Ar, FSkeletalMeshHalfEdgeBuffer& HalfEdgeBuffer)
+{
+	HalfEdgeBuffer.Serialize(Ar);
+	return Ar;
+}
diff --git a/Engine/Source/Runtime/Engine/Private/SkeletalMeshLODRenderData.cpp b/Engine/Source/Runtime/Engine/Private/SkeletalMeshLODRenderData.cpp
--- a/Engine/Source/Runtime/Engine/Private/SkeletalMeshLODRenderData.cpp
+++ b/Engine/Source/Runtime/Engine/Private/SkeletalMeshLODRenderData.cpp
@@ -85,6 +85,12 @@ namespace
 	};
 }
 
+static bool IsMeshDeformerAvailable(EShaderPlatform InPlatform)
+{
+	static IMeshDeformerProvider* MeshDeformerProvider = IMeshDeformerProvider::Get();
+	return MeshDeformerProvider && MeshDeformerProvider->IsSupported(InPlatform);
+}
+
 /** 
  * This function returns true if the duplicate vertices should be cooked. 
  * The data is used to deal with seams along split vertices when recomputing normals at runtime.
@@ -94,8 +100,7 @@ static bool RequiresDuplicateVerticesInCook(EShaderPlatform InPlatform, FSkelMes
 	// DuplicatedVertices are cooked if we have GPUSkinCache or MeshDeformer systems for this platform.
 	// We always cook when GPUSkinCache is available so that we can support runtime switch of r.SkinCache.RecomputeTangents.
 	// For MeshDeformer we only cook if the section was marked for bRecomputeTangent.
-	static IMeshDeformerProvider* MeshDeformerProvider = IMeshDeformerProvider::Get();
-	bool bMeshDeformersAvailable = MeshDeformerProvider && MeshDeformerProvider->IsSupported(InPlatform);
+	const bool bMeshDeformersAvailable = IsMeshDeformerAvailable(InPlatform);
 	return (RenderSection.bRecomputeTangent && bMeshDeformersAvailable) || IsGPUSkinCacheAvailable(InPlatform);
 }
 
@@ -121,8 +126,7 @@ static bool RequiresDuplicateVerticesInCook(const ITargetPlatform& InTargetPlatf
 static bool RequiresDuplicateVertices()
 {
 	// Never drop at runtime if the data was cooked for MeshDeformers.
-	static IMeshDeformerProvider* MeshDeformerProvider = IMeshDeformerProvider::Get();
-	if (MeshDeformerProvider && MeshDeformerProvider->IsSupported(GMaxRHIShaderPlatform))
+	if (IsMeshDeformerAvailable(GMaxRHIShaderPlatform))
 	{
 		return true;
 	}
@@ -131,6 +135,22 @@ static bool RequiresDuplicateVertices()
 	return GPUSkinCacheNeedsDuplicatedVertices();
 }
 
+static bool IsMeshDeformerSupportedByCookTargetPlatform(const ITargetPlatform& InTargetPlatform)
+{
+	TArray<FName> TargetedShaderFormats;
+	InTargetPlatform.GetAllTargetedShaderFormats(TargetedShaderFormats);
+	for (int32 FormatIndex = 0; FormatIndex < TargetedShaderFormats.Num(); ++FormatIndex)
+	{
+		const EShaderPlatform LegacyShaderPlatform = ShaderFormatToLegacyShaderPlatform(TargetedShaderFormats[FormatIndex]);
+		
+		if (IsMeshDeformerAvailable(LegacyShaderPlatform))
+		{
+			return true;
+		}
+		
+	}
+	return false;
+}
 
 // Serialization.
 FArchive& operator<<(FArchive& Ar, FSkelMeshRenderSection& S)
@@ -283,6 +303,12 @@ void FSkeletalMeshLODRenderData::InitResources(bool bNeedsVertexColors, int32 LO
 	}
 
 	VertexAttributeBuffers.InitResources();
+	
+	if (HalfEdgeBuffer.IsCPUDataValid())
+	{
+		HalfEdgeBuffer.SetOwnerName(OwnerName);
+		BeginInitResource(&HalfEdgeBuffer, &UE::RenderCommandPipe::SkeletalMesh);	
+	}
 
 #if RHI_RAYTRACING
 	if (IsRayTracingAllowed())
@@ -317,6 +343,9 @@ void FSkeletalMeshLODRenderData::ReleaseResources()
 	SkinWeightProfilesData.ReleaseResources();
 
 	VertexAttributeBuffers.ReleaseResources();
+	
+	BeginReleaseResource(&HalfEdgeBuffer, &UE::RenderCommandPipe::SkeletalMesh);
+	
 #if RHI_RAYTRACING
 	if (IsRayTracingAllowed())
 	{
@@ -360,14 +389,16 @@ void FSkeletalMeshLODRenderData::DecrementMemoryStats()
 void FSkeletalMeshLODRenderData::BuildFromLODModel(
 	const FSkeletalMeshLODModel* InLODModel,
 	TConstArrayView<FSkeletalMeshVertexAttributeInfo> InVertexAttributeInfos,
-	ESkeletalMeshVertexFlags InBuildFlags
+	const FBuildSettings& InBuildSettings
 	)
 {
-	const bool bUseFullPrecisionUVs = EnumHasAllFlags(InBuildFlags, ESkeletalMeshVertexFlags::UseFullPrecisionUVs);
-	const bool bUseHighPrecisionTangentBasis = EnumHasAllFlags(InBuildFlags, ESkeletalMeshVertexFlags::UseHighPrecisionTangentBasis);
-	const bool bHasVertexColors = EnumHasAllFlags(InBuildFlags, ESkeletalMeshVertexFlags::HasVertexColors);
-	const bool bUseBackwardsCompatibleF16TruncUVs = EnumHasAllFlags(InBuildFlags, ESkeletalMeshVertexFlags::UseBackwardsCompatibleF16TruncUVs);
-	const bool bUseHighPrecisionWeights = EnumHasAllFlags(InBuildFlags, ESkeletalMeshVertexFlags::UseHighPrecisionWeights);
+	const ESkeletalMeshVertexFlags& BuildFlags = InBuildSettings.BuildFlags;
+	
+	const bool bUseFullPrecisionUVs = EnumHasAllFlags(BuildFlags, ESkeletalMeshVertexFlags::UseFullPrecisionUVs);
+	const bool bUseHighPrecisionTangentBasis = EnumHasAllFlags(BuildFlags, ESkeletalMeshVertexFlags::UseHighPrecisionTangentBasis);
+	const bool bHasVertexColors = EnumHasAllFlags(BuildFlags, ESkeletalMeshVertexFlags::HasVertexColors);
+	const bool bUseBackwardsCompatibleF16TruncUVs = EnumHasAllFlags(BuildFlags, ESkeletalMeshVertexFlags::UseBackwardsCompatibleF16TruncUVs);
+	const bool bUseHighPrecisionWeights = EnumHasAllFlags(BuildFlags, ESkeletalMeshVertexFlags::UseHighPrecisionWeights);
 
 	// Copy required info from source sections
 	RenderSections.Empty();
@@ -502,6 +533,11 @@ void FSkeletalMeshLODRenderData::BuildFromLODModel(
 		}
 	}
 
+	if (InBuildSettings.bBuildHalfEdgeBuffers)
+	{
+		HalfEdgeBuffer.Init(*this);
+	}
+	
 	ActiveBoneIndices = InLODModel->ActiveBoneIndices;
 	RequiredBones = InLODModel->RequiredBones;
 }
@@ -527,6 +563,7 @@ void FSkeletalMeshLODRenderData::ReleaseCPUResources(bool bForStreaming)
 			StaticVertexBuffers.ColorVertexBuffer.CleanUp();
 			SkinWeightProfilesData.ReleaseCPUResources();
 			VertexAttributeBuffers.CleanUp();
+			HalfEdgeBuffer.CleanUp();
 		}
 	}
 }
@@ -550,6 +587,8 @@ void FSkeletalMeshLODRenderData::GetResourceSizeEx(FResourceSizeEx& CumulativeRe
 	CumulativeResourceSize.AddUnknownMemoryBytes(TEXT("ClothVertexBuffer"), ClothVertexBuffer.GetVertexDataSize());
 	CumulativeResourceSize.AddUnknownMemoryBytes(TEXT("SkinWeightProfilesData"), SkinWeightProfilesData.GetResourcesSize());
 	CumulativeResourceSize.AddUnknownMemoryBytes(TEXT("VertexAttributeData"), VertexAttributeBuffers.GetResourceSize());
+	CumulativeResourceSize.AddUnknownMemoryBytes(TEXT("HalfEdgeBuffer"), HalfEdgeBuffer.GetResourceSize());
+	
 }
 
 SIZE_T FSkeletalMeshLODRenderData::GetCPUAccessMemoryOverhead() const
@@ -703,7 +742,8 @@ void FSkeletalMeshLODRenderData::SerializeStreamedData(FArchive& Ar, USkinnedAss
 {
 	Ar.UsingCustomVersion(FUE5PrivateFrostyStreamObjectVersion::GUID);
 	Ar.UsingCustomVersion(FUE5ReleaseStreamObjectVersion::GUID);
-
+	Ar.UsingCustomVersion(FFortniteMainBranchObjectVersion::GUID);
+	
 	FStripDataFlags StripFlags(Ar, ClassDataStripFlags);
 
 	// TODO: A lot of data in a render section is needed during initialization but maybe some can still be streamed
@@ -825,6 +865,23 @@ void FSkeletalMeshLODRenderData::SerializeStreamedData(FArchive& Ar, USkinnedAss
 	{
 		Ar << VertexAttributeBuffers;
 	}
+	
+	if (Ar.CustomVer(FFortniteMainBranchObjectVersion::GUID) >= FFortniteMainBranchObjectVersion::SkeletalHalfEdgeData)
+	{
+		uint8 MeshDeformerStripFlag = 0;
+		if (Ar.IsCooking() && !IsMeshDeformerSupportedByCookTargetPlatform(*Ar.CookingTarget()))
+		{
+			MeshDeformerStripFlag = 1;
+		}
+
+		FStripDataFlags MeshDeformerStripFlags(Ar, MeshDeformerStripFlag);
+		
+		if (!MeshDeformerStripFlags.IsClassDataStripped(MeshDeformerStripFlag))
+		{
+			Ar << HalfEdgeBuffer;
+		}
+	}
+	
 }
 
 void FSkeletalMeshLODRenderData::SerializeAvailabilityInfo(FArchive& Ar, int32 LODIdx, bool bAdjacencyDataStripped, bool bNeedsCPUAccess)
diff --git a/Engine/Source/Runtime/Engine/Private/SkeletalMeshRenderData.cpp b/Engine/Source/Runtime/Engine/Private/SkeletalMeshRenderData.cpp
--- a/Engine/Source/Runtime/Engine/Private/SkeletalMeshRenderData.cpp
+++ b/Engine/Source/Runtime/Engine/Private/SkeletalMeshRenderData.cpp
@@ -354,7 +354,11 @@ void FSkeletalMeshRenderData::Cache(const ITargetPlatform* TargetPlatform, USkin
 						VertexBufferBuildFlags |= ESkeletalMeshVertexFlags::UseHighPrecisionWeights;
 					}
 				}
-				LODData->BuildFromLODModel(LODModel, LODInfo->VertexAttributes, VertexBufferBuildFlags);
+				FSkeletalMeshLODRenderData::FBuildSettings BuildSettings;
+				BuildSettings.BuildFlags = VertexBufferBuildFlags;
+				BuildSettings.bBuildHalfEdgeBuffers = LODInfo->bBuildHalfEdgeBuffers;
+				
+				LODData->BuildFromLODModel(LODModel, LODInfo->VertexAttributes, BuildSettings);
 			}
 
 			FLargeMemoryWriter Ar(0, /*bIsPersistent=*/ true);
diff --git a/Engine/Source/Runtime/Engine/Private/SkinnedAssetCommon.cpp b/Engine/Source/Runtime/Engine/Private/SkinnedAssetCommon.cpp
--- a/Engine/Source/Runtime/Engine/Private/SkinnedAssetCommon.cpp
+++ b/Engine/Source/Runtime/Engine/Private/SkinnedAssetCommon.cpp
@@ -107,6 +107,7 @@ FGuid FSkeletalMeshLODInfo::ComputeDeriveDataCacheKey(const FSkeletalMeshLODGrou
 	}
 	
 	FArchive_Serialize_BitfieldBool(Ar, bAllowCPUAccess);
+	FArchive_Serialize_BitfieldBool(Ar, bBuildHalfEdgeBuffers);
 	FArchive_Serialize_BitfieldBool(Ar, bSupportUniformlyDistributedSampling);
 
 	//Use the LOD settings asset if there is one
diff --git a/Engine/Source/Runtime/Engine/Public/Rendering/SkeletalMeshHalfEdgeBuffer.h b/Engine/Source/Runtime/Engine/Public/Rendering/SkeletalMeshHalfEdgeBuffer.h
new file mode 100644
--- /dev/null
+++ b/Engine/Source/Runtime/Engine/Public/Rendering/SkeletalMeshHalfEdgeBuffer.h
@@ -0,0 +1,57 @@
+// Copyright Epic Games, Inc. All Rights Reserved.
+
+#pragma once
+
+#include "RenderResource.h"
+#include "Containers/DynamicRHIResourceArray.h"
+
+class FSkeletalMeshLODRenderData;
+class FRHIShaderResourceView;
+
+namespace SkeletalMeshHalfEdgeUtility
+{
+	void ENGINE_API BuildHalfEdgeBuffers(const FSkeletalMeshLODRenderData& InLodRenderData, TResourceArray<int32>& OutVertexToEdge, TResourceArray<int32>& OutEdgeToTwinEdge);
+}
+
+/*
+ * Render resource containing the half edge buffers. 
+ */ 
+class FSkeletalMeshHalfEdgeBuffer : public FRenderResource
+{
+public:
+	void Init(const FSkeletalMeshLODRenderData& InLodRenderData);
+		
+	void ENGINE_API InitRHI(FRHICommandListBase& RHICmdList) override;
+
+	void ENGINE_API ReleaseRHI() override;
+
+	bool IsCPUDataValid() const;
+	
+	void CleanUp();
+
+	int32 GetResourceSize() const;
+	
+	void Serialize(FArchive& Ar);
+	
+	friend class FSkeletalMeshLODRenderData;
+	friend FArchive& operator<<(FArchive& Ar, FSkeletalMeshHalfEdgeBuffer& MorphTargetVertexInfoBuffers);
+
+	FRHIShaderResourceView* GetVertexToEdgeBufferSRV() const
+	{
+		return VertexToEdgeBufferSRV;
+	}
+
+	FRHIShaderResourceView* GetEdgeToTwinEdgeBufferSRV() const
+	{
+		return EdgeToTwinEdgeBufferSRV;
+	}
+
+private:
+	TResourceArray<int32> VertexToEdgeData;
+	TResourceArray<int32> EdgeToTwinEdgeData;
+	
+	FBufferRHIRef VertexToEdgeBufferRHI;
+	FShaderResourceViewRHIRef VertexToEdgeBufferSRV;		
+	FBufferRHIRef EdgeToTwinEdgeBufferRHI;
+	FShaderResourceViewRHIRef EdgeToTwinEdgeBufferSRV;	
+};
\ No newline at end of file
diff --git a/Engine/Source/Runtime/Engine/Public/Rendering/SkeletalMeshLODRenderData.h b/Engine/Source/Runtime/Engine/Public/Rendering/SkeletalMeshLODRenderData.h
--- a/Engine/Source/Runtime/Engine/Public/Rendering/SkeletalMeshLODRenderData.h
+++ b/Engine/Source/Runtime/Engine/Public/Rendering/SkeletalMeshLODRenderData.h
@@ -10,6 +10,7 @@
 #include "Rendering/SkeletalMeshDuplicatedVerticesBuffer.h"
 #include "Rendering/SkeletalMeshVertexClothBuffer.h"
 #include "Rendering/MorphTargetVertexInfoBuffers.h"
+#include "Rendering/SkeletalMeshHalfEdgeBuffer.h"
 #include "SkeletalMeshTypes.h"
 #include "BoneIndices.h"
 #include "StaticMeshResources.h"
@@ -145,6 +146,9 @@ public:
 
 	FSkeletalMeshVertexAttributeRenderData VertexAttributeBuffers;
 
+	/** GPU buffer for half edge data of an LOD, useful for mesh deformers */
+	FSkeletalMeshHalfEdgeBuffer HalfEdgeBuffer;
+
 	TArray<FBoneIndexType> ActiveBoneIndices;
 
 	TArray<FBoneIndexType> RequiredBones;
@@ -231,16 +235,23 @@ public:
 	void SerializeAvailabilityInfo(FArchive& Ar, int32 LODIdx, bool bHasAdjacencyData, bool bNeedsCPUAccess);
 
 #if WITH_EDITOR
+
+	struct ENGINE_API FBuildSettings
+	{
+		ESkeletalMeshVertexFlags BuildFlags = ESkeletalMeshVertexFlags::None;
+		bool bBuildHalfEdgeBuffers = false;
+	};
+	
 	/**
 	 * Initialize render data (e.g. vertex buffers) from model info
 	 * @param InLODModel The model to build the render data from.
 	 * @param InVertexAttributeInfos The vertex attributes to possibly include and their stored data type.
-	 * @param InBuildFlags See ESkeletalMeshVertexFlags.
+	 * @param InBuildSettings Forwards relevant settings from LodInfo, See ESkeletalMeshVertexFlags.
 	 */
 	void ENGINE_API BuildFromLODModel(
 		const FSkeletalMeshLODModel* InLODModel,
 		TConstArrayView<FSkeletalMeshVertexAttributeInfo> InVertexAttributeInfos = {},
-		ESkeletalMeshVertexFlags InBuildFlags = ESkeletalMeshVertexFlags::None
+		const FBuildSettings& InBuildSettings = FBuildSettings{ESkeletalMeshVertexFlags::None, false}
 		);
 #endif // WITH_EDITOR
 
