[Deformer Graph] Avoid creating deformer instance for dedicated server

diff --git a/Engine/Plugins/Animation/DeformerGraph/Source/OptimusCore/Private/OptimusDeformer.cpp b/Engine/Plugins/Animation/DeformerGraph/Source/OptimusCore/Private/OptimusDeformer.cpp
--- a/Engine/Plugins/Animation/DeformerGraph/Source/OptimusCore/Private/OptimusDeformer.cpp
+++ b/Engine/Plugins/Animation/DeformerGraph/Source/OptimusCore/Private/OptimusDeformer.cpp
@@ -36,6 +36,7 @@
 #include "ShaderCore.h"
 #include "UObject/AssetRegistryTagsContext.h"
 #include "UObject/Package.h"
+#include "Engine/World.h"
 
 // FIXME: We should not be accessing nodes directly.
 #include "OptimusValueContainer.h"
@@ -3413,6 +3414,13 @@ UMeshDeformerInstance* UOptimusDeformer::CreateInstance(
 		return nullptr;
 	}
 
+	// Return nullptr if running dedicated server
+	const UWorld* World = InMeshComponent->GetWorld();
+	if (World && World->IsNetMode(NM_DedicatedServer))
+	{
+		return nullptr;
+	}
+
 	const FName InstanceName(GetName() + TEXT("_Instance"));
 	UOptimusDeformerInstance* Instance = NewObject<UOptimusDeformerInstance>(InMeshComponent, InstanceName);
 	Instance->SetMeshComponent(InMeshComponent);
