[Deformer Graph] fix issues with uninitialized resources after compilation.

#rnx

diff --git a/Engine/Plugins/Animation/DeformerGraph/Source/OptimusCore/Private/OptimusDeformer.cpp b/Engine/Plugins/Animation/DeformerGraph/Source/OptimusCore/Private/OptimusDeformer.cpp
--- a/Engine/Plugins/Animation/DeformerGraph/Source/OptimusCore/Private/OptimusDeformer.cpp
+++ b/Engine/Plugins/Animation/DeformerGraph/Source/OptimusCore/Private/OptimusDeformer.cpp
@@ -2023,6 +2023,10 @@ TArray<FOptimusComputeGraphInfo> UOptimusDeformer::CompileNodeGraphToComputeGrap
 				{
 					LinksToInsertCopyKernel.FindOrAdd(SourcePin).Add(TargetPin);
 				}
+				else if (Cast<const IOptimusValueProvider>(SourceNode))
+				{
+					LinksToInsertCopyKernel.FindOrAdd(SourcePin).Add(TargetPin);
+				}
 			}
 		}
 	}
@@ -2190,7 +2194,6 @@ TArray<FOptimusComputeGraphInfo> UOptimusDeformer::CompileNodeGraphToComputeGrap
 						{
 							const FOptimusRoutedConstNode TargetRoutedNode = TargetInstancedPin.InstancedNode.RoutedNode;
 							const UOptimusNode* TargetNode = TargetRoutedNode.Node;
-							const UOptimusNodePin* TargetPin = TargetInstancedPin.Pin;
 
 							if (Cast<const IOptimusDataInterfaceProvider>(TargetNode))
 							{
@@ -2277,7 +2280,7 @@ TArray<FOptimusComputeGraphInfo> UOptimusDeformer::CompileNodeGraphToComputeGrap
 	
 	struct FDataInterfaceFunctionBinding
 	{
-		UOptimusComputeDataInterface* DataInterface;
+		UComputeDataInterface* DataInterface;
 		int32 FunctionIndex;
 	};
 	
@@ -2288,6 +2291,8 @@ TArray<FOptimusComputeGraphInfo> UOptimusDeformer::CompileNodeGraphToComputeGrap
 	for (TPair<FOptimusInstancedPin, TArray<FOptimusInstancedPin>> OutputLink : LinksToInsertCopyKernel)
 	{
 		const FOptimusInstancedPin& SourceInstancedPin = OutputLink.Key;
+		const TArray<FOptimusInstancedPin>& TargetInstancedPins = OutputLink.Value;
+		
 		const UOptimusNode* SourceNode = SourceInstancedPin.Pin->GetOwningNode();
 		if (const IOptimusDataInterfaceProvider* InterfaceProvider = Cast<const IOptimusDataInterfaceProvider>(SourceNode))
 		{
@@ -2296,6 +2301,13 @@ TArray<FOptimusComputeGraphInfo> UOptimusDeformer::CompileNodeGraphToComputeGrap
 			DataInterfaceBinding.FunctionIndex = InterfaceProvider->GetDataFunctionIndexFromPin(SourceInstancedPin.Pin);
 			CopyFromDataInterfaceMap.Add(SourceInstancedPin) = DataInterfaceBinding;
 		}
+		else if (Cast<const IOptimusValueProvider>(SourceNode))
+		{
+			FDataInterfaceFunctionBinding DataInterfaceBinding;
+			DataInterfaceBinding.DataInterface = GraphDataInterface;
+			DataInterfaceBinding.FunctionIndex = ValueNodes.Find(SourceNode);
+			CopyFromDataInterfaceMap.Add(SourceInstancedPin) = DataInterfaceBinding;	
+		}
 		else if (Cast<const IOptimusComputeKernelProvider>(SourceNode))
 		{
 			FDataInterfaceFunctionBinding DataInterfaceBinding;
@@ -2305,22 +2317,30 @@ TArray<FOptimusComputeGraphInfo> UOptimusDeformer::CompileNodeGraphToComputeGrap
 			CopyFromDataInterfaceMap.Add(SourceInstancedPin) = DataInterfaceBinding;
 		}
 
-		UOptimusCopyKernelDataInterface* CopyKernelDataInterface = NewObject<UOptimusCopyKernelDataInterface>(this);
-
-		CopyKernelDataInterface->SetExecutionDomain(*SourceInstancedPin.Pin->GetDataDomain().AsExpression());
-		UOptimusComponentSourceBinding* Binding = *SourceInstancedPin.Pin->GetComponentSourceBindings().CreateConstIterator();
-		CopyKernelDataInterface->SetComponentBinding(Binding);
-
-		CopyKernelDataInterfaceMap.Add(SourceInstancedPin) = CopyKernelDataInterface;
-		DataInterfaceToBindingIndexMap.Add(CopyKernelDataInterface) = Binding->GetIndex();
-
-		for (const FOptimusInstancedPin& TargetInstancedPin : OutputLink.Value)
+		bool bIsCopyKernelDataInterfaceCreated = false;
+		for (const FOptimusInstancedPin& TargetInstancedPin : TargetInstancedPins)
 		{
 			const UOptimusNode* TargetNode = TargetInstancedPin.Pin->GetOwningNode();
 
 			if (const IOptimusDataInterfaceProvider* InterfaceProvider = Cast<const IOptimusDataInterfaceProvider>(TargetNode);
 				ensure(InterfaceProvider))
 			{
+				// One-time Initialization of the copy kernel based on the first target pin, because if source is a value provider, it does not
+				// have a meaningful data domain and a meaning component source binding
+				if (!bIsCopyKernelDataInterfaceCreated)
+				{
+					bIsCopyKernelDataInterfaceCreated = true;
+					
+					UOptimusCopyKernelDataInterface* CopyKernelDataInterface = NewObject<UOptimusCopyKernelDataInterface>(this);
+					CopyKernelDataInterface->SetExecutionDomain(*TargetInstancedPin.Pin->GetDataDomain().AsExpression());
+					
+					UOptimusComponentSourceBinding* Binding = InterfaceProvider->GetComponentBinding();
+					CopyKernelDataInterface->SetComponentBinding(Binding);
+					CopyKernelDataInterfaceMap.Add(SourceInstancedPin) = CopyKernelDataInterface;
+					DataInterfaceToBindingIndexMap.Add(CopyKernelDataInterface) = Binding->GetIndex();
+				}
+				
+				
 				FDataInterfaceFunctionBinding DataInterfaceBinding;
 				DataInterfaceBinding.DataInterface = NodeDataInterfaceMap[TargetNode];
 				DataInterfaceBinding.FunctionIndex = InterfaceProvider->GetDataFunctionIndexFromPin(TargetInstancedPin.Pin);
@@ -2607,7 +2627,12 @@ TArray<FOptimusComputeGraphInfo> UOptimusDeformer::CompileNodeGraphToComputeGrap
 		
 					KernelSource->ExternalInputs.Emplace(FuncDef);
 
-					SourceText += FString::Printf(TEXT("%s Value = %s(Index);\n"), *ValueType->ToString(), *InterfaceBinding.BindingFunctionName);
+					FString IndexString;
+					if (FuncDef.ParamTypes.Num() == 2)
+					{
+						IndexString = TEXT("Index");
+					}
+					SourceText += FString::Printf(TEXT("%s Value = %s(%s);\n"), *ValueType->ToString(), *InterfaceBinding.BindingFunctionName, *IndexString);
 				}
 				
 				for (int32 TargetIndex = 0 ; TargetIndex < OutputLink.Value.Num(); TargetIndex++)
diff --git a/Engine/Plugins/Animation/DeformerGraph/Source/OptimusCore/Private/OptimusDeformerInstance.cpp b/Engine/Plugins/Animation/DeformerGraph/Source/OptimusCore/Private/OptimusDeformerInstance.cpp
--- a/Engine/Plugins/Animation/DeformerGraph/Source/OptimusCore/Private/OptimusDeformerInstance.cpp
+++ b/Engine/Plugins/Animation/DeformerGraph/Source/OptimusCore/Private/OptimusDeformerInstance.cpp
@@ -512,42 +512,37 @@ void UOptimusDeformerInstance::EnqueueWork(FEnqueueWorkDesc const& InDesc)
 		return;
 	}
 
-	// Get the current queued graphs.
-	TSet<FName> GraphsToRun;
-	{
-		UE::TScopeLock<FCriticalSection> Lock(GraphsToRunOnNextTickLock);
-		Swap(GraphsToRunOnNextTick, GraphsToRun);
-	}
-	
-	
-	
+
 	
 	// Enqueue work.
 	bool bIsWorkEnqueued = false;
 	if (bCanBeActive)
 	{
+		bool AreAllGraphsReady = true;
 		for (FOptimusDeformerInstanceExecInfo& Info: ComputeGraphExecInfos)
 		{
-			if (Info.GraphType == EOptimusNodeGraphType::Update || GraphsToRun.Contains(Info.GraphName))
+			if (Info.ComputeGraph->HasKernelResourcesPendingShaderCompilation())
 			{
-				FSimpleDelegate FallbackDelegate = InDesc.FallbackDelegate;
+				AreAllGraphsReady = false;
+				break;
+			}
+		}
 
-#if WITH_EDITOR
-				// Pending shader compilation may cause queued graph to not execute, so requeue on failure
-				// Ideally all compute graph should be ready before the deformer instance starts to enqueue work
-				if (Info.GraphType != EOptimusNodeGraphType::Update)
+		if (AreAllGraphsReady)
+		{
+			// Get the current queued graphs.
+			TSet<FName> GraphsToRun;
+			{
+				UE::TScopeLock<FCriticalSection> Lock(GraphsToRunOnNextTickLock);
+				Swap(GraphsToRunOnNextTick, GraphsToRun);
+			}
+			
+			for (FOptimusDeformerInstanceExecInfo& Info: ComputeGraphExecInfos)
+			{
+				if (Info.GraphType == EOptimusNodeGraphType::Update || GraphsToRun.Contains(Info.GraphName))
 				{
-					FallbackDelegate = FSimpleDelegate::CreateLambda([InDesc, this, GraphName=Info.GraphName]()
-					{
-						{
-							UE::TScopeLock<FCriticalSection> Lock(GraphsToRunOnNextTickLock);
-							GraphsToRunOnNextTick.Add(GraphName);
-						}	
-					});		
+					bIsWorkEnqueued |= Info.ComputeGraphInstance.EnqueueWork(Info.ComputeGraph, InDesc.Scene, ExecutionGroupName, InDesc.OwnerName, InDesc.FallbackDelegate, this);
 				}
-#endif
-				
-				bIsWorkEnqueued |= Info.ComputeGraphInstance.EnqueueWork(Info.ComputeGraph, InDesc.Scene, ExecutionGroupName, InDesc.OwnerName, FallbackDelegate, this);
 			}
 		}
 	}
diff --git a/Engine/Plugins/Runtime/ComputeFramework/Source/ComputeFramework/Private/ComputeGraph.cpp b/Engine/Plugins/Runtime/ComputeFramework/Source/ComputeFramework/Private/ComputeGraph.cpp
--- a/Engine/Plugins/Runtime/ComputeFramework/Source/ComputeFramework/Private/ComputeGraph.cpp
+++ b/Engine/Plugins/Runtime/ComputeFramework/Source/ComputeFramework/Private/ComputeGraph.cpp
@@ -210,6 +210,11 @@ void UComputeGraph::UpdateResources()
 	RenderProxy = CreateRenderProxy();
 }
 
+bool UComputeGraph::HasKernelResourcesPendingShaderCompilation() const
+{
+	return !KernelResourceIndicesPendingShaderCompilation.IsEmpty();
+}
+
 FComputeGraphRenderProxy const* UComputeGraph::GetRenderProxy() const 
 {
 	return RenderProxy;
@@ -621,6 +626,13 @@ void UComputeGraph::CacheResourceShadersForRendering(uint32 CompilationFlags)
 	if (FApp::CanEverRender())
 	{
 		KernelResources.SetNum(KernelInvocations.Num());
+
+		KernelResourceIndicesPendingShaderCompilation.Reset();
+		for (int32 KernelIndex = 0; KernelIndex < KernelInvocations.Num(); ++KernelIndex)
+		{
+			KernelResourceIndicesPendingShaderCompilation.Add(KernelIndex);
+		}
+		
 		for (int32 KernelIndex = 0; KernelIndex < KernelInvocations.Num(); ++KernelIndex)
 		{
 			UComputeKernel* Kernel = KernelInvocations[KernelIndex];
@@ -721,6 +733,7 @@ void UComputeGraph::ShaderCompileCompletionCallback(FComputeKernelResource const
 		if (KernelResource == KernelResources[KernelIndex].Get())
 		{
 			OnKernelCompilationComplete(KernelIndex, KernelResource->GetCompilationResults());
+			KernelResourceIndicesPendingShaderCompilation.Remove(KernelIndex);
 		}
 	}
 }
diff --git a/Engine/Plugins/Runtime/ComputeFramework/Source/ComputeFramework/Public/ComputeFramework/ComputeGraph.h b/Engine/Plugins/Runtime/ComputeFramework/Source/ComputeFramework/Public/ComputeFramework/ComputeGraph.h
--- a/Engine/Plugins/Runtime/ComputeFramework/Source/ComputeFramework/Public/ComputeFramework/ComputeGraph.h
+++ b/Engine/Plugins/Runtime/ComputeFramework/Source/ComputeFramework/Public/ComputeFramework/ComputeGraph.h
@@ -131,6 +131,12 @@ public:
 	 */
 	void UpdateResources();
 
+	/**
+	 * Shader compilations triggered by UpdateResources() are async,
+	 * This checks if all kernel resources are ready.
+	 */	
+	bool HasKernelResourcesPendingShaderCompilation() const;
+
 protected:
 	//~ Begin UObject Interface.
 	void Serialize(FArchive& Ar) override;
@@ -215,6 +221,10 @@ private:
 
 	/** Kernel resources stored with the same indexing as the KernelInvocations array. */
 	TArray<FComputeKernelResourceSet>  KernelResources;
+	
+	/** Indices of kernels pending shader compilation */
+	TSet<int32> KernelResourceIndicesPendingShaderCompilation;
+	
 	/** Render proxy that owns all render thread resources. */
 	FComputeGraphRenderProxy* RenderProxy = nullptr;
 };
