[Deformer Graph] minor fix for crash when loading a deformer asset that has dependency on a disabled/unloaded plugin

#rnx

diff --git a/Engine/Plugins/Animation/DeformerGraph/Source/OptimusCore/Public/OptimusDataType.h b/Engine/Plugins/Animation/DeformerGraph/Source/OptimusCore/Public/OptimusDataType.h
--- a/Engine/Plugins/Animation/DeformerGraph/Source/OptimusCore/Public/OptimusDataType.h
+++ b/Engine/Plugins/Animation/DeformerGraph/Source/OptimusCore/Public/OptimusDataType.h
@@ -125,7 +125,9 @@ struct OPTIMUSCORE_API FOptimusDataTypeRef
 
 	bool IsValid() const
 	{
-		return !TypeName.IsNone();
+		// The serialized DataTypeRef can become invalid when we load a deformer graph asset that has dependency on disabled plugins
+		// So make sure we always check the data type registry.
+		return Resolve().IsValid();
 	}
 
 	explicit operator bool() const
