[Deformer Graph] Minor fix 2 for crash when opening deformer graph asset with disabled dependent plugin

#rnx

diff --git a/Engine/Plugins/Animation/DeformerGraph/Source/OptimusCore/Private/Nodes/OptimusNode_CustomComputeKernel.cpp b/Engine/Plugins/Animation/DeformerGraph/Source/OptimusCore/Private/Nodes/OptimusNode_CustomComputeKernel.cpp
--- a/Engine/Plugins/Animation/DeformerGraph/Source/OptimusCore/Private/Nodes/OptimusNode_CustomComputeKernel.cpp
+++ b/Engine/Plugins/Animation/DeformerGraph/Source/OptimusCore/Private/Nodes/OptimusNode_CustomComputeKernel.cpp
@@ -1719,7 +1719,7 @@ void UOptimusNode_CustomComputeKernel::UpdatePreamble()
 		{
 			AddCountFunctionIfNeeded(Binding.DataDomain.DimensionNames);
 		
-			if (Binding.DataType->ShaderValueType.IsValid())
+			if (Binding.DataType.IsValid() && Binding.DataType->ShaderValueType.IsValid())
 			{
 				Declarations.Add(GetDeclarationForBinding(Binding, true));
 			}
