[Deformer Graph] fix an issue with dup vert data interface shader causing out of bound access. dup vert buffers are provided per section and thus does not need to use the per section vertex index offset

diff --git a/Engine/Plugins/Animation/DeformerGraph/Shaders/Private/DataInterfaceDuplicateVertices.ush b/Engine/Plugins/Animation/DeformerGraph/Shaders/Private/DataInterfaceDuplicateVertices.ush
--- a/Engine/Plugins/Animation/DeformerGraph/Shaders/Private/DataInterfaceDuplicateVertices.ush
+++ b/Engine/Plugins/Animation/DeformerGraph/Shaders/Private/DataInterfaceDuplicateVertices.ush
@@ -19,7 +19,8 @@ uint ReadNumDuplicateVertices_{DataInterfaceName}()
 int2 ReadDuplicateVerticesStartAndLength_{DataInterfaceName}(uint VertexIndex)
 {
 #if ENABLE_DUPLICATED_VERTICES
-	uint BufferIndex = ({DataInterfaceName}_InputStreamStart + VertexIndex) * 2;
+	// DuplicateVertex buffer is provided per section and thus we don't need to offset index with InputStreamStart
+	uint BufferIndex = VertexIndex * 2;
 	return int2({DataInterfaceName}_DuplicateVertexStartAndLength[BufferIndex  + 1], {DataInterfaceName}_DuplicateVertexStartAndLength[BufferIndex]);
 #else	
 	return 0;
