[Deformer Graph] Added ability to collapse nodes to subgraph and expand a subgraph node

diff --git a/Engine/Plugins/Animation/DeformerGraph/Source/OptimusCore/Private/DataInterfaces/OptimusDataInterfaceDebugDraw.cpp b/Engine/Plugins/Animation/DeformerGraph/Source/OptimusCore/Private/DataInterfaces/OptimusDataInterfaceDebugDraw.cpp
--- a/Engine/Plugins/Animation/DeformerGraph/Source/OptimusCore/Private/DataInterfaces/OptimusDataInterfaceDebugDraw.cpp
+++ b/Engine/Plugins/Animation/DeformerGraph/Source/OptimusCore/Private/DataInterfaces/OptimusDataInterfaceDebugDraw.cpp
@@ -53,7 +53,7 @@ void UOptimusDebugDrawDataInterface::RegisterTypes()
 		FName("FDebugDraw"),
 		nullptr,
 		FLinearColor(0.3f, 0.7f, 0.4f, 1.0f),
-		EOptimusDataTypeUsageFlags::None);
+		EOptimusDataTypeUsageFlags::DataInterfaceOutput | EOptimusDataTypeUsageFlags::PinType);
 }
 
 void UOptimusDebugDrawDataInterface::GetSupportedInputs(TArray<FShaderFunctionDefinition>& OutFunctions) const
diff --git a/Engine/Plugins/Animation/DeformerGraph/Source/OptimusCore/Private/DataInterfaces/OptimusDataInterfaceRawBuffer.cpp b/Engine/Plugins/Animation/DeformerGraph/Source/OptimusCore/Private/DataInterfaces/OptimusDataInterfaceRawBuffer.cpp
--- a/Engine/Plugins/Animation/DeformerGraph/Source/OptimusCore/Private/DataInterfaces/OptimusDataInterfaceRawBuffer.cpp
+++ b/Engine/Plugins/Animation/DeformerGraph/Source/OptimusCore/Private/DataInterfaces/OptimusDataInterfaceRawBuffer.cpp
@@ -2,6 +2,7 @@
 
 #include "OptimusDataInterfaceRawBuffer.h"
 
+#include "OptimusDataTypeRegistry.h"
 #include "ComputeFramework/ShaderParamTypeDefinition.h"
 #include "OptimusDeformerInstance.h"
 #include "OptimusExpressionEvaluator.h"
@@ -27,7 +28,16 @@ const UOptimusComponentSource* UOptimusRawBufferDataInterface::GetComponentSourc
 
 bool UOptimusRawBufferDataInterface::SupportsAtomics() const
 {
-	return ValueType->Type == EShaderFundamentalType::Int;
+	TArray<FOptimusDataTypeHandle> Types = FOptimusDataTypeRegistry::Get().GetAllTypesWithAtomicSupport();
+	for (FOptimusDataTypeHandle Type : Types)
+	{
+		if (Type->ShaderValueType == ValueType)
+		{
+			return true;
+		}
+	}
+	
+	return false;
 }
 
 FString UOptimusRawBufferDataInterface::GetRawType() const
diff --git a/Engine/Plugins/Animation/DeformerGraph/Source/OptimusCore/Private/Nodes/OptimusNode_ComponentSource.h b/Engine/Plugins/Animation/DeformerGraph/Source/OptimusCore/Private/Nodes/OptimusNode_ComponentSource.h
--- a/Engine/Plugins/Animation/DeformerGraph/Source/OptimusCore/Private/Nodes/OptimusNode_ComponentSource.h
+++ b/Engine/Plugins/Animation/DeformerGraph/Source/OptimusCore/Private/Nodes/OptimusNode_ComponentSource.h
@@ -3,6 +3,7 @@
 #pragma once
 
 #include "IOptimusComponentBindingProvider.h"
+#include "IOptimusNonCollapsibleNode.h"
 #include "OptimusNode.h"
 
 #include "OptimusComponentSource.h"
@@ -24,7 +25,8 @@ struct FOptimusNode_ComponentSource_DuplicationInfo
 UCLASS(Hidden)
 class UOptimusNode_ComponentSource :
 	public UOptimusNode,
-	public IOptimusComponentBindingProvider
+	public IOptimusComponentBindingProvider,
+	public IOptimusNonCollapsibleNode
 {
 	GENERATED_BODY()
 public:
@@ -36,7 +38,7 @@ public:
 	FName GetNodeCategory() const override;
 
 	// IOptimusComponentSourceBindingProvider implementation
-	UOptimusComponentSourceBinding *GetComponentBinding() const override
+	UOptimusComponentSourceBinding* GetComponentBinding(const FOptimusPinTraversalContext& InContext = {}) const override
 	{
 		return Binding;
 	}
diff --git a/Engine/Plugins/Animation/DeformerGraph/Source/OptimusCore/Private/Nodes/OptimusNode_ComputeKernelBase.cpp b/Engine/Plugins/Animation/DeformerGraph/Source/OptimusCore/Private/Nodes/OptimusNode_ComputeKernelBase.cpp
--- a/Engine/Plugins/Animation/DeformerGraph/Source/OptimusCore/Private/Nodes/OptimusNode_ComputeKernelBase.cpp
+++ b/Engine/Plugins/Animation/DeformerGraph/Source/OptimusCore/Private/Nodes/OptimusNode_ComputeKernelBase.cpp
@@ -137,7 +137,7 @@ FOptimus_ComputeKernelResult UOptimusNode_ComputeKernelBase::CreateComputeKernel
 		}
 	}
 
-	TSet<UOptimusComponentSourceBinding*> PrimaryBindings = GetPrimaryGroupPin()->GetComponentSourceBindingsRecursively();
+	TSet<UOptimusComponentSourceBinding*> PrimaryBindings = GetPrimaryGroupPin()->GetComponentSourceBindingsRecursively(InTraversalContext);
 	
 	// ValidateForCompile() should guaranteed that we don't hit this error
 	if (!ensure(PrimaryBindings.Num() == 1))
@@ -242,7 +242,7 @@ FOptimus_ComputeKernelResult UOptimusNode_ComputeKernelBase::CreateComputeKernel
 }
 
 
-TOptional<FText> UOptimusNode_ComputeKernelBase::ValidateForCompile() const
+TOptional<FText> UOptimusNode_ComputeKernelBase::ValidateForCompile(const FOptimusPinTraversalContext& InContext) const
 {
 	auto GetStructTypeDefFromPin = [](const UOptimusNodePin* InPin) -> TOptional<FText>
 	{
@@ -262,11 +262,11 @@ TOptional<FText> UOptimusNode_ComputeKernelBase::ValidateForCompile() const
 	};
 
 	// Collect struct types and ensure all groups have the same component source (or none) within each group
-	auto VerifyPinBindings = [](TSet<UOptimusComponentSourceBinding*> &InCollectedBindings, const UOptimusNodePin* InPin)-> TOptional<FText>
+	auto VerifyPinBindings = [InContext](TSet<UOptimusComponentSourceBinding*> &InCollectedBindings, const UOptimusNodePin* InPin)-> TOptional<FText>
 	{
 		// Check for all pins, even if the data domain is singleton, it can have different values per invocation
 		// UOptimusNode_DataInterface::ValidateForCompile also double checks
-		const TSet<UOptimusComponentSourceBinding*> PinBindings = InPin->GetComponentSourceBindings();
+		const TSet<UOptimusComponentSourceBinding*> PinBindings = InPin->GetComponentSourceBindings(InContext);
 		if (PinBindings.Num() > 1)
 		{
 			return FText::Format(LOCTEXT("MultipleBindingsOnPin", "Multiple component bindings arriving into pin '{0}'"), FText::FromName(InPin->GetUniqueName()));
diff --git a/Engine/Plugins/Animation/DeformerGraph/Source/OptimusCore/Private/Nodes/OptimusNode_ComputeKernelBase.h b/Engine/Plugins/Animation/DeformerGraph/Source/OptimusCore/Private/Nodes/OptimusNode_ComputeKernelBase.h
--- a/Engine/Plugins/Animation/DeformerGraph/Source/OptimusCore/Private/Nodes/OptimusNode_ComputeKernelBase.h
+++ b/Engine/Plugins/Animation/DeformerGraph/Source/OptimusCore/Private/Nodes/OptimusNode_ComputeKernelBase.h
@@ -54,7 +54,7 @@ public:
 	bool DoesOutputPinSupportRead(const UOptimusNodePin* InPin) const override PURE_VIRTUAL(UOptimusNode_ComputeKernelBase::DoesOutputPinSupportRead, return false; );
 	
 	// -- UOptimusNode overrides
-	TOptional<FText> ValidateForCompile() const override;
+	TOptional<FText> ValidateForCompile(const FOptimusPinTraversalContext& InContext) const override;
 	
 protected:
 	static TArray<FString> GetIndexNamesFromDataDomainLevels(
diff --git a/Engine/Plugins/Animation/DeformerGraph/Source/OptimusCore/Private/Nodes/OptimusNode_CustomComputeKernel.cpp b/Engine/Plugins/Animation/DeformerGraph/Source/OptimusCore/Private/Nodes/OptimusNode_CustomComputeKernel.cpp
--- a/Engine/Plugins/Animation/DeformerGraph/Source/OptimusCore/Private/Nodes/OptimusNode_CustomComputeKernel.cpp
+++ b/Engine/Plugins/Animation/DeformerGraph/Source/OptimusCore/Private/Nodes/OptimusNode_CustomComputeKernel.cpp
@@ -68,8 +68,7 @@ static bool IsSecondaryGroupInputPin(const UOptimusNodePin *InPin)
 
 static bool DoesBindingSupportAtomic(const FOptimusParameterBinding& InBinding)
 {
-	FOptimusDataTypeHandle IntType = FOptimusDataTypeRegistry::Get().FindType(*FIntProperty::StaticClass());
-	return InBinding.bSupportAtomicIfCompatibleDataType && InBinding.DataType == IntType;	
+	return InBinding.bSupportAtomicIfCompatibleDataType && FOptimusDataTypeRegistry::Get().DoesTypeSupportAtomic(InBinding.DataType.Resolve()) ;	
 }
 
 bool UOptimusNode_CustomComputeKernel::DoesSourceSupportUnifiedDispatch(const UOptimusNodePin& InOtherNodesPin)
@@ -235,8 +234,7 @@ bool UOptimusNode_CustomComputeKernel::GetBindingSupportAtomicCheckBoxVisibility
 
 	if (const FOptimusParameterBinding* Binding = OutputBindingArray.FindByPredicate(ParameterBindingPredicate))
 	{
-		FOptimusDataTypeHandle IntType = FOptimusDataTypeRegistry::Get().FindType(*FIntProperty::StaticClass());
-		if (Binding->DataType == IntType)
+		if (FOptimusDataTypeRegistry::Get().DoesTypeSupportAtomic(Binding->DataType.Resolve()))
 		{
 			return true;
 		}
@@ -265,6 +263,16 @@ bool UOptimusNode_CustomComputeKernel::GetBindingSupportReadCheckBoxVisibility(F
 	return false;
 }
 
+EOptimusDataTypeUsageFlags UOptimusNode_CustomComputeKernel::GetTypeUsageFlags(const FOptimusDataDomain& InDataDomain) const
+{
+	if (InDataDomain.IsSingleton())
+	{
+		return EOptimusDataTypeUsageFlags::Variable | EOptimusDataTypeUsageFlags::AnimAttributes | EOptimusDataTypeUsageFlags::DataInterfaceOutput;
+	}
+
+	return EOptimusDataTypeUsageFlags::Resource;
+}
+
 
 TArray<IOptimusNodeAdderPinProvider::FAdderPinAction> UOptimusNode_CustomComputeKernel::GetAvailableAdderPinActions(
 	const UOptimusNodePin* InSourcePin,
@@ -272,21 +280,21 @@ TArray<IOptimusNodeAdderPinProvider::FAdderPinAction> UOptimusNode_CustomCompute
 	FString* OutReason
 	) const
 {
-	if (InSourcePin->GetDataType().IsValid() && InSourcePin->GetDataType()->UsageFlags == EOptimusDataTypeUsageFlags::None)
+	if (InSourcePin->GetDataType().IsValid())
 	{
-		// FIXME: avoid special case debug draw
-		if (InSourcePin->GetDataType() != FOptimusDataTypeRegistry::Get().FindType(FName(TEXT("FDebugDraw"))))
+		if (!EnumHasAnyFlags(InSourcePin->GetDataType()->UsageFlags, GetTypeUsageFlags(InSourcePin->GetDataDomain())))
 		{
 			if (OutReason)
 			{
 				*OutReason = TEXT("Can't add pin with this type");
 			}
-
+			
 			return {};
 		}
+
 	}
 	
-	TSet<UOptimusComponentSourceBinding*> SourceComponentBindings = InSourcePin->GetComponentSourceBindings();
+	TSet<UOptimusComponentSourceBinding*> SourceComponentBindings = InSourcePin->GetComponentSourceBindings({});
 	
 	TArray<UOptimusNodePin*> GroupPins;
 
@@ -321,7 +329,7 @@ TArray<IOptimusNodeAdderPinProvider::FAdderPinAction> UOptimusNode_CustomCompute
 		if (InSourcePin->GetDirection() != InNewPinDirection)
 		{
 			const UOptimusNodePin* PrimaryGroupPin = GetPrimaryGroupPin();
-			TSet<UOptimusComponentSourceBinding*> PrimaryGroupComponentBindings = PrimaryGroupPin->GetComponentSourceBindingsRecursively();
+			TSet<UOptimusComponentSourceBinding*> PrimaryGroupComponentBindings = PrimaryGroupPin->GetComponentSourceBindingsRecursively({});
 			TSet<UOptimusComponentSourceBinding*> MatchingBindings = PrimaryGroupComponentBindings.Intersect(SourceComponentBindings);
 		
 			if (MatchingBindings.Num() > 0 ||
@@ -371,7 +379,7 @@ TArray<IOptimusNodeAdderPinProvider::FAdderPinAction> UOptimusNode_CustomCompute
 				Action.DisplayName = GroupPin->GetFName();
 				Action.Key = GroupPin->GetFName();
 					
-				TSet<UOptimusComponentSourceBinding*> ComponentBindingsForGroup = GroupPin->GetComponentSourceBindingsRecursively();
+				TSet<UOptimusComponentSourceBinding*> ComponentBindingsForGroup = GroupPin->GetComponentSourceBindingsRecursively({});
 				TSet<UOptimusComponentSourceBinding*> MatchingBindings = ComponentBindingsForGroup.Intersect(SourceComponentBindings);
 	
 				if (MatchingBindings.Num() > 0 ||
@@ -456,7 +464,7 @@ TArray<UOptimusNodePin*> UOptimusNode_CustomComputeKernel::TryAddPinFromPin(
 		{
 			//Creating a new group
 			
-			TSet<UOptimusComponentSourceBinding*> ComponentBindings = InSourcePin->GetComponentSourceBindings();
+			TSet<UOptimusComponentSourceBinding*> ComponentBindings = InSourcePin->GetComponentSourceBindings({});
 			
 			check(ComponentBindings.Num() != 0);
 
@@ -646,7 +654,7 @@ TArray<FName> UOptimusNode_CustomComputeKernel::GetExecutionDomains() const
 {
 	// Find all component sources for the primary pins. If we end up with any other number
 	// than one, then something's gone wrong and we can't determine the execution domains.
-	TSet<UOptimusComponentSourceBinding*> PrimaryBindings = GetPrimaryGroupPin()->GetComponentSourceBindingsRecursively();
+	TSet<UOptimusComponentSourceBinding*> PrimaryBindings = GetPrimaryGroupPin()->GetComponentSourceBindingsRecursively({});
 
 	if (PrimaryBindings.Num() == 1)
 	{
@@ -654,7 +662,7 @@ TArray<FName> UOptimusNode_CustomComputeKernel::GetExecutionDomains() const
 	}
 	else
 	{
-		return {};
+		return UOptimusComponentSource::GetAllExecutionDomains().Array();
 	}
 }
 
@@ -1356,79 +1364,39 @@ void UOptimusNode_CustomComputeKernel::PropertyArrayItemMoved(
 		TFunction<bool(const UOptimusNodePin *)> InPinPredicate
 		)
 	{
-		FName NameToMove = NAME_None;
-		
-		// Find the first entry that's different. That's an element we can consider moved. Since array move only
-		// deals with a single item moving either forward or backward.
+		TArray<FName> PinNames;
+	
+		for (const UOptimusNodePin* Pin: InPins)
 		{
-			int32 DivergeIndex = INDEX_NONE;
-			FName PinNameAtDiverge = NAME_None;
-			FName BindingNameAtDiverge = NAME_None;
-			
-			int32 BindingIndex = 0;
-			for (const UOptimusNodePin* Pin: InPins)
+			if (InPinPredicate(Pin))
 			{
-				if (InPinPredicate(Pin))
-				{
-					if (Pin->GetFName() != BindingNameArray[BindingIndex] && DivergeIndex == INDEX_NONE)
-					{
-						DivergeIndex = BindingIndex;
-						PinNameAtDiverge = Pin->GetFName();
-						BindingNameAtDiverge = BindingNameArray[BindingIndex];
-						BindingIndex++;
-						continue;
-					}
-
-					if (DivergeIndex != INDEX_NONE)
-					{
-						if (BindingNameAtDiverge == Pin->GetFName())
-						{
-							NameToMove = PinNameAtDiverge;
-						}
-						else if (ensure(BindingNameArray[BindingIndex] == PinNameAtDiverge))
-						{
-							NameToMove = BindingNameAtDiverge;
-						}
-
-						break;
-					}
-					
-					BindingIndex++;
-					
-					
-					if (BindingIndex == BindingNameArray.Num())
-					{
-						// Nothing got moved.
-						return;
-					}
-				}
+				PinNames.Add(Pin->GetFName());
 			}
 		}
 
-		UOptimusNodePin* MovedPin = *InPins.FindByPredicate([NameToMove](const UOptimusNodePin* InPin)
-		{
-			return InPin->GetFName() == NameToMove;
-		});
 
-		int32 NameIndex = BindingNameArray.IndexOfByPredicate([NameToMove](FName InBindingName)
-		{
-			return InBindingName == NameToMove;
-		});
-
-		const int32 NextNameIndex = NameIndex + 1;
-
-		const FName NextName = BindingNameArray.IsValidIndex(NextNameIndex) ? BindingNameArray[NextNameIndex] : NAME_None;
+		
+		FName PinName = NAME_None;
+		FName NextPinName = NAME_None;
 
-		const UOptimusNodePin* NextPin = nullptr;
-		if (!NextName.IsNone())
+		if (Optimus::FindMovedItemInNameArray(PinNames, BindingNameArray, PinName, NextPinName))
 		{
-			NextPin = *InPins.FindByPredicate([NextName](const UOptimusNodePin* InPin)
+			UOptimusNodePin* MovedPin = *InPins.FindByPredicate([PinName](const UOptimusNodePin* InPin)
 			{
-				return InPin->GetFName() == NextName;
+				return InPin->GetFName() == PinName;
 			});
-		}
 
-		MovePin(MovedPin, NextPin);
+			UOptimusNodePin* NextPin = nullptr;
+			if (!NextPinName.IsNone())
+			{
+				NextPin = *InPins.FindByPredicate([NextPinName](const UOptimusNodePin* InPin)
+				{
+					return InPin->GetFName() == NextPinName;
+				});	
+			}
+
+			MovePin(MovedPin, NextPin);
+		}
 	};
 
 	auto MakeBindingNameArray = [](const FOptimusParameterBindingArray& InBindings) -> TArray<FName>
@@ -1623,9 +1591,9 @@ bool UOptimusNode_CustomComputeKernel::ValidateConnection(
 	return true;
 }
 
-TOptional<FText> UOptimusNode_CustomComputeKernel::ValidateForCompile() const
+TOptional<FText> UOptimusNode_CustomComputeKernel::ValidateForCompile(const FOptimusPinTraversalContext& InContext) const
 {
-	if (TOptional<FText> Result = Super::ValidateForCompile(); Result.IsSet())
+	if (TOptional<FText> Result = Super::ValidateForCompile(InContext); Result.IsSet())
 	{
 		return Result;
 	}
diff --git a/Engine/Plugins/Animation/DeformerGraph/Source/OptimusCore/Private/Nodes/OptimusNode_CustomComputeKernel.h b/Engine/Plugins/Animation/DeformerGraph/Source/OptimusCore/Private/Nodes/OptimusNode_CustomComputeKernel.h
--- a/Engine/Plugins/Animation/DeformerGraph/Source/OptimusCore/Private/Nodes/OptimusNode_CustomComputeKernel.h
+++ b/Engine/Plugins/Animation/DeformerGraph/Source/OptimusCore/Private/Nodes/OptimusNode_CustomComputeKernel.h
@@ -78,6 +78,7 @@ public:
 	FString GetBindingDeclaration(FName BindingName) const override;
 	bool GetBindingSupportAtomicCheckBoxVisibility(FName BindingName) const override;
 	bool GetBindingSupportReadCheckBoxVisibility(FName BindingName) const override;
+	EOptimusDataTypeUsageFlags GetTypeUsageFlags(const FOptimusDataDomain& InDataDomain) const override;
 	
 	// IOptimusNodeAdderPinProvider
 	TArray<FAdderPinAction> GetAvailableAdderPinActions(
@@ -173,7 +174,7 @@ PRAGMA_ENABLE_DEPRECATION_WARNINGS
 protected:
 	void ConstructNode() override;
 	bool ValidateConnection(const UOptimusNodePin& InThisNodesPin, const UOptimusNodePin& InOtherNodesPin, FString* OutReason) const override;
-	TOptional<FText> ValidateForCompile() const override;
+	TOptional<FText> ValidateForCompile(const FOptimusPinTraversalContext& InContext) const override;
 	
 private:
 
diff --git a/Engine/Plugins/Animation/DeformerGraph/Source/OptimusCore/Private/Nodes/OptimusNode_DataInterface.cpp b/Engine/Plugins/Animation/DeformerGraph/Source/OptimusCore/Private/Nodes/OptimusNode_DataInterface.cpp
--- a/Engine/Plugins/Animation/DeformerGraph/Source/OptimusCore/Private/Nodes/OptimusNode_DataInterface.cpp
+++ b/Engine/Plugins/Animation/DeformerGraph/Source/OptimusCore/Private/Nodes/OptimusNode_DataInterface.cpp
@@ -7,6 +7,7 @@
 #include "OptimusDataTypeRegistry.h"
 #include "OptimusDeformer.h"
 #include "OptimusNodeGraph.h"
+#include "OptimusNodeSubGraph.h"
 #include "OptimusNode_ComponentSource.h"
 #include "OptimusObjectVersion.h"
 
@@ -36,39 +37,36 @@ bool UOptimusNode_DataInterface::ValidateConnection(
 	if (!GetPins().IsEmpty() && &InThisNodesPin == GetPins()[0])
 	{
 		const UOptimusNode_ComponentSource* SourceNode = Cast<UOptimusNode_ComponentSource>(InOtherNodesPin.GetOwningNode());
-		if (!SourceNode)
+		if (SourceNode)
 		{
-			if (OutReason)
+			const UOptimusComponentSource* ComponentSource = SourceNode->GetComponentBinding()->GetComponentSource();
+			if (!IsComponentSourceCompatible(ComponentSource))
 			{
-				*OutReason = TEXT("Other node should be a Component Source node");
-			}
-			return false;
+				if (OutReason)
+				{
+					*OutReason = FString::Printf(TEXT("This data interface requires a %s which is not a child class of %s from the Component Source."),
+						*DataInterfaceData->GetRequiredComponentClass()->GetName(),
+						*ComponentSource->GetComponentClass()->GetName());
+				}
+				return false;
+			}	
 		}
 
-		const UOptimusComponentSource* ComponentSource = SourceNode->GetComponentBinding()->GetComponentSource();
-		if (!IsComponentSourceCompatible(ComponentSource))
-		{
-			if (OutReason)
-			{
-				*OutReason = FString::Printf(TEXT("This data interface requires a %s which is not a child class of %s from the Component Source."),
-					*DataInterfaceData->GetRequiredComponentClass()->GetName(),
-					*ComponentSource->GetComponentClass()->GetName());
-			}
-			return false;
-		}
+		// In other cases, the component source may come from the upstream of the connected node (eg. Sub Graph Terminal),
+		// and thus no way to provide error check until compile time
 	}
 
 	return true;
 }
 
-TOptional<FText> UOptimusNode_DataInterface::ValidateForCompile() const
+TOptional<FText> UOptimusNode_DataInterface::ValidateForCompile(const FOptimusPinTraversalContext& InContext) const
 {
 	if (!DataInterfaceClass)
 	{
 		return LOCTEXT("NoAssociatedClass", "Node has none or invalid data interface class associated with it. Delete and re-create the node.");
 	}
 	// Ensure that we have something connected to the component binding input pin.
-	UOptimusComponentSourceBinding* PrimaryBinding = GetComponentBinding();
+	UOptimusComponentSourceBinding* PrimaryBinding = GetComponentBinding(InContext);
 	if (PrimaryBinding == nullptr)
 	{
 		return FText::Format(LOCTEXT("NoBindingConnected", "No component binding connected to the {0} pin"), FText::FromName(GetComponentPin()->GetUniqueName()));
@@ -80,7 +78,7 @@ TOptional<FText> UOptimusNode_DataInterface::ValidateForCompile() const
 	{
 		if (Pin->GetDirection() == EOptimusNodePinDirection::Input && Pin != GetComponentPin())
 		{
-			TSet<UOptimusComponentSourceBinding*> Bindings = Graph->GetComponentSourceBindingsForPin(Pin);
+			TSet<UOptimusComponentSourceBinding*> Bindings = Graph->GetComponentSourceBindingsForPin(Pin, InContext);
 			if (Bindings.Num() > 1)
 			{
 				return FText::Format(LOCTEXT("MultipleBindingsOnPin", "Multiple bindings found for pin {0}"), FText::FromName(Pin->GetUniqueName()));
@@ -202,10 +200,10 @@ void UOptimusNode_DataInterface::SetDataInterfaceClass(
 	DataInterfaceData = NewObject<UOptimusComputeDataInterface>(this, DataInterfaceClass);
 }
 
-UOptimusComponentSourceBinding* UOptimusNode_DataInterface::GetComponentBinding() const
+UOptimusComponentSourceBinding* UOptimusNode_DataInterface::GetComponentBinding(const FOptimusPinTraversalContext& InContext) const
 {
 	const UOptimusNodeGraph* Graph = GetOwningGraph();
-	TSet<UOptimusComponentSourceBinding*> Bindings = Graph->GetComponentSourceBindingsForPin(GetComponentPin());
+	TSet<UOptimusComponentSourceBinding*> Bindings = Graph->GetComponentSourceBindingsForPin(GetComponentPin(), InContext);
 	
 	if (!Bindings.IsEmpty() && ensure(Bindings.Num() == 1))
 	{
@@ -217,6 +215,11 @@ UOptimusComponentSourceBinding* UOptimusNode_DataInterface::GetComponentBinding(
 	{
 		return Deformer->GetPrimaryComponentBinding();
 	}
+
+	if (const UOptimusNodeSubGraph* SubGraph = Cast<UOptimusNodeSubGraph>(Graph))
+	{
+		return SubGraph->GetDefaultComponentBinding(InContext);
+	}
 	
 	return nullptr;
 }
diff --git a/Engine/Plugins/Animation/DeformerGraph/Source/OptimusCore/Private/Nodes/OptimusNode_DataInterface.h b/Engine/Plugins/Animation/DeformerGraph/Source/OptimusCore/Private/Nodes/OptimusNode_DataInterface.h
--- a/Engine/Plugins/Animation/DeformerGraph/Source/OptimusCore/Private/Nodes/OptimusNode_DataInterface.h
+++ b/Engine/Plugins/Animation/DeformerGraph/Source/OptimusCore/Private/Nodes/OptimusNode_DataInterface.h
@@ -45,7 +45,7 @@ public:
 	int32 GetDataFunctionIndexFromPin(const UOptimusNodePin* InPin) const override;
 	
 	// -- IOptimusDataInterfaceProvider & IOptimusComponentBindingProvider
-	UOptimusComponentSourceBinding* GetComponentBinding() const override;
+	UOptimusComponentSourceBinding* GetComponentBinding(const FOptimusPinTraversalContext& InContext) const override;
 
 	// -- IOptimusPinMutabilityDefiner
 	EOptimusPinMutability GetOutputPinMutability(const UOptimusNodePin* InPin) const override;
@@ -54,7 +54,7 @@ protected:
 	// -- UOptimusNode overrides
 	void ConstructNode() override;
 	bool ValidateConnection(const UOptimusNodePin& InThisNodesPin, const UOptimusNodePin& InOtherNodesPin, FString* OutReason) const override;
-	TOptional<FText> ValidateForCompile() const override;
+	TOptional<FText> ValidateForCompile(const FOptimusPinTraversalContext& InContext) const override;
 
 	void SaveState(FArchive& Ar) const override;
 	void RestoreState(FArchive& Ar) override;
diff --git a/Engine/Plugins/Animation/DeformerGraph/Source/OptimusCore/Private/Nodes/OptimusNode_GetVariable.cpp b/Engine/Plugins/Animation/DeformerGraph/Source/OptimusCore/Private/Nodes/OptimusNode_GetVariable.cpp
--- a/Engine/Plugins/Animation/DeformerGraph/Source/OptimusCore/Private/Nodes/OptimusNode_GetVariable.cpp
+++ b/Engine/Plugins/Animation/DeformerGraph/Source/OptimusCore/Private/Nodes/OptimusNode_GetVariable.cpp
@@ -184,7 +184,7 @@ UOptimusVariableDescription* UOptimusNode_GetVariable::GetVariableDescription()
 }
 
 
-TOptional<FText> UOptimusNode_GetVariable::ValidateForCompile() const
+TOptional<FText> UOptimusNode_GetVariable::ValidateForCompile(const FOptimusPinTraversalContext& InContext) const
 {
 	const UOptimusVariableDescription* VariableDescription = GetVariableDescription();
 	if (!VariableDescription)
diff --git a/Engine/Plugins/Animation/DeformerGraph/Source/OptimusCore/Private/Nodes/OptimusNode_GetVariable.h b/Engine/Plugins/Animation/DeformerGraph/Source/OptimusCore/Private/Nodes/OptimusNode_GetVariable.h
--- a/Engine/Plugins/Animation/DeformerGraph/Source/OptimusCore/Private/Nodes/OptimusNode_GetVariable.h
+++ b/Engine/Plugins/Animation/DeformerGraph/Source/OptimusCore/Private/Nodes/OptimusNode_GetVariable.h
@@ -2,6 +2,7 @@
 
 #pragma once
 
+#include "IOptimusNonCollapsibleNode.h"
 #include "IOptimusPinMutabilityDefiner.h"
 #include "IOptimusValueProvider.h"
 #include "OptimusNode.h"
@@ -34,7 +35,8 @@ UCLASS(Hidden)
 class UOptimusNode_GetVariable : 
 	public UOptimusNode,
 	public IOptimusValueProvider,
-	public IOptimusPinMutabilityDefiner
+	public IOptimusPinMutabilityDefiner,
+	public IOptimusNonCollapsibleNode
 {
 	GENERATED_BODY()
 
@@ -49,7 +51,7 @@ public:
 		return CategoryName::Variables;
 	}
 
-	TOptional<FText> ValidateForCompile() const override;
+	TOptional<FText> ValidateForCompile(const FOptimusPinTraversalContext& InContext) const override;
 	
 	// IOptimusValueProvider overrides 
 	FString GetValueName() const override;
diff --git a/Engine/Plugins/Animation/DeformerGraph/Source/OptimusCore/Private/Nodes/OptimusNode_GraphTerminal.cpp b/Engine/Plugins/Animation/DeformerGraph/Source/OptimusCore/Private/Nodes/OptimusNode_GraphTerminal.cpp
--- a/Engine/Plugins/Animation/DeformerGraph/Source/OptimusCore/Private/Nodes/OptimusNode_GraphTerminal.cpp
+++ b/Engine/Plugins/Animation/DeformerGraph/Source/OptimusCore/Private/Nodes/OptimusNode_GraphTerminal.cpp
@@ -5,7 +5,11 @@
 #include "OptimusFunctionNodeGraph.h"
 #include "OptimusNodePin.h"
 #include "OptimusBindingTypes.h"
+#include "OptimusComponentSource.h"
+#include "OptimusDataTypeRegistry.h"
+#include "OptimusHelpers.h"
 #include "OptimusNode_ComputeKernelBase.h"
+#include "OptimusNode_SubGraphReference.h"
 
 #include UE_INLINE_GENERATED_CPP_BY_NAME(OptimusNode_GraphTerminal)
 
@@ -13,6 +17,11 @@
 #define LOCTEXT_NAMESPACE "OptimusNodeGraphTerminal"
 
 
+UOptimusNode_GraphTerminal::UOptimusNode_GraphTerminal()
+{
+	EnableDynamicPins();
+}
+
 FText UOptimusNode_GraphTerminal::GetDisplayName() const
 {
 	switch(TerminalType)
@@ -40,6 +49,11 @@ void UOptimusNode_GraphTerminal::ConstructNode()
 	{
 	case EOptimusTerminalType::Entry:
 		{
+			const FOptimusDataTypeRegistry& TypeRegistry = FOptimusDataTypeRegistry::Get();
+			FOptimusDataTypeRef ComponentSourceType = TypeRegistry.FindType(*UOptimusComponentSourceBinding::StaticClass());
+
+			DefaultComponentPin = AddPinDirect(UOptimusNodeSubGraph::GraphDefaultComponentPinName, EOptimusNodePinDirection::Output, {}, ComponentSourceType);
+			
 			// The input bindings represent the _inputs_ to the graph, so when they are accessed
 			// from within the graph, they get added as _outputs_ to the entry node.
 			for (const FOptimusParameterBinding& Binding: Graph->InputBindings)
@@ -60,6 +74,22 @@ void UOptimusNode_GraphTerminal::ConstructNode()
 	case EOptimusTerminalType::Unknown:
 		checkNoEntry();
 	}
+
+	SubscribeToOwningGraph();
+}
+
+void UOptimusNode_GraphTerminal::PostLoad()
+{
+	Super::PostLoad();
+
+	SubscribeToOwningGraph();
+}
+
+void UOptimusNode_GraphTerminal::BeginDestroy()
+{
+	Super::BeginDestroy();
+
+	UnsubscribeFromOwningGraph();
 }
 
 
@@ -73,7 +103,7 @@ FOptimusRoutedNodePin UOptimusNode_GraphTerminal::GetPinCounterpart(
 		return {};
 	}
 	
-	if (!ensure(!InTraversalContext.ReferenceNesting.IsEmpty()))
+	if (InTraversalContext.ReferenceNesting.IsEmpty())
 	{
 		return {};
 	}
@@ -93,4 +123,244 @@ FOptimusRoutedNodePin UOptimusNode_GraphTerminal::GetPinCounterpart(
 	return Result;
 }
 
+UObject* UOptimusNode_GraphTerminal::GetObjectToShowWhenSelected() const
+{
+	return OwningGraph.Get();
+}
+
+UOptimusComponentSourceBinding* UOptimusNode_GraphTerminal::GetDefaultComponentBinding(const FOptimusPinTraversalContext& InTraversalContext) const
+{
+	if (!ensure(DefaultComponentPin.IsValid()))
+	{
+		return nullptr;
+	}
+	
+	const FOptimusRoutedNodePin PinCounterpart = GetPinCounterpart(DefaultComponentPin.Get(), InTraversalContext);
+
+	if (!PinCounterpart.NodePin)
+	{
+		return nullptr;
+	}
+	
+	const UOptimusNode_SubGraphReference* ReferenceNode = CastChecked<UOptimusNode_SubGraphReference>(PinCounterpart.NodePin->GetOwningNode());
+
+	return ReferenceNode->GetDefaultComponentBinding(PinCounterpart.TraversalContext);
+}
+
+void UOptimusNode_GraphTerminal::SubscribeToOwningGraph()
+{
+	if (ensure(!OwningGraph->GetOnBindingArrayPasted().IsBoundToObject(this)))
+	{
+		OwningGraph->GetOnBindingArrayPasted().AddUObject(this, &UOptimusNode_GraphTerminal::RecreateBindingPins);
+		OwningGraph->GetOnBindingValueChanged().AddUObject(this, &UOptimusNode_GraphTerminal::SyncPinsToBindings);
+		OwningGraph->GetOnBindingArrayItemAdded().AddUObject(this, &UOptimusNode_GraphTerminal::AddPinForNewBinding);
+		OwningGraph->GetOnBindingArrayItemRemoved().AddUObject(this, &UOptimusNode_GraphTerminal::RemoveStalePins);
+		OwningGraph->GetOnBindingArrayCleared().AddUObject(this, &UOptimusNode_GraphTerminal::RemoveStalePins);
+		OwningGraph->GetOnBindingArrayItemMoved().AddUObject(this, &UOptimusNode_GraphTerminal::OnBindingMoved);
+	}
+	
+}
+
+void UOptimusNode_GraphTerminal::UnsubscribeFromOwningGraph() const
+{
+	// Technically we don't need to unsubscribe since the life time of the terminal is the same as that of the owning graph
+	// but we still do it for good measure
+	if (OwningGraph.IsValid())
+	{
+		OwningGraph->GetOnBindingArrayPasted().RemoveAll(this);
+		OwningGraph->GetOnBindingValueChanged().RemoveAll(this);
+		OwningGraph->GetOnBindingArrayItemAdded().RemoveAll(this);
+		OwningGraph->GetOnBindingArrayItemRemoved().RemoveAll(this);
+		OwningGraph->GetOnBindingArrayCleared().RemoveAll(this);
+		OwningGraph->GetOnBindingArrayItemMoved().RemoveAll(this);
+	}
+}
+
+TArray<UOptimusNodePin*> UOptimusNode_GraphTerminal::GetBindingPins()
+{
+	if (TerminalType == EOptimusTerminalType::Entry)
+	{
+		TArray<UOptimusNodePin*> BindingPins = GetPinsByDirection(EOptimusNodePinDirection::Output, false);
+
+		BindingPins.RemoveAt(0);
+
+		return BindingPins;
+	}
+	else if (TerminalType == EOptimusTerminalType::Return)
+	{
+		return GetPinsByDirection(EOptimusNodePinDirection::Input, false);
+	}
+
+	return {};
+}
+
+void UOptimusNode_GraphTerminal::AddPinForNewBinding(FName InBindingArrayPropertyName)
+{
+	FOptimusParameterBindingArray* BindingArrayPtr = nullptr;
+	EOptimusNodePinDirection NewPinDirection = EOptimusNodePinDirection::Unknown;
+	
+	if (InBindingArrayPropertyName == UOptimusNodeSubGraph::InputBindingsPropertyName && TerminalType == EOptimusTerminalType::Entry)
+	{
+		BindingArrayPtr = &OwningGraph->InputBindings;
+		NewPinDirection = EOptimusNodePinDirection::Output;
+	}
+	else if (InBindingArrayPropertyName == UOptimusNodeSubGraph::OutputBindingsPropertyName && TerminalType == EOptimusTerminalType::Return)
+	{
+		BindingArrayPtr = &OwningGraph->OutputBindings;
+		NewPinDirection = EOptimusNodePinDirection::Input;
+	}
+
+	if (BindingArrayPtr)
+	{
+		check(BindingArrayPtr->Num() > 0);
+		AddPin(BindingArrayPtr->Last(), NewPinDirection, nullptr);
+	}
+}
+
+void UOptimusNode_GraphTerminal::RemoveStalePins(FName InBindingArrayPropertyName)
+{
+	FOptimusParameterBindingArray* BindingArrayPtr = nullptr;
+	TArray<UOptimusNodePin*> PinsToCheck = GetBindingPins();
+	if (InBindingArrayPropertyName == UOptimusNodeSubGraph::InputBindingsPropertyName && TerminalType == EOptimusTerminalType::Entry)
+	{
+		BindingArrayPtr = &OwningGraph->InputBindings;
+	}
+	else if (InBindingArrayPropertyName == UOptimusNodeSubGraph::OutputBindingsPropertyName && TerminalType == EOptimusTerminalType::Return)
+	{
+		BindingArrayPtr = &OwningGraph->OutputBindings;
+	}
+
+	if (BindingArrayPtr)
+	{
+		for (UOptimusNodePin* Pin : PinsToCheck)
+		{
+			int32 Index = BindingArrayPtr->IndexOfByPredicate([ PinName = Pin->GetFName()](const FOptimusParameterBinding& InBinding){
+				 return InBinding.Name == PinName;
+			});
+
+			if (Index == INDEX_NONE)
+			{
+				RemovePin(Pin);
+			}
+		}	
+	}
+}
+
+void UOptimusNode_GraphTerminal::OnBindingMoved(FName InBindingArrayPropertyName)
+{
+	FOptimusParameterBindingArray* BindingArrayPtr = nullptr;
+	TArray<UOptimusNodePin*> PinsToCheck = GetBindingPins();
+	if (InBindingArrayPropertyName == UOptimusNodeSubGraph::InputBindingsPropertyName && TerminalType == EOptimusTerminalType::Entry)
+	{
+		BindingArrayPtr = &OwningGraph->InputBindings;
+	}
+	else if (InBindingArrayPropertyName == UOptimusNodeSubGraph::OutputBindingsPropertyName && TerminalType == EOptimusTerminalType::Return)
+	{
+		BindingArrayPtr = &OwningGraph->OutputBindings;
+	}
+
+	if (BindingArrayPtr)
+	{
+		TArray<FName> PinNames;
+
+		for (UOptimusNodePin* Pin : PinsToCheck)
+		{
+			PinNames.Add(Pin->GetFName());
+		}
+
+		TArray<FName> BindingNames;
+		for (const FOptimusParameterBinding& Binding : *BindingArrayPtr)
+		{
+			BindingNames.Add(Binding.Name);
+		}
+
+		FName PinName;
+		FName NextPinName;
+		if (Optimus::FindMovedItemInNameArray(PinNames, BindingNames, PinName, NextPinName))
+		{
+			const int32 PinIndex = PinsToCheck.IndexOfByPredicate([PinName](const UOptimusNodePin* InPin)
+			{
+				return InPin->GetFName() == PinName;
+			});
+		
+			const int32 NextPinIndex = NextPinName == NAME_None ?
+				INDEX_NONE : PinsToCheck.IndexOfByPredicate([NextPinName](const UOptimusNodePin* InPin)
+			{
+				return InPin->GetFName() == NextPinName;
+			});
+
+			MovePin(PinsToCheck[PinIndex], NextPinIndex != INDEX_NONE ? PinsToCheck[NextPinIndex] : nullptr);
+		}
+	}
+}
+
+void UOptimusNode_GraphTerminal::RecreateBindingPins(FName InBindingArrayPropertyName)
+{
+	FOptimusParameterBindingArray* BindingArrayPtr = nullptr;
+	EOptimusNodePinDirection NewPinDirection = EOptimusNodePinDirection::Unknown;
+	TArray<UOptimusNodePin*> PinsToRemove = GetBindingPins();
+	
+	if (InBindingArrayPropertyName == UOptimusNodeSubGraph::InputBindingsPropertyName && TerminalType == EOptimusTerminalType::Entry)
+	{
+		BindingArrayPtr = &OwningGraph->InputBindings;
+		NewPinDirection = EOptimusNodePinDirection::Output;
+	}
+	else if (InBindingArrayPropertyName == UOptimusNodeSubGraph::OutputBindingsPropertyName && TerminalType == EOptimusTerminalType::Return)
+	{
+		BindingArrayPtr = &OwningGraph->OutputBindings;
+		NewPinDirection = EOptimusNodePinDirection::Input;
+	}
+
+	if (BindingArrayPtr)
+	{
+		for (UOptimusNodePin* Pin : PinsToRemove)
+		{
+			RemovePin(Pin);
+		}
+
+		for (const FOptimusParameterBinding& Binding : *BindingArrayPtr)
+		{
+			AddPin(Binding, NewPinDirection, nullptr);
+		}
+	}	
+}
+
+void UOptimusNode_GraphTerminal::SyncPinsToBindings(FName InBindingArrayPropertyName)
+{
+	FOptimusParameterBindingArray* BindingArrayPtr = nullptr;
+
+	TArray<UOptimusNodePin*> BindingPins = GetBindingPins();
+	if (InBindingArrayPropertyName == UOptimusNodeSubGraph::InputBindingsPropertyName && TerminalType == EOptimusTerminalType::Entry)
+	{
+		BindingArrayPtr = &OwningGraph->InputBindings;
+	}
+	else if (InBindingArrayPropertyName == UOptimusNodeSubGraph::OutputBindingsPropertyName && TerminalType == EOptimusTerminalType::Return)
+	{
+		BindingArrayPtr = &OwningGraph->OutputBindings;
+	}
+
+	if (BindingArrayPtr)
+	{
+		check(BindingArrayPtr->Num() == BindingPins.Num());
+		for (int32 Index = 0 ; Index < BindingArrayPtr->Num(); Index++)
+		{
+			UOptimusNodePin* BindingPin = BindingPins[Index];
+			const FOptimusParameterBinding& Binding = (*BindingArrayPtr)[Index];
+			if (BindingPin->GetFName() != Binding.Name)
+			{
+				SetPinName(BindingPin, Binding.Name);
+			}
+			if (BindingPin->GetDataType() != Binding.DataType.Resolve())
+			{
+				SetPinDataType(BindingPin, Binding.DataType);
+			}
+			if (BindingPin->GetDataDomain() != Binding.DataDomain)
+			{
+				SetPinDataDomain(BindingPin, Binding.DataDomain);
+			}
+		}
+	}	
+}
+
+
 #undef LOCTEXT_NAMESPACE
diff --git a/Engine/Plugins/Animation/DeformerGraph/Source/OptimusCore/Private/Nodes/OptimusNode_GraphTerminal.h b/Engine/Plugins/Animation/DeformerGraph/Source/OptimusCore/Private/Nodes/OptimusNode_GraphTerminal.h
--- a/Engine/Plugins/Animation/DeformerGraph/Source/OptimusCore/Private/Nodes/OptimusNode_GraphTerminal.h
+++ b/Engine/Plugins/Animation/DeformerGraph/Source/OptimusCore/Private/Nodes/OptimusNode_GraphTerminal.h
@@ -1,7 +1,9 @@
 ﻿// Copyright Epic Games, Inc. All Rights Reserved.
 #pragma once
 
+#include "IOptimusAlternativeSelectedObjectProvider.h"
 #include "IOptimusNodePinRouter.h"
+#include "IOptimusNonCollapsibleNode.h"
 #include "OptimusNode.h"
 
 #include "OptimusNode_GraphTerminal.generated.h"
@@ -22,26 +24,51 @@ enum class EOptimusTerminalType
 UCLASS(Hidden)
 class UOptimusNode_GraphTerminal :
 	public UOptimusNode,
-	public IOptimusNodePinRouter
+	public IOptimusNodePinRouter,
+	public IOptimusAlternativeSelectedObjectProvider,
+	public IOptimusNonCollapsibleNode
 {
 	GENERATED_BODY()
 	
 public:
+	UOptimusNode_GraphTerminal();
+	
 	// UOptimusNode overrides
 	bool CanUserDeleteNode() const override { return false; }
 	FName GetNodeCategory() const override { return NAME_None; }
 	FText GetDisplayName() const override;
 	void ConstructNode() override;
+
+	// UObject overrides
+	void PostLoad() override;
+	void BeginDestroy() override;
 	
 	// IOptimusNodePinRouter implementation
 	FOptimusRoutedNodePin GetPinCounterpart(
 		UOptimusNodePin* InNodePin,
 		const FOptimusPinTraversalContext& InTraversalContext
 	) const override;
+	
+	// IOptimusAlternativeSelectedObjectProvider
+	UObject* GetObjectToShowWhenSelected() const override;
+
+	UOptimusComponentSourceBinding* GetDefaultComponentBinding(const FOptimusPinTraversalContext& InTraversalContext) const;
 
 protected:
 	friend class UOptimusNodeGraph;
+	friend class UOptimusNodeSubGraph;
+
+	void SubscribeToOwningGraph();
 
+	void UnsubscribeFromOwningGraph() const;
+	TArray<UOptimusNodePin*> GetBindingPins();
+	
+	void AddPinForNewBinding(FName InBindingArrayPropertyName);
+	void RemoveStalePins(FName InBindingArrayPropertyName);
+	void OnBindingMoved(FName InBindingArrayPropertyName);
+	void RecreateBindingPins(FName InBindingArrayPropertyName);
+	void SyncPinsToBindings(FName InBindingArrayPropertyName);	
+	
 	/** Indicates whether this is an entry or a return terminal node */
 	UPROPERTY()
 	EOptimusTerminalType TerminalType;
@@ -51,4 +78,7 @@ protected:
 	 */
 	UPROPERTY()
 	TWeakObjectPtr<UOptimusNodeSubGraph> OwningGraph;
+
+	UPROPERTY()
+	TWeakObjectPtr<UOptimusNodePin> DefaultComponentPin;
 };
diff --git a/Engine/Plugins/Animation/DeformerGraph/Source/OptimusCore/Private/Nodes/OptimusNode_LoopTerminal.cpp b/Engine/Plugins/Animation/DeformerGraph/Source/OptimusCore/Private/Nodes/OptimusNode_LoopTerminal.cpp
--- a/Engine/Plugins/Animation/DeformerGraph/Source/OptimusCore/Private/Nodes/OptimusNode_LoopTerminal.cpp
+++ b/Engine/Plugins/Animation/DeformerGraph/Source/OptimusCore/Private/Nodes/OptimusNode_LoopTerminal.cpp
@@ -78,6 +78,8 @@ FText UOptimusNode_LoopTerminal::GetDisplayName() const
 
 void UOptimusNode_LoopTerminal::ConstructNode()
 {
+	check(GetPins().IsEmpty());
+	
 	if (TerminalType == EOptimusTerminalType::Entry)
 	{
 		FOptimusDataTypeHandle UintType = FOptimusDataTypeRegistry::Get().FindType(*FUInt32Property::StaticClass());
@@ -85,15 +87,7 @@ void UOptimusNode_LoopTerminal::ConstructNode()
 		CountPin = AddPinDirect(TEXT("Count"),  EOptimusNodePinDirection::Output, FOptimusDataDomain(), UintType);
 	}
 
-	
-	if (FOptimusLoopTerminalInfo* LoopInfoPtr = GetLoopInfo())
-	{
-		PinPairInfos.Reset();
-		for (const FOptimusParameterBinding& Binding : LoopInfoPtr->Bindings)
-		{
-			AddPinPairsDirect(Binding);
-		}
-	}
+	PinPairInfos.Reset();
 }
 
 bool UOptimusNode_LoopTerminal::ValidateConnection(const UOptimusNodePin& InThisNodesPin, const UOptimusNodePin& InOtherNodesPin, FString* OutReason) const
@@ -235,10 +229,38 @@ EOptimusPinMutability UOptimusNode_LoopTerminal::GetOutputPinMutability(const UO
 
 void UOptimusNode_LoopTerminal::PairToCounterpartNode(const IOptimusNodePairProvider* NodePairProvider)
 {
-	if (TerminalType == EOptimusTerminalType::Return)
+	if (FOptimusLoopTerminalInfo* LoopInfoPtr = GetLoopInfo(); ensure(LoopInfoPtr))
 	{
-		ConstructNode();
+		for (const FOptimusParameterBinding& Binding : LoopInfoPtr->Bindings)
+		{
+			AddPinPairsDirect(Binding);
+		}
+	}	
+}
+
+FString UOptimusNode_LoopTerminal::GetBindingDeclaration(FName BindingName) const
+{
+	return {};
+}
+
+bool UOptimusNode_LoopTerminal::GetBindingSupportAtomicCheckBoxVisibility(FName BindingName) const
+{
+	return false;
+}
+
+bool UOptimusNode_LoopTerminal::GetBindingSupportReadCheckBoxVisibility(FName BindingName) const
+{
+	return false;
+}
+
+EOptimusDataTypeUsageFlags UOptimusNode_LoopTerminal::GetTypeUsageFlags(const FOptimusDataDomain& InDataDomain) const
+{
+	if (ensure(!InDataDomain.IsSingleton()))
+	{
+		return EOptimusDataTypeUsageFlags::Resource;
 	}
+
+	return EOptimusDataTypeUsageFlags::None;
 }
 
 UOptimusNodePin* UOptimusNode_LoopTerminal::GetPinCounterpart(const UOptimusNodePin* InNodePin, EOptimusTerminalType InTerminalType, TOptional<EOptimusNodePinDirection> InDirection) const
@@ -416,14 +438,17 @@ void UOptimusNode_LoopTerminal::ClearPinPairs()
 
 void UOptimusNode_LoopTerminal::MovePinPair()
 {
-	FName PinNameToMove = NAME_None;
-	int32 DivergeIndex = INDEX_NONE;
-
-	TArray<TArray<UOptimusNodePin*>> PinPairs;
 
 	FOptimusParameterBindingArray& Bindings = GetLoopInfo()->Bindings;
 
-	check(Bindings.Num() == PinPairInfos.Num())
+	TArray<FName> BindingNames;
+	for (FOptimusParameterBinding& Binding : Bindings)
+	{
+		BindingNames.Add(Binding.Name);
+	}
+
+	check(PinPairInfos.Num() == Bindings.Num());
+	TArray<TArray<UOptimusNodePin*>> PinPairs;
 	
 	for (const FOptimusPinPairInfo& PinPairInfo : PinPairInfos)
 	{
@@ -435,56 +460,38 @@ void UOptimusNode_LoopTerminal::MovePinPair()
 	{
 		NamedPinIndex = 1;
 	}
+
+	TArray<FName> PinNames;
 	
-	for (int32 Index = 0; Index < Bindings.Num(); Index++)
+	for (int32 Index = 0; Index < PinPairs.Num(); Index++)
 	{
-		if (Bindings[Index].Name != PinPairs[Index][NamedPinIndex]->GetFName() && DivergeIndex == INDEX_NONE)
-		{
-			DivergeIndex = Index;
-			continue;
-		}
-
-		if (DivergeIndex != INDEX_NONE)
-		{
-			if (Bindings[DivergeIndex].Name == PinPairs[Index][NamedPinIndex]->GetFName())
-			{
-				PinNameToMove = PinPairs[DivergeIndex][NamedPinIndex]->GetFName();
-			}
-			else if (ensure(Bindings[Index].Name == PinPairs[DivergeIndex][NamedPinIndex]->GetFName()))
-			{
-				PinNameToMove = Bindings[DivergeIndex].Name;
-			}
-			break;
-		}
+		const UOptimusNodePin* Pin = PinPairs[Index][NamedPinIndex];
+		PinNames.Add({Pin->GetFName()});
 	}
 
-	if (DivergeIndex != INDEX_NONE)
+
+	FName PinName = NAME_None;
+	FName NextPinName = NAME_None;
+
+	if (Optimus::FindMovedItemInNameArray(PinNames, BindingNames, PinName, NextPinName))
 	{
-		const int32 PinIndex = PinPairs.IndexOfByPredicate([PinNameToMove, NamedPinIndex](const TArray<UOptimusNodePin*>& InPinPair)
+		const int32 PinIndex = PinPairs.IndexOfByPredicate([PinName, NamedPinIndex](const TArray<UOptimusNodePin*>& InPinPair)
 		{
-			return InPinPair[NamedPinIndex]->GetFName() == PinNameToMove;
+			return InPinPair[NamedPinIndex]->GetFName() == PinName;
 		});
 
-		const int32 NameIndex = Bindings.InnerArray.IndexOfByPredicate([PinNameToMove](const FOptimusParameterBinding& InBinding)
-		{
-			return InBinding.Name == PinNameToMove;
-		}); 
-
-		const int32 NextNameIndex = NameIndex + 1;
-		FName NextName = Bindings.IsValidIndex(NextNameIndex) ?  FName(Bindings[NextNameIndex].Name) : NAME_None;
 		int32 NextPinIndex = INDEX_NONE;
-		if (!NextName.IsNone())
+		if (!NextPinName.IsNone())
 		{
-			NextPinIndex = PinPairs.IndexOfByPredicate([NextName, NamedPinIndex](const TArray<UOptimusNodePin*>& InPinPair)
+			NextPinIndex = PinPairs.IndexOfByPredicate([NextPinName, NamedPinIndex](const TArray<UOptimusNodePin*>& InPinPair)
 			{
-				return InPinPair[NamedPinIndex]->GetFName() == NextName;
-			});	
+				return InPinPair[NamedPinIndex]->GetFName() == NextPinName;
+			});
 		}
 
 		MovePin(PinPairs[PinIndex][0], NextPinIndex != INDEX_NONE ? PinPairs[NextPinIndex][0]: nullptr);
 		MovePin(PinPairs[PinIndex][1], NextPinIndex != INDEX_NONE ? PinPairs[NextPinIndex][1]: IndexPin);
 
-
 		FOptimusPinPairInfo Pair = PinPairInfos[PinIndex];
 
 		FOptimusPinPairInfo PairBefore = NextPinIndex != INDEX_NONE ? PinPairInfos[NextPinIndex] : FOptimusPinPairInfo();
diff --git a/Engine/Plugins/Animation/DeformerGraph/Source/OptimusCore/Private/Nodes/OptimusNode_LoopTerminal.h b/Engine/Plugins/Animation/DeformerGraph/Source/OptimusCore/Private/Nodes/OptimusNode_LoopTerminal.h
--- a/Engine/Plugins/Animation/DeformerGraph/Source/OptimusCore/Private/Nodes/OptimusNode_LoopTerminal.h
+++ b/Engine/Plugins/Animation/DeformerGraph/Source/OptimusCore/Private/Nodes/OptimusNode_LoopTerminal.h
@@ -3,6 +3,7 @@
 
 #include "IOptimusNodeAdderPinProvider.h"
 #include "IOptimusNodePairProvider.h"
+#include "IOptimusParameterBindingProvider.h"
 #include "IOptimusPinMutabilityDefiner.h"
 #include "IOptimusUnnamedNodePinProvider.h"
 #include "OptimusBindingTypes.h"
@@ -51,7 +52,8 @@ class OPTIMUSCORE_API UOptimusNode_LoopTerminal :
 	public IOptimusNodeAdderPinProvider,
 	public IOptimusUnnamedNodePinProvider,
 	public IOptimusNodePairProvider,
-	public IOptimusPinMutabilityDefiner
+	public IOptimusPinMutabilityDefiner,
+	public IOptimusParameterBindingProvider
 {
 	GENERATED_BODY()
 	
@@ -84,6 +86,12 @@ public:
 
 	// IOptimusNodePairProvider
 	void PairToCounterpartNode(const IOptimusNodePairProvider* NodePairProvider) override;
+
+	// IOptimusParameterBindingProvider
+	FString GetBindingDeclaration(FName BindingName) const;
+	bool GetBindingSupportAtomicCheckBoxVisibility(FName BindingName) const;
+	bool GetBindingSupportReadCheckBoxVisibility(FName BindingName) const;
+	EOptimusDataTypeUsageFlags GetTypeUsageFlags(const FOptimusDataDomain& InDataDomain) const;
 	
 	UOptimusNodePin* GetPinCounterpart(const UOptimusNodePin* InNodePin, EOptimusTerminalType InTerminalType, TOptional<EOptimusNodePinDirection> InDirection = {}) const;
 	UOptimusNode_LoopTerminal* GetOtherTerminal() const;
@@ -103,7 +111,6 @@ protected:
 	UPROPERTY(EditAnywhere, Category="Loop Terminal", DisplayName="Settings", meta=(FullyExpand="true", EditCondition="TerminalType==EOptimusTerminalType::Entry", EditConditionHides))
 	FOptimusLoopTerminalInfo LoopInfo;
 
-
 	UPROPERTY()
 	UOptimusNodePin* IndexPin;
 	
diff --git a/Engine/Plugins/Animation/DeformerGraph/Source/OptimusCore/Private/Nodes/OptimusNode_ResourceAccessorBase.cpp b/Engine/Plugins/Animation/DeformerGraph/Source/OptimusCore/Private/Nodes/OptimusNode_ResourceAccessorBase.cpp
--- a/Engine/Plugins/Animation/DeformerGraph/Source/OptimusCore/Private/Nodes/OptimusNode_ResourceAccessorBase.cpp
+++ b/Engine/Plugins/Animation/DeformerGraph/Source/OptimusCore/Private/Nodes/OptimusNode_ResourceAccessorBase.cpp
@@ -196,7 +196,7 @@ FName UOptimusNode_ResourceAccessorBase::GetResourcePinName(int32 InPinIndex, FN
 }
 
 
-TOptional<FText> UOptimusNode_ResourceAccessorBase::ValidateForCompile() const
+TOptional<FText> UOptimusNode_ResourceAccessorBase::ValidateForCompile(const FOptimusPinTraversalContext& InContext) const
 {
 	const UOptimusResourceDescription* ResourceDescription = GetResourceDescription();
 	if (!ResourceDescription)
@@ -236,7 +236,7 @@ UOptimusComputeDataInterface* UOptimusNode_ResourceAccessorBase::GetDataInterfac
 	return Description->DataInterface;
 }
 
-UOptimusComponentSourceBinding* UOptimusNode_ResourceAccessorBase::GetComponentBinding() const
+UOptimusComponentSourceBinding* UOptimusNode_ResourceAccessorBase::GetComponentBinding(const FOptimusPinTraversalContext& InContext) const
 {
 	if (const UOptimusResourceDescription* Description = ResourceDesc.Get())
 	{
diff --git a/Engine/Plugins/Animation/DeformerGraph/Source/OptimusCore/Private/Nodes/OptimusNode_ResourceAccessorBase.h b/Engine/Plugins/Animation/DeformerGraph/Source/OptimusCore/Private/Nodes/OptimusNode_ResourceAccessorBase.h
--- a/Engine/Plugins/Animation/DeformerGraph/Source/OptimusCore/Private/Nodes/OptimusNode_ResourceAccessorBase.h
+++ b/Engine/Plugins/Animation/DeformerGraph/Source/OptimusCore/Private/Nodes/OptimusNode_ResourceAccessorBase.h
@@ -5,6 +5,7 @@
 #include "DataInterfaces/OptimusDataInterfaceRawBuffer.h"
 #include "IOptimusComponentBindingProvider.h"
 #include "IOptimusDataInterfaceProvider.h"
+#include "IOptimusNonCollapsibleNode.h"
 #include "IOptimusPinMutabilityDefiner.h"
 #include "OptimusNode.h"
 
@@ -36,7 +37,8 @@ class UOptimusNode_ResourceAccessorBase :
 	public UOptimusNode,
 	public IOptimusDataInterfaceProvider,
 	public IOptimusComponentBindingProvider,
-	public IOptimusPinMutabilityDefiner
+	public IOptimusPinMutabilityDefiner,
+	public IOptimusNonCollapsibleNode
 {
 	GENERATED_BODY()
 
@@ -54,13 +56,13 @@ public:
 		return CategoryName::Resources;
 	}
 
-	TOptional<FText> ValidateForCompile() const override;
+	TOptional<FText> ValidateForCompile(const FOptimusPinTraversalContext& InContext) const override;
 
 	// IOptimusDataInterfaceProvider implementations
 	UOptimusComputeDataInterface* GetDataInterface(UObject *InOuter) const override;
 	int32 GetDataFunctionIndexFromPin(const UOptimusNodePin* InPin) const override { return INDEX_NONE; }
 	// Also IOptimusComponentBindingProvider implementation
-	UOptimusComponentSourceBinding* GetComponentBinding() const override;
+	UOptimusComponentSourceBinding* GetComponentBinding(const FOptimusPinTraversalContext& InContext) const override;
 
 	PRAGMA_DISABLE_DEPRECATION_WARNINGS
 	EOptimusBufferWriteType GetDeprecatedBufferWriteType() const { return WriteType_DEPRECATED; }
diff --git a/Engine/Plugins/Animation/DeformerGraph/Source/OptimusCore/Private/Nodes/OptimusNode_SubGraphReference.cpp b/Engine/Plugins/Animation/DeformerGraph/Source/OptimusCore/Private/Nodes/OptimusNode_SubGraphReference.cpp
--- a/Engine/Plugins/Animation/DeformerGraph/Source/OptimusCore/Private/Nodes/OptimusNode_SubGraphReference.cpp
+++ b/Engine/Plugins/Animation/DeformerGraph/Source/OptimusCore/Private/Nodes/OptimusNode_SubGraphReference.cpp
@@ -5,16 +5,25 @@
 #include "OptimusNodePin.h"
 #include "OptimusNodeSubGraph.h"
 #include "OptimusBindingTypes.h"
+#include "OptimusComponentSource.h"
+#include "OptimusDataTypeRegistry.h"
+#include "OptimusDeformer.h"
+#include "OptimusHelpers.h"
 #include "OptimusNode_GraphTerminal.h"
 
 #include UE_INLINE_GENERATED_CPP_BY_NAME(OptimusNode_SubGraphReference)
 
 
+UOptimusNode_SubGraphReference::UOptimusNode_SubGraphReference()
+{
+	EnableDynamicPins();
+}
+
 FText UOptimusNode_SubGraphReference::GetDisplayName() const
 {
-	if (SubGraph)
+	if (SubGraph.Get())
 	{
-		return FText::FromString(SubGraph->GetName());
+		return FText::FromString(SubGraph.Get()->GetName());
 	}
 
 	return FText::FromString("<graph missing>");
@@ -23,8 +32,12 @@ FText UOptimusNode_SubGraphReference::GetDisplayName() const
 
 void UOptimusNode_SubGraphReference::ConstructNode()
 {
-	if (ensure(SubGraph))
+	if (ensure(SubGraph.Get()))
 	{
+		const FOptimusDataTypeRegistry& TypeRegistry = FOptimusDataTypeRegistry::Get();
+		FOptimusDataTypeRef ComponentSourceType = TypeRegistry.FindType(*UOptimusComponentSourceBinding::StaticClass());
+		DefaultComponentPin = AddPinDirect(UOptimusNodeSubGraph::GraphDefaultComponentPinName, EOptimusNodePinDirection::Input, {}, ComponentSourceType);
+		
 		// After a duplicate, the kernel node has no pins, so we need to reconstruct them from
 		// the bindings. We can assume that all naming clashes have already been dealt with.
 		for (const FOptimusParameterBinding& Binding: SubGraph->InputBindings)
@@ -35,9 +48,25 @@ void UOptimusNode_SubGraphReference::ConstructNode()
 		{
 			AddPinDirect(Binding, EOptimusNodePinDirection::Output);
 		}
+
+		SubscribeToSubGraph();
 	}
 }
 
+void UOptimusNode_SubGraphReference::PostLoad()
+{
+	Super::PostLoad();
+
+	SubscribeToSubGraph();
+}
+
+void UOptimusNode_SubGraphReference::BeginDestroy()
+{
+	Super::BeginDestroy();
+
+	UnsubscribeFromSubGraph();
+}
+
 
 FOptimusRoutedNodePin UOptimusNode_SubGraphReference::GetPinCounterpart(
 	UOptimusNodePin* InNodePin,
@@ -49,7 +78,7 @@ FOptimusRoutedNodePin UOptimusNode_SubGraphReference::GetPinCounterpart(
 		return {};
 	}
 
-	if (!ensure(SubGraph))
+	if (!ensure(SubGraph.Get()))
 	{
 		return {};
 	}
@@ -77,3 +106,259 @@ FOptimusRoutedNodePin UOptimusNode_SubGraphReference::GetPinCounterpart(
 
 	return Result;
 }
+
+UOptimusNodeGraph* UOptimusNode_SubGraphReference::GetNodeGraphToShow()
+{
+	return SubGraph.Get();
+}
+
+UOptimusComponentSourceBinding* UOptimusNode_SubGraphReference::GetDefaultComponentBinding(const FOptimusPinTraversalContext& InTraversalContext) const
+{
+	if (!ensure(DefaultComponentPin.IsValid()))
+	{
+		return nullptr;
+	}
+	
+	const UOptimusNodeGraph* OwningGraph = GetOwningGraph();
+	TSet<UOptimusComponentSourceBinding*> Bindings = OwningGraph->GetComponentSourceBindingsForPin(DefaultComponentPin.Get(), InTraversalContext);
+	
+	if (!Bindings.IsEmpty() && ensure(Bindings.Num() == 1))
+	{
+		return Bindings.Array()[0];
+	}
+
+	// Default to the primary binding, but only if we're at the top-most level of the graph.
+	if (const UOptimusDeformer* Deformer = Cast<UOptimusDeformer>(OwningGraph->GetCollectionOwner()))
+	{
+		return Deformer->GetPrimaryComponentBinding();
+	}
+
+	if (const UOptimusNodeSubGraph* OwningSubGraph = Cast<UOptimusNodeSubGraph>(OwningGraph))
+	{
+		return OwningSubGraph->GetDefaultComponentBinding(InTraversalContext);
+	}
+
+	return nullptr;	
+}
+
+void UOptimusNode_SubGraphReference::SubscribeToSubGraph()
+{
+	if (ensure(!SubGraph->GetOnBindingArrayPasted().IsBoundToObject(this)))
+	{
+		SubGraph->GetOnBindingArrayPasted().AddUObject(this, &UOptimusNode_SubGraphReference::RecreateBindingPins);
+		SubGraph->GetOnBindingValueChanged().AddUObject(this, &UOptimusNode_SubGraphReference::SyncPinsToBindings);
+		SubGraph->GetOnBindingArrayItemAdded().AddUObject(this, &UOptimusNode_SubGraphReference::AddPinForNewBinding);
+		SubGraph->GetOnBindingArrayItemRemoved().AddUObject(this, &UOptimusNode_SubGraphReference::RemoveStalePins);
+		SubGraph->GetOnBindingArrayCleared().AddUObject(this, &UOptimusNode_SubGraphReference::RemoveStalePins);
+		SubGraph->GetOnBindingArrayItemMoved().AddUObject(this, &UOptimusNode_SubGraphReference::OnBindingMoved);	
+	}
+	
+}
+
+void UOptimusNode_SubGraphReference::UnsubscribeFromSubGraph() const
+{
+	if (SubGraph.IsValid())
+	{
+		SubGraph->GetOnBindingArrayPasted().RemoveAll(this);
+		SubGraph->GetOnBindingValueChanged().RemoveAll(this);
+		SubGraph->GetOnBindingArrayItemAdded().RemoveAll(this);
+		SubGraph->GetOnBindingArrayItemRemoved().RemoveAll(this);
+		SubGraph->GetOnBindingArrayCleared().RemoveAll(this);
+		SubGraph->GetOnBindingArrayItemMoved().RemoveAll(this);
+	}
+}
+
+void UOptimusNode_SubGraphReference::AddPinForNewBinding(FName InBindingArrayPropertyName)
+{
+	FOptimusParameterBindingArray* BindingArrayPtr = nullptr;
+	EOptimusNodePinDirection NewPinDirection = EOptimusNodePinDirection::Unknown;
+	UOptimusNodePin* BeforePin = nullptr;
+	
+	if (InBindingArrayPropertyName == UOptimusNodeSubGraph::InputBindingsPropertyName)
+	{
+		BindingArrayPtr = &SubGraph->InputBindings;
+		NewPinDirection = EOptimusNodePinDirection::Input;
+		TArray<UOptimusNodePin*> OutputPins = GetPinsByDirection(EOptimusNodePinDirection::Output, false);
+		BeforePin = OutputPins.Num() > 0 ? OutputPins[0] : nullptr;
+	}
+	else if (InBindingArrayPropertyName == UOptimusNodeSubGraph::OutputBindingsPropertyName)
+	{
+		BindingArrayPtr = &SubGraph->OutputBindings;
+		NewPinDirection = EOptimusNodePinDirection::Output;
+		BeforePin = nullptr;
+	}
+
+	if (BindingArrayPtr)
+	{
+		check(BindingArrayPtr->Num() > 0);
+		AddPin(BindingArrayPtr->Last(), NewPinDirection, BeforePin);
+	}
+	
+}
+
+void UOptimusNode_SubGraphReference::RemoveStalePins(FName InBindingArrayPropertyName)
+{
+	FOptimusParameterBindingArray* BindingArrayPtr = nullptr;
+	TArray<UOptimusNodePin*> PinsToCheck;
+	if (InBindingArrayPropertyName == UOptimusNodeSubGraph::InputBindingsPropertyName)
+	{
+		BindingArrayPtr = &SubGraph->InputBindings;
+		PinsToCheck = GetBindingPinsByDirection(EOptimusNodePinDirection::Input);
+	}
+	else if (InBindingArrayPropertyName == UOptimusNodeSubGraph::OutputBindingsPropertyName)
+	{
+		BindingArrayPtr = &SubGraph->OutputBindings;
+		PinsToCheck = GetBindingPinsByDirection(EOptimusNodePinDirection::Output);
+	}
+
+	for (UOptimusNodePin* Pin : PinsToCheck)
+	{
+		int32 Index = BindingArrayPtr->IndexOfByPredicate([ PinName = Pin->GetFName()](const FOptimusParameterBinding& InBinding){
+			 return InBinding.Name == PinName;
+		});
+
+		if (Index == INDEX_NONE)
+		{
+			RemovePin(Pin);
+		}
+	}
+}
+
+void UOptimusNode_SubGraphReference::OnBindingMoved(FName InBindingArrayPropertyName)
+{
+	FOptimusParameterBindingArray* BindingArrayPtr = nullptr;
+	TArray<UOptimusNodePin*> PinsToCheck;
+	if (InBindingArrayPropertyName == UOptimusNodeSubGraph::InputBindingsPropertyName)
+	{
+		BindingArrayPtr = &SubGraph->InputBindings;
+		PinsToCheck = GetBindingPinsByDirection(EOptimusNodePinDirection::Input);
+	}
+	else if (InBindingArrayPropertyName == UOptimusNodeSubGraph::OutputBindingsPropertyName)
+	{
+		BindingArrayPtr = &SubGraph->OutputBindings;
+		PinsToCheck = GetBindingPinsByDirection(EOptimusNodePinDirection::Output);
+	}
+
+	if (BindingArrayPtr)
+	{
+		TArray<FName> PinNames;
+
+		for (UOptimusNodePin* Pin : PinsToCheck)
+		{
+			PinNames.Add(Pin->GetFName());
+		}
+
+		TArray<FName> BindingNames;
+		for (const FOptimusParameterBinding& Binding : *BindingArrayPtr)
+		{
+			BindingNames.Add(Binding.Name);
+		}
+
+		FName PinName;
+		FName NextPinName;
+		if (Optimus::FindMovedItemInNameArray(PinNames, BindingNames, PinName, NextPinName))
+		{
+			const int32 PinIndex = PinsToCheck.IndexOfByPredicate([PinName](const UOptimusNodePin* InPin)
+			{
+				return InPin->GetFName() == PinName;
+			});
+		
+			const int32 NextPinIndex = NextPinName == NAME_None ?
+				INDEX_NONE : PinsToCheck.IndexOfByPredicate([NextPinName](const UOptimusNodePin* InPin)
+			{
+				return InPin->GetFName() == NextPinName;
+			});
+
+			MovePin(PinsToCheck[PinIndex], NextPinIndex != INDEX_NONE ? PinsToCheck[NextPinIndex] : nullptr);
+		}
+	}
+}
+
+void UOptimusNode_SubGraphReference::RecreateBindingPins(FName InBindingArrayPropertyName)
+{
+	FOptimusParameterBindingArray* BindingArrayPtr = nullptr;
+	EOptimusNodePinDirection NewPinDirection = EOptimusNodePinDirection::Unknown;
+	TArray<UOptimusNodePin*> PinsToRemove;
+	UOptimusNodePin* BeforePin = nullptr;
+	
+	if (InBindingArrayPropertyName == UOptimusNodeSubGraph::InputBindingsPropertyName)
+	{
+		BindingArrayPtr = &SubGraph->InputBindings;
+		NewPinDirection = EOptimusNodePinDirection::Input;
+		PinsToRemove = GetBindingPinsByDirection(EOptimusNodePinDirection::Input);
+		
+		TArray<UOptimusNodePin*> OutputPins = GetPinsByDirection(EOptimusNodePinDirection::Output, false);
+		BeforePin = OutputPins.Num() > 0 ? OutputPins[0] : nullptr;
+	}
+	else if (InBindingArrayPropertyName == UOptimusNodeSubGraph::OutputBindingsPropertyName)
+	{
+		BindingArrayPtr = &SubGraph->OutputBindings;
+		NewPinDirection = EOptimusNodePinDirection::Output;
+		PinsToRemove = GetBindingPinsByDirection(EOptimusNodePinDirection::Output);
+		
+		BeforePin = nullptr;
+	}
+
+	if (BindingArrayPtr)
+	{
+		for (UOptimusNodePin* Pin : PinsToRemove)
+		{
+			RemovePin(Pin);
+		}
+
+		for (const FOptimusParameterBinding& Binding : *BindingArrayPtr)
+		{
+			AddPin(Binding, NewPinDirection, BeforePin);
+		}
+	}	
+}
+
+void UOptimusNode_SubGraphReference::SyncPinsToBindings(FName InBindingArrayPropertyName)
+{
+	FOptimusParameterBindingArray* BindingArrayPtr = nullptr;
+
+	TArray<UOptimusNodePin*> BindingPins;
+	if (InBindingArrayPropertyName == UOptimusNodeSubGraph::InputBindingsPropertyName)
+	{
+		BindingArrayPtr = &SubGraph->InputBindings;
+		BindingPins = GetBindingPinsByDirection(EOptimusNodePinDirection::Input);
+	}
+	else if (InBindingArrayPropertyName == UOptimusNodeSubGraph::OutputBindingsPropertyName)
+	{
+		BindingArrayPtr = &SubGraph->OutputBindings;
+		BindingPins = GetBindingPinsByDirection(EOptimusNodePinDirection::Output);
+	}
+
+	if (BindingArrayPtr)
+	{
+		check(BindingArrayPtr->Num() == BindingPins.Num());
+		for (int32 Index = 0 ; Index < BindingArrayPtr->Num(); Index++)
+		{
+			UOptimusNodePin* BindingPin = BindingPins[Index];
+			const FOptimusParameterBinding& Binding = (*BindingArrayPtr)[Index];
+			if (BindingPin->GetFName() != Binding.Name)
+			{
+				SetPinName(BindingPin, Binding.Name);
+			}
+			if (BindingPin->GetDataType() != Binding.DataType.Resolve())
+			{
+				SetPinDataType(BindingPin, Binding.DataType);
+			}
+			if (BindingPin->GetDataDomain() != Binding.DataDomain)
+			{
+				SetPinDataDomain(BindingPin, Binding.DataDomain);
+			}
+		}
+	}		
+}
+
+TArray<UOptimusNodePin*> UOptimusNode_SubGraphReference::GetBindingPinsByDirection(EOptimusNodePinDirection InDirection)
+{
+	TArray<UOptimusNodePin*> Results = GetPinsByDirection(InDirection, false);
+	
+	Results.Remove(DefaultComponentPin.Get());
+
+	return Results;
+}
+
+
diff --git a/Engine/Plugins/Animation/DeformerGraph/Source/OptimusCore/Private/Nodes/OptimusNode_SubGraphReference.h b/Engine/Plugins/Animation/DeformerGraph/Source/OptimusCore/Private/Nodes/OptimusNode_SubGraphReference.h
--- a/Engine/Plugins/Animation/DeformerGraph/Source/OptimusCore/Private/Nodes/OptimusNode_SubGraphReference.h
+++ b/Engine/Plugins/Animation/DeformerGraph/Source/OptimusCore/Private/Nodes/OptimusNode_SubGraphReference.h
@@ -2,6 +2,7 @@
 
 #pragma once
 
+#include "IOptimusNodeGraphProvider.h"
 #include "IOptimusNodePinRouter.h"
 #include "OptimusNode.h"
 
@@ -14,28 +15,54 @@ class UOptimusNodeSubGraph;
 UCLASS(Hidden)
 class UOptimusNode_SubGraphReference :
 	public UOptimusNode,
-	public IOptimusNodePinRouter
+	public IOptimusNodePinRouter,
+	public IOptimusNodeGraphProvider
 {
 	GENERATED_BODY()
 
 public:
+	UOptimusNode_SubGraphReference();
+
 	// UOptimusNode overrides
 	FName GetNodeCategory() const override { return NAME_None; }
 	FText GetDisplayName() const override;
 	void ConstructNode() override;
 
+	// UObject overrides
+	void PostLoad() override;
+	void BeginDestroy() override;
+
 	// IOptimusNodePinRouter implementation
 	FOptimusRoutedNodePin GetPinCounterpart(
 		UOptimusNodePin* InNodePin,
 		const FOptimusPinTraversalContext& InTraversalContext
 	) const override;
 
+	UOptimusNodeGraph* GetNodeGraphToShow() override;
+
+	UOptimusComponentSourceBinding* GetDefaultComponentBinding(const FOptimusPinTraversalContext& InTraversalContext) const;
+	
 protected:
 	friend class UOptimusNodeGraph;
-	
+
+	void SubscribeToSubGraph();
+	void UnsubscribeFromSubGraph() const;
+
+	void AddPinForNewBinding(FName InBindingArrayPropertyName);
+	void RemoveStalePins(FName InBindingArrayPropertyName);
+	void OnBindingMoved(FName InBindingArrayPropertyName);
+	void RecreateBindingPins(FName InBindingArrayPropertyName);
+	void SyncPinsToBindings(FName InBindingArrayPropertyName);
+
+	TArray<UOptimusNodePin*> GetBindingPinsByDirection(EOptimusNodePinDirection InDirection);
 	/** The graph that owns us. This contains all the necessary pin information to add on
 	 * the terminal node.
 	 */
 	UPROPERTY()
-	TObjectPtr<UOptimusNodeSubGraph> SubGraph;
+	TWeakObjectPtr<UOptimusNodeSubGraph> SubGraph;
+
+	UPROPERTY()
+	TWeakObjectPtr<UOptimusNodePin> DefaultComponentPin;
+
+	
 };
diff --git a/Engine/Plugins/Animation/DeformerGraph/Source/OptimusCore/Private/OptimusComponentSource.cpp b/Engine/Plugins/Animation/DeformerGraph/Source/OptimusCore/Private/OptimusComponentSource.cpp
--- a/Engine/Plugins/Animation/DeformerGraph/Source/OptimusCore/Private/OptimusComponentSource.cpp
+++ b/Engine/Plugins/Animation/DeformerGraph/Source/OptimusCore/Private/OptimusComponentSource.cpp
@@ -36,6 +36,18 @@ TArray<const UOptimusComponentSource*> UOptimusComponentSource::GetAllSources()
 	return ComponentSources;
 }
 
+TSet<FName> UOptimusComponentSource::GetAllExecutionDomains()
+{
+	TSet<FName> Domains;
+
+	for (const UOptimusComponentSource* Source : GetAllSources())
+	{
+		Domains.Append(Source->GetExecutionDomains());
+	}
+
+	return Domains;
+}
+
 const UOptimusComponentSource* UOptimusComponentSource::GetSourceFromDataInterface(
 	const UOptimusComputeDataInterface* InDataInterface
 	)
diff --git a/Engine/Plugins/Animation/DeformerGraph/Source/OptimusCore/Private/OptimusDataTypeRegistry.cpp b/Engine/Plugins/Animation/DeformerGraph/Source/OptimusCore/Private/OptimusDataTypeRegistry.cpp
--- a/Engine/Plugins/Animation/DeformerGraph/Source/OptimusCore/Private/OptimusDataTypeRegistry.cpp
+++ b/Engine/Plugins/Animation/DeformerGraph/Source/OptimusCore/Private/OptimusDataTypeRegistry.cpp
@@ -201,7 +201,7 @@ EOptimusDataTypeUsageFlags FOptimusDataTypeRegistry::GetStructTypeUsageFlag(UScr
         bool bFoundArray = false;
 
 		EOptimusDataTypeUsageFlags ApplicableFlags =
-			EOptimusDataTypeUsageFlags::Resource | EOptimusDataTypeUsageFlags::Variable | EOptimusDataTypeUsageFlags::AnimAttributes;
+			EOptimusDataTypeUsageFlags::Resource | EOptimusDataTypeUsageFlags::Variable | EOptimusDataTypeUsageFlags::AnimAttributes | EOptimusDataTypeUsageFlags::DataInterfaceOutput | EOptimusDataTypeUsageFlags::PinType;
     };
 
 	if (FOptimusDataTypeHandle DataType = Get().FindType(Optimus::GetTypeName(InStruct)))
@@ -238,7 +238,7 @@ void FOptimusDataTypeRegistry::RegisterBuiltinTypes()
 		},
 		ConvertPropertyValuePOD<bool, int32>,
 		FName(TEXT("bool")), {},
-	    EOptimusDataTypeUsageFlags::Resource | EOptimusDataTypeUsageFlags::Variable);
+	    EOptimusDataTypeUsageFlags::Resource | EOptimusDataTypeUsageFlags::Variable | EOptimusDataTypeUsageFlags::DataInterfaceOutput | EOptimusDataTypeUsageFlags::PinType);
 
 	// int -> int
 	Registry.RegisterType(
@@ -252,7 +252,7 @@ void FOptimusDataTypeRegistry::RegisterBuiltinTypes()
 	    },
 		ConvertPropertyValuePOD<int32, int32>,
 		FName(TEXT("int")), {}, 
-	    EOptimusDataTypeUsageFlags::Resource | EOptimusDataTypeUsageFlags::Variable | EOptimusDataTypeUsageFlags::AnimAttributes);
+	    EOptimusDataTypeUsageFlags::Resource | EOptimusDataTypeUsageFlags::Variable | EOptimusDataTypeUsageFlags::AnimAttributes | EOptimusDataTypeUsageFlags::DataInterfaceOutput | EOptimusDataTypeUsageFlags::PinType);
 
 	// FIntPoint -> int2
 	Registry.RegisterType(
@@ -261,7 +261,7 @@ void FOptimusDataTypeRegistry::RegisterBuiltinTypes()
 		FShaderValueType::Get(EShaderFundamentalType::Int, 2),
 		{},
 		bShowElements,
-		EOptimusDataTypeUsageFlags::Resource | EOptimusDataTypeUsageFlags::Variable);
+		EOptimusDataTypeUsageFlags::Resource | EOptimusDataTypeUsageFlags::Variable | EOptimusDataTypeUsageFlags::DataInterfaceOutput | EOptimusDataTypeUsageFlags::PinType);
 
 	// FIntVector -> int3
 	Registry.RegisterType(
@@ -270,7 +270,7 @@ void FOptimusDataTypeRegistry::RegisterBuiltinTypes()
 		FShaderValueType::Get(EShaderFundamentalType::Int, 3),
 		{},
 		bShowElements,
-		EOptimusDataTypeUsageFlags::Resource | EOptimusDataTypeUsageFlags::Variable);
+		EOptimusDataTypeUsageFlags::Resource | EOptimusDataTypeUsageFlags::Variable | EOptimusDataTypeUsageFlags::DataInterfaceOutput | EOptimusDataTypeUsageFlags::PinType);
 
 	// FIntVector4 -> int4
 	Registry.RegisterType(
@@ -278,7 +278,7 @@ void FOptimusDataTypeRegistry::RegisterBuiltinTypes()
 		FShaderValueType::Get(EShaderFundamentalType::Int, 4),
 		{},
 		bShowElements,
-		EOptimusDataTypeUsageFlags::Resource | EOptimusDataTypeUsageFlags::Variable);
+		EOptimusDataTypeUsageFlags::Resource | EOptimusDataTypeUsageFlags::Variable | EOptimusDataTypeUsageFlags::DataInterfaceOutput | EOptimusDataTypeUsageFlags::PinType);
 
 	// uint -> uint
 	Registry.RegisterType(
@@ -292,7 +292,7 @@ void FOptimusDataTypeRegistry::RegisterBuiltinTypes()
 		},
 		ConvertPropertyValuePOD<uint32, uint32>,
 		FName(TEXT("uint")), FLinearColor(0.0275f, 0.733, 0.820f, 1.0f), 
-		EOptimusDataTypeUsageFlags::Resource | EOptimusDataTypeUsageFlags::Variable);
+		EOptimusDataTypeUsageFlags::Resource | EOptimusDataTypeUsageFlags::Variable | EOptimusDataTypeUsageFlags::DataInterfaceOutput | EOptimusDataTypeUsageFlags::PinType);
 
 	/* FIXME: Need the corresponding definitions in UObject/Class.h @ line 3537
 	// FIntPoint -> int2
@@ -340,7 +340,7 @@ void FOptimusDataTypeRegistry::RegisterBuiltinTypes()
 	    },
 		ConvertPropertyValuePOD<float, float>,
 		FName(TEXT("real")), {}, 
-	    EOptimusDataTypeUsageFlags::Resource | EOptimusDataTypeUsageFlags::AnimAttributes);
+	    EOptimusDataTypeUsageFlags::Resource | EOptimusDataTypeUsageFlags::AnimAttributes | EOptimusDataTypeUsageFlags::DataInterfaceOutput | EOptimusDataTypeUsageFlags::PinType);
 
 	// double -> float 
 	Registry.RegisterType(
@@ -360,7 +360,7 @@ void FOptimusDataTypeRegistry::RegisterBuiltinTypes()
 	    },
 		ConvertPropertyValuePOD<double, float>,
 		FName(TEXT("real")), {}, 
-	    EOptimusDataTypeUsageFlags::Variable);
+	    EOptimusDataTypeUsageFlags::Variable | EOptimusDataTypeUsageFlags::DataInterfaceOutput | EOptimusDataTypeUsageFlags::PinType);
 
 	// FVector2D -> float2
 	Registry.RegisterType(
@@ -369,7 +369,7 @@ void FOptimusDataTypeRegistry::RegisterBuiltinTypes()
 		FShaderValueType::Get(EShaderFundamentalType::Float, 2),
 	    {},
 	    bShowElements,
-	    EOptimusDataTypeUsageFlags::Resource | EOptimusDataTypeUsageFlags::Variable);
+	    EOptimusDataTypeUsageFlags::Resource | EOptimusDataTypeUsageFlags::Variable | EOptimusDataTypeUsageFlags::DataInterfaceOutput | EOptimusDataTypeUsageFlags::PinType);
 
 	// FVector -> float3
 	Registry.RegisterType(
@@ -378,7 +378,7 @@ void FOptimusDataTypeRegistry::RegisterBuiltinTypes()
 		FShaderValueType::Get(EShaderFundamentalType::Float, 3),
 		{},
 		bShowElements,
-	    EOptimusDataTypeUsageFlags::Resource | EOptimusDataTypeUsageFlags::Variable | EOptimusDataTypeUsageFlags::AnimAttributes);
+	    EOptimusDataTypeUsageFlags::Resource | EOptimusDataTypeUsageFlags::Variable | EOptimusDataTypeUsageFlags::AnimAttributes | EOptimusDataTypeUsageFlags::DataInterfaceOutput | EOptimusDataTypeUsageFlags::PinType);
 
 	// FVector4 -> float4
 	Registry.RegisterType(
@@ -387,7 +387,7 @@ void FOptimusDataTypeRegistry::RegisterBuiltinTypes()
 		FShaderValueType::Get(EShaderFundamentalType::Float, 4),
 	    {},
 	    bShowElements,
-	    EOptimusDataTypeUsageFlags::Resource | EOptimusDataTypeUsageFlags::Variable);
+	    EOptimusDataTypeUsageFlags::Resource | EOptimusDataTypeUsageFlags::Variable | EOptimusDataTypeUsageFlags::DataInterfaceOutput | EOptimusDataTypeUsageFlags::PinType);
 	
 	// FLinearColor -> float4
 	Registry.RegisterType(
@@ -395,7 +395,7 @@ void FOptimusDataTypeRegistry::RegisterBuiltinTypes()
 	    FShaderValueType::Get(EShaderFundamentalType::Float, 4),
 	    {},
 	    bShowElements,
-	    EOptimusDataTypeUsageFlags::Resource | EOptimusDataTypeUsageFlags::Variable);
+	    EOptimusDataTypeUsageFlags::Resource | EOptimusDataTypeUsageFlags::Variable | EOptimusDataTypeUsageFlags::DataInterfaceOutput | EOptimusDataTypeUsageFlags::PinType);
 
 	// FQuat -> float4
 	Registry.RegisterType(
@@ -403,7 +403,7 @@ void FOptimusDataTypeRegistry::RegisterBuiltinTypes()
 		FShaderValueType::Get(EShaderFundamentalType::Float, 4),
 		{},
 		bShowElements,
-		EOptimusDataTypeUsageFlags::Resource | EOptimusDataTypeUsageFlags::Variable | EOptimusDataTypeUsageFlags::AnimAttributes);
+		EOptimusDataTypeUsageFlags::Resource | EOptimusDataTypeUsageFlags::Variable | EOptimusDataTypeUsageFlags::AnimAttributes | EOptimusDataTypeUsageFlags::DataInterfaceOutput | EOptimusDataTypeUsageFlags::PinType);
 
 	// FRotator -> float3x3
 	Registry.RegisterType(
@@ -411,7 +411,7 @@ void FOptimusDataTypeRegistry::RegisterBuiltinTypes()
 	    FShaderValueType::Get(EShaderFundamentalType::Float, 3, 3),
 	    {},
 	    bShowElements,
-	    EOptimusDataTypeUsageFlags::Resource | EOptimusDataTypeUsageFlags::Variable);
+	    EOptimusDataTypeUsageFlags::Resource | EOptimusDataTypeUsageFlags::Variable | EOptimusDataTypeUsageFlags::DataInterfaceOutput | EOptimusDataTypeUsageFlags::PinType);
 
 	// FTransform -> float4x4
 	Registry.RegisterType(
@@ -433,7 +433,7 @@ void FOptimusDataTypeRegistry::RegisterBuiltinTypes()
 	    },
 	    {},
 	    bHideElements,
-	    EOptimusDataTypeUsageFlags::Resource | EOptimusDataTypeUsageFlags::Variable | EOptimusDataTypeUsageFlags::AnimAttributes);
+	    EOptimusDataTypeUsageFlags::Resource | EOptimusDataTypeUsageFlags::Variable | EOptimusDataTypeUsageFlags::AnimAttributes | EOptimusDataTypeUsageFlags::DataInterfaceOutput | EOptimusDataTypeUsageFlags::PinType);
 
 	// HLSL types
 	Registry.RegisterType(
@@ -443,7 +443,7 @@ void FOptimusDataTypeRegistry::RegisterBuiltinTypes()
 		FName("float3x4"),
 		nullptr,
 		FLinearColor(0.7f, 0.3f, 0.4f, 1.0f),
-		EOptimusDataTypeUsageFlags::Resource);
+		EOptimusDataTypeUsageFlags::Resource | EOptimusDataTypeUsageFlags::DataInterfaceOutput | EOptimusDataTypeUsageFlags::PinType);
 
 	// FIXME: Add type aliases (e.g. "3x4 Float" above should really be "float3x4")
 
@@ -453,8 +453,11 @@ void FOptimusDataTypeRegistry::RegisterBuiltinTypes()
 		UOptimusComponentSourceBinding::StaticClass(),
 		FText::FromString("Component"),
 		FLinearColor(0.3f, 0.3f, 0.4f, 1.0f),
-		EOptimusDataTypeUsageFlags::None
+		EOptimusDataTypeUsageFlags::PinType
 		);
+
+	
+	Registry.TypeWithAtomicSupport = {FIntProperty::StaticClass()->GetFName()};
 }
 
 
@@ -1486,6 +1489,26 @@ UScriptStruct* FOptimusDataTypeRegistry::FindAttributeType(FName InTypeName) con
 	return nullptr;
 }
 
+TArray<FOptimusDataTypeHandle> FOptimusDataTypeRegistry::GetAllTypesWithAtomicSupport() const
+{
+	TArray<FOptimusDataTypeHandle> Types;
+	for (FName TypeName : TypeWithAtomicSupport)
+	{
+		FOptimusDataTypeHandle Type = FindType(TypeName);
+		if (Type.IsValid())
+		{
+			Types.Add(Type);
+		}
+	}
+
+	return Types;
+}
+
+bool FOptimusDataTypeRegistry::DoesTypeSupportAtomic(FOptimusDataTypeHandle InType)
+{
+	return GetAllTypesWithAtomicSupport().Contains(InType);
+}
+
 FOptimusDataTypeRegistry::FOnDataTypeChanged& FOptimusDataTypeRegistry::GetOnDataTypeChanged()
 {
 	return OnDataTypeChanged;
diff --git a/Engine/Plugins/Animation/DeformerGraph/Source/OptimusCore/Private/OptimusDeformer.cpp b/Engine/Plugins/Animation/DeformerGraph/Source/OptimusCore/Private/OptimusDeformer.cpp
--- a/Engine/Plugins/Animation/DeformerGraph/Source/OptimusCore/Private/OptimusDeformer.cpp
+++ b/Engine/Plugins/Animation/DeformerGraph/Source/OptimusCore/Private/OptimusDeformer.cpp
@@ -1641,7 +1641,7 @@ TArray<FOptimusComputeGraphInfo> UOptimusDeformer::CompileNodeGraphToComputeGrap
 	bool bValidationFailed = false;
 	for (FOptimusRoutedConstNode ConnectedNode: ConnectedNodes)
 	{
-		TOptional<FText> ErrorMessage = ConnectedNode.Node->ValidateForCompile();
+		TOptional<FText> ErrorMessage = ConnectedNode.Node->ValidateForCompile(ConnectedNode.TraversalContext);
 		if (ErrorMessage.IsSet())
 		{
 			bValidationFailed = true;
@@ -1808,7 +1808,7 @@ TArray<FOptimusComputeGraphInfo> UOptimusDeformer::CompileNodeGraphToComputeGrap
 			}
 			else
 			{
-				if (ConnectedNode.Node->GetOwningGraph()->DoesNodeHaveMutableInput(ConnectedNode.Node))
+				if (ConnectedNode.Node->GetOwningGraph()->DoesNodeHaveMutableInput(ConnectedNode.Node, ConnectedNode.TraversalContext))
 				{
 					KernelToGraphType.Add(ConnectedNode) = EOptimusNodeGraphType::Update;
 					GraphTypes.Add(EOptimusNodeGraphType::Update);
@@ -2058,7 +2058,7 @@ TArray<FOptimusComputeGraphInfo> UOptimusDeformer::CompileNodeGraphToComputeGrap
 			}
 
 			NodeDataInterfaceMap.Add(ConnectedNode.Node, DataInterface);
-			DataInterfaceToBindingIndexMap.Add(DataInterface) = NodeDataInterfaceProvider->GetComponentBinding()->GetIndex();
+			DataInterfaceToBindingIndexMap.Add(DataInterface) = NodeDataInterfaceProvider->GetComponentBinding(ConnectedNode.TraversalContext)->GetIndex();
 		}
 		else if (Cast<const IOptimusValueProvider>(ConnectedNode.Node))
 		{
@@ -2118,7 +2118,7 @@ TArray<FOptimusComputeGraphInfo> UOptimusDeformer::CompileNodeGraphToComputeGrap
 		if (const IOptimusComputeKernelProvider* KernelProvider = Cast<const IOptimusComputeKernelProvider>(Node))
 		{
 			UComputeDataInterface* KernelDataInterface = KernelProvider->MakeKernelDataInterface(this);
-			TSet<UOptimusComponentSourceBinding*> KernelPrimaryBindings = KernelProvider->GetPrimaryGroupPin()->GetComponentSourceBindingsRecursively();
+			TSet<UOptimusComponentSourceBinding*> KernelPrimaryBindings = KernelProvider->GetPrimaryGroupPin()->GetComponentSourceBindingsRecursively(RoutedNode.TraversalContext);
 
 			if (!ensure(KernelPrimaryBindings.Num() == 1))
 			{
@@ -2334,7 +2334,7 @@ TArray<FOptimusComputeGraphInfo> UOptimusDeformer::CompileNodeGraphToComputeGrap
 					UOptimusCopyKernelDataInterface* CopyKernelDataInterface = NewObject<UOptimusCopyKernelDataInterface>(this);
 					CopyKernelDataInterface->SetExecutionDomain(*TargetInstancedPin.Pin->GetDataDomain().AsExpression());
 					
-					UOptimusComponentSourceBinding* Binding = InterfaceProvider->GetComponentBinding();
+					UOptimusComponentSourceBinding* Binding = InterfaceProvider->GetComponentBinding(TargetInstancedPin.InstancedNode.RoutedNode.TraversalContext);
 					CopyKernelDataInterface->SetComponentBinding(Binding);
 					CopyKernelDataInterfaceMap.Add(SourceInstancedPin) = CopyKernelDataInterface;
 					DataInterfaceToBindingIndexMap.Add(CopyKernelDataInterface) = Binding->GetIndex();
@@ -2956,6 +2956,7 @@ void UOptimusDeformer::Notify(EOptimusGlobalNotifyType InNotifyType, UObject* In
 		{
 			ConstantValueUpdateDelegate.Broadcast(ConstantValue->GetValueName(), ConstantValue->GetShaderValue().ShaderValue);
 		}
+		
 		break;
 	default:
 		checkfSlow(false, TEXT("Unchecked EOptimusGlobalNotifyType!"));
diff --git a/Engine/Plugins/Animation/DeformerGraph/Source/OptimusCore/Private/OptimusHelpers.cpp b/Engine/Plugins/Animation/DeformerGraph/Source/OptimusCore/Private/OptimusHelpers.cpp
--- a/Engine/Plugins/Animation/DeformerGraph/Source/OptimusCore/Private/OptimusHelpers.cpp
+++ b/Engine/Plugins/Animation/DeformerGraph/Source/OptimusCore/Private/OptimusHelpers.cpp
@@ -195,3 +195,87 @@ FString Optimus::GetCookedKernelSource(
 			), *ComputeShaderUtilsInclude,*KernelFunc, *UnWrappedDispatchThreadId, *ShaderPathName, *Source);
 	}
 }
+
+bool Optimus::FindMovedItemInNameArray(const TArray<FName>& Old, const TArray<FName>& New, FName& OutSubjectName, FName& OutNextName)
+{
+	FName NameToMove = NAME_None;
+	int32 DivergeIndex = INDEX_NONE;
+
+	check(New.Num() == Old.Num())
+	
+	for (int32 Index = 0; Index < New.Num(); Index++)
+	{
+		if (New[Index]!= Old[Index] && DivergeIndex == INDEX_NONE)
+		{
+			DivergeIndex = Index;
+			continue;
+		}
+
+		if (DivergeIndex != INDEX_NONE)
+		{
+			if (New[DivergeIndex] == Old[Index])
+			{
+				NameToMove = Old[DivergeIndex];
+			}
+			else if (ensure(New[Index] == Old[DivergeIndex]))
+			{
+				NameToMove = New[DivergeIndex];
+			}
+			break;
+		}
+	}
+
+	if (DivergeIndex != INDEX_NONE)
+	{
+		OutSubjectName = NameToMove;
+		
+		const int32 NameIndex = New.IndexOfByPredicate([NameToMove](const FName& InBinding)
+		{
+			return InBinding == NameToMove;
+		}); 
+
+		const int32 NextNameIndex = NameIndex + 1;
+		OutNextName = New.IsValidIndex(NextNameIndex) ?  FName(New[NextNameIndex]) : NAME_None;
+
+		return true;
+	}
+
+	OutSubjectName = NAME_None;
+	OutNextName = NAME_None;
+	return false;
+}
+
+FName Optimus::GenerateUniqueNameFromExistingNames(FName InBaseName, const TArray<FName>& InExistingNames)
+{
+	TMap<FName, int32> BaseNameToMaxNumber;
+
+	for(const FName ExistingName : InExistingNames)
+	{
+		const int32 Number = ExistingName.GetNumber();
+		FName BaseName = ExistingName;
+		BaseName.SetNumber(0);
+
+		if (const int32* ExistingNumber = BaseNameToMaxNumber.Find(BaseName))
+		{
+			BaseNameToMaxNumber[BaseName] = FMath::Max(Number, *ExistingNumber);
+		}
+		else
+		{
+			BaseNameToMaxNumber.Add(BaseName) = Number;
+		}
+	}
+
+	FName NewName = InBaseName;
+	const int32 InputNumber = InBaseName.GetNumber();
+	FName InputBaseName = InBaseName;
+	InputBaseName.SetNumber(0);
+	if (const int32* ExistingNumber = BaseNameToMaxNumber.Find(InputBaseName))
+	{
+		if (InputNumber <= *ExistingNumber)
+		{
+			NewName.SetNumber(*ExistingNumber+1);
+		}
+	}
+
+	return NewName;
+}
diff --git a/Engine/Plugins/Animation/DeformerGraph/Source/OptimusCore/Private/OptimusNode.cpp b/Engine/Plugins/Animation/DeformerGraph/Source/OptimusCore/Private/OptimusNode.cpp
--- a/Engine/Plugins/Animation/DeformerGraph/Source/OptimusCore/Private/OptimusNode.cpp
+++ b/Engine/Plugins/Animation/DeformerGraph/Source/OptimusCore/Private/OptimusNode.cpp
@@ -389,37 +389,14 @@ FName UOptimusNode::GetAvailablePinNameStable(const UObject* InNodeOrPin, FName
 		PinsToCompare.Append(Pin->GetSubPins());
 	}
 
-	TMap<FName, int32> BaseNameToMaxNumber;
+	TArray<FName> PinNames;
 
 	for(const UOptimusNodePin* Pin : PinsToCompare)
 	{
-		const int32 Number = Pin->GetFName().GetNumber();
-		FName BaseName = Pin->GetFName();
-		BaseName.SetNumber(0);
-
-		if (const int32* ExistingNumber = BaseNameToMaxNumber.Find(BaseName))
-		{
-			BaseNameToMaxNumber[BaseName] = FMath::Max(Number, *ExistingNumber);
-		}
-		else
-		{
-			BaseNameToMaxNumber.Add(BaseName) = Number;
-		}
-	}
-
-	FName NewName = InName;
-	const int32 InputNumber = InName.GetNumber();
-	FName InputBaseName = InName;
-	InputBaseName.SetNumber(0);
-	if (const int32* ExistingNumber = BaseNameToMaxNumber.Find(InputBaseName))
-	{
-		if (InputNumber <= *ExistingNumber)
-		{
-			NewName.SetNumber(*ExistingNumber+1);
-		}
+		PinNames.Add(Pin->GetFName());
 	}
 
-	return NewName;
+	return Optimus::GenerateUniqueNameFromExistingNames(InName, PinNames);
 }
 
 
@@ -507,6 +484,10 @@ void UOptimusNode::SaveState(FArchive& Ar) const
 
 void UOptimusNode::RestoreState(FArchive& Ar)
 {
+	// Currently a warning appears when it can't find objects like pins and data interface data
+	// that were present during SaveState. However, they are harmless since we recreate those objects.
+	// But ideally there should be a better way to specify what should be saved/restored and what can be skipped/recreated
+	
 	FObjectAndNameAsStringProxyArchive NodeProxyArchive(
 		Ar, /* bInLoadIfFindFails=*/true);
 	SerializeScriptProperties(NodeProxyArchive);
@@ -571,6 +552,10 @@ UOptimusNodePin* UOptimusNode::AddPin(
 	return AddPinAction->GetPin(GetActionStack()->GetGraphCollectionRoot());
 }
 
+UOptimusNodePin* UOptimusNode::AddPin(const FOptimusParameterBinding& InBinding, EOptimusNodePinDirection InDirection, UOptimusNodePin* InBeforePin)
+{
+	return AddPin(InBinding.Name, InDirection, InBinding.DataDomain, InBinding.DataType, InBeforePin);
+}
 
 UOptimusNodePin* UOptimusNode::AddPinDirect(
     FName InName,
diff --git a/Engine/Plugins/Animation/DeformerGraph/Source/OptimusCore/Private/OptimusNodeGraph.cpp b/Engine/Plugins/Animation/DeformerGraph/Source/OptimusCore/Private/OptimusNodeGraph.cpp
--- a/Engine/Plugins/Animation/DeformerGraph/Source/OptimusCore/Private/OptimusNodeGraph.cpp
+++ b/Engine/Plugins/Animation/DeformerGraph/Source/OptimusCore/Private/OptimusNodeGraph.cpp
@@ -28,6 +28,7 @@
 
 #include <limits>
 
+#include "IOptimusComponentBindingReceiver.h"
 #include "IOptimusNodePairProvider.h"
 #include "IOptimusPinMutabilityDefiner.h"
 #include "OptimusNodeSubGraph.h"
@@ -1065,16 +1066,47 @@ UOptimusNode* UOptimusNodeGraph::CollapseNodesToSubGraph(
 	const TArray<UOptimusNode*>& InNodes
 	)
 {
-	const TSet<UOptimusNode*> NodeSet(InNodes);
 	
+	TArray<UOptimusNode*> NodesToCollapse;
+	for (UOptimusNode* Node : InNodes)
+	{
+		if (Cast<IOptimusNonCollapsibleNode>(Node))
+		{
+			continue;
+		}
+
+		NodesToCollapse.Add(Node);
+	}
+
+	TArray<UOptimusNode*> NodesToDuplicate;
+	TArray<UOptimusNodePair*> NodePairsToDuplicate;
+	AddPairNodesToArray(NodesToCollapse, NodesToDuplicate, NodePairsToDuplicate);
+
+	const TSet<UOptimusNode*> NodeSet(NodesToDuplicate);
+
 	// Collect all links that go to nodes that are not a part of the group and all links that
 	// are within elements of the group. At the same time, collect the bindings that apply.
 	TArray<const UOptimusNodeLink*> InternalLinks;
 	TArray<const UOptimusNodeLink*> InputLinks;			// Links going into the node set
+	TMap<const UOptimusNodeLink*, FName> EntryNodePinNames;
 	TArray<const UOptimusNodeLink*> OutputLinks;		// Links coming from the node set
+	TMap<const UOptimusNodeLink*, FName> ReturnNodePinNames;
 	TArray<FOptimusParameterBinding> InputBindings;
 	TArray<FOptimusParameterBinding> OutputBindings;
 
+	TSet<FName> BindingNames;
+	BindingNames.Add(UOptimusNodeSubGraph::GraphDefaultComponentPinName);
+	auto GetUniqueBindingName = [&BindingNames](FName InName)
+	{
+		FName NewName = InName;
+		while (BindingNames.Contains(NewName))
+		{
+			NewName.SetNumber(NewName.GetNumber() + 1);
+		}
+
+		return NewName;
+	};
+
 	// FIXME: The bindings should be ordered by node's Y-coordinate and input number.
 	for (const UOptimusNodeLink* Link: Links)
 	{
@@ -1093,10 +1125,12 @@ UOptimusNode* UOptimusNodeGraph::CollapseNodesToSubGraph(
 			// being collapsed into it.
 			const UOptimusNodePin* InputPin = Link->GetNodeInputPin();
 			FOptimusParameterBinding Binding;
-			Binding.Name = InputPin->GetFName();
+			Binding.Name = GetUniqueBindingName(InputPin->GetFName());
 			Binding.DataType = InputPin->GetDataType();
 			Binding.DataDomain = InputPin->GetDataDomain();
 			InputBindings.Add(Binding);
+			BindingNames.Add(Binding.Name);
+			EntryNodePinNames.Add(Link) = Binding.Name;
 		}
 		else if (bOutputNodeInSet)
 		{
@@ -1104,10 +1138,12 @@ UOptimusNode* UOptimusNodeGraph::CollapseNodesToSubGraph(
 			
 			const UOptimusNodePin* OutputPin = Link->GetNodeOutputPin();
 			FOptimusParameterBinding Binding;
-			Binding.Name = OutputPin->GetFName();
+			Binding.Name = GetUniqueBindingName(OutputPin->GetFName());
 			Binding.DataType = OutputPin->GetDataType();
 			Binding.DataDomain = OutputPin->GetDataDomain();
 			OutputBindings.Add(Binding);
+			BindingNames.Add(Binding.Name);
+			ReturnNodePinNames.Add(Link) = Binding.Name;
 		}
 	}
 
@@ -1132,7 +1168,7 @@ UOptimusNode* UOptimusNodeGraph::CollapseNodesToSubGraph(
 
 	// Create the entry and return nodes.
 	FBox2D NodeBox(ForceInit);
-	for (const UOptimusNode* Node: InNodes)
+	for (const UOptimusNode* Node: NodesToDuplicate)
 	{
 		NodeBox += Node->GetGraphPosition();
 	}
@@ -1148,7 +1184,7 @@ UOptimusNode* UOptimusNodeGraph::CollapseNodesToSubGraph(
 			
 			EntryNode->TerminalType = EOptimusTerminalType::Entry;
 			EntryNode->OwningGraph = SubGraph;
-			return EntryNode->SetGraphPositionDirect({NodeBox.Min.X - 150.0f, NodeBox.GetCenter().Y});
+			return EntryNode->SetGraphPositionDirect({NodeBox.Min.X - 450.0f, NodeBox.GetCenter().Y});
 		});
 
 	Action->AddSubAction<FOptimusNodeGraphAction_AddNode>(
@@ -1162,16 +1198,34 @@ UOptimusNode* UOptimusNodeGraph::CollapseNodesToSubGraph(
 			
 			ReturnNode->TerminalType = EOptimusTerminalType::Return;
 			ReturnNode->OwningGraph = SubGraph;
-			return ReturnNode->SetGraphPositionDirect({NodeBox.Max.X + 300.0f, NodeBox.GetCenter().Y});
+			return ReturnNode->SetGraphPositionDirect({NodeBox.Max.X + 750.0f, NodeBox.GetCenter().Y});
 		});
 	
 	// Duplicate the nodes into the graph.
-	for (UOptimusNode* Node: InNodes)
+	for (UOptimusNode* Node: NodesToDuplicate)
 	{
 		Action->AddSubAction<FOptimusNodeGraphAction_DuplicateNode>(
 			SubGraphPath, Node, Node->GetFName(), [](UOptimusNode*) { return true; });
 	}
 
+	TArray<TPair<FString, FString>> NewNodePairs;
+	for (const UOptimusNodePair* NodePair : NodePairsToDuplicate)
+	{
+		// Genernerate new pair
+		const UOptimusNode *FirstNode = NodePair->GetFirst();
+		const UOptimusNode *SecondNode = NodePair->GetSecond();
+
+		FString NewFirstNodePath = ConstructPath(SubGraphPath, FirstNode->GetName(), {});
+		FString NewSecondNodePath = ConstructPath(SubGraphPath, SecondNode->GetName(), {});
+
+		NewNodePairs.Add({NewFirstNodePath, NewSecondNodePath});
+	}
+
+	for (const TPair<FString, FString>& NewPair : NewNodePairs)
+	{
+		Action->AddSubAction<FOptimusNodeGraphAction_AddNodePair>(NewPair.Key, NewPair.Value);
+	}
+
 	// Create the reference node and connect it.
 	FName GraphNodeRefName("SubGraphNode");
 	GraphNodeRefName = Optimus::GetUniqueNameForScope(this, GraphNodeRefName);
@@ -1216,35 +1270,44 @@ UOptimusNode* UOptimusNodeGraph::CollapseNodesToSubGraph(
 	}
 	for (const UOptimusNodeLink* Link: InputLinks)
 	{
+		FString EntryNodePinName = EntryNodePinNames[Link].ToString(); 
 		// Once for the Entry -> Sub-graph nodes, and another for Outer graph nodes -> Sub-graph ref node inputs.
 		FString NodeOutputPinPath = FString::Printf(TEXT("%s/Entry.%s"),
-			*SubGraphPath, *Link->NodeInputPin->GetUniqueName().ToString());
+			*SubGraphPath, *EntryNodePinName);
 		FString NodeInputPinPath = FString::Printf(TEXT("%s/%s.%s"),
 			*SubGraphPath, *Link->NodeInputPin->GetOwningNode()->GetName(), *Link->NodeInputPin->GetUniqueName().ToString());
 		Action->AddSubAction<FOptimusNodeGraphAction_AddLink>(NodeOutputPinPath, NodeInputPinPath);
 		
 		NodeOutputPinPath = Link->NodeOutputPin->GetPinPath();
 		NodeInputPinPath = FString::Printf(TEXT("%s/%s.%s"),
-			*GetGraphPath(), *GraphNodeRefName.ToString(), *Link->NodeInputPin->GetUniqueName().ToString());
+			*GetGraphPath(), *GraphNodeRefName.ToString(), *EntryNodePinName);
 		Action->AddSubAction<FOptimusNodeGraphAction_AddLink>(NodeOutputPinPath, NodeInputPinPath);
 	}
 	for (const UOptimusNodeLink* Link: OutputLinks)
 	{
+		FString ReturnNodePinName = ReturnNodePinNames[Link].ToString();
+		
 		// Once for the Sub-graph nodes -> Return, and another for Sub-graph ref node outputs -> Outer graph nodes.
 		FString NodeOutputPinPath = FString::Printf(TEXT("%s/%s.%s"),
 			*SubGraphPath, *Link->NodeOutputPin->GetOwningNode()->GetName(), *Link->NodeOutputPin->GetUniqueName().ToString());
 		FString NodeInputPinPath = FString::Printf(TEXT("%s/Return.%s"),
-			*SubGraphPath, *Link->NodeOutputPin->GetUniqueName().ToString());
+			*SubGraphPath, *ReturnNodePinName);
 		Action->AddSubAction<FOptimusNodeGraphAction_AddLink>(NodeOutputPinPath, NodeInputPinPath);
 		
 		NodeOutputPinPath = FString::Printf(TEXT("%s/%s.%s"),
-			*GetGraphPath(), *GraphNodeRefName.ToString(), *Link->NodeOutputPin->GetUniqueName().ToString());
+			*GetGraphPath(), *GraphNodeRefName.ToString(), *ReturnNodePinName);
 		NodeInputPinPath = Link->NodeInputPin->GetPinPath();
 		Action->AddSubAction<FOptimusNodeGraphAction_AddLink>(NodeOutputPinPath, NodeInputPinPath);
 	}
+
+	// Delete the node pairs that we have moved into the subgraph
+	for (UOptimusNodePair* NodePair : NodePairsToDuplicate)
+	{
+		Action->AddSubAction<FOptimusNodeGraphAction_RemoveNodePair>(NodePair);
+	}
 	
 	// Delete the existing nodes
-	for (UOptimusNode* Node: InNodes)
+	for (UOptimusNode* Node: NodeSet)
 	{
 		Action->AddSubAction<FOptimusNodeGraphAction_RemoveNode>(Node);
 	}
@@ -1269,6 +1332,240 @@ TArray<UOptimusNode*> UOptimusNodeGraph::ExpandCollapsedNodes(
 		return {};
 	}
 
+	// Figure out the non-clashing names to use, to avoid collisions during actual execution.
+	TSet<FName> ExistingObjects;
+	for (const UOptimusNode* Node: Nodes)
+	{
+		if (ensure(Node != nullptr))
+		{
+			ExistingObjects.Add(Node->GetFName());
+		}
+	}
+
+	auto MakeUniqueNodeName = [&ExistingObjects](FName InName)
+	{
+		while(ExistingObjects.Contains(InName))
+		{
+			InName.SetNumber(InName.GetNumber() + 1);
+		}
+		ExistingObjects.Add(InName);
+		return InName;
+	};
+	
+	TArray<UOptimusNodeLink*> ExternalLinks;	
+
+	for (const int32 LinkIndex : GetAllLinkIndexesToNode(InFunctionNode))
+	{
+		ExternalLinks.Add(Links[LinkIndex]);
+	}
+
+	// Use to move all nodes in the outer graph away from the expansion location
+	// such that nodes currently in the subgraph would have enough space to spawn in the outer graph
+	
+	FVector2D ExpandLocation = InFunctionNode->GetGraphPosition();
+	
+	TArray<TPair<UOptimusNode*, FVector2D>> NodePositionInfos;
+	for (UOptimusNode* Node : Nodes)
+	{
+		if (Node == InFunctionNode)
+		{
+			continue;
+		}
+
+		NodePositionInfos.Add({Node, Node->GetGraphPosition() - InFunctionNode->GetGraphPosition()});
+	}
+
+	
+	if (UOptimusNode_SubGraphReference* SubGraphReferenceNode = Cast<UOptimusNode_SubGraphReference>(InFunctionNode))
+	{
+		UOptimusNodeSubGraph* SubGraph = SubGraphReferenceNode->SubGraph.Get();
+		
+		TArray<UOptimusNode*> NodesToDuplicate = SubGraph->Nodes;
+		TArray<UOptimusNodePair*> NodePairsToDuplicate = SubGraph->NodePairs;
+
+		NodesToDuplicate.RemoveAll([](const UOptimusNode* InNode)
+		{
+			return Cast<UOptimusNode_GraphTerminal>(InNode);
+		});
+
+
+		using FloatType = decltype(FVector2D::X);
+		FVector2D TopLeft{std::numeric_limits<FloatType>::max()};
+		FVector2D BotRight{std::numeric_limits<FloatType>::lowest()};
+		FVector2D HalfBoundSize;
+		FVector2D Center;
+		TMap<UOptimusNode*, FName> NewNodeNameMap;
+		for (UOptimusNode* Node: NodesToDuplicate)
+		{
+			TopLeft = FVector2D::Min(TopLeft, Node->GraphPosition);
+			BotRight = FVector2D::Max(BotRight, Node->GraphPosition);
+			NewNodeNameMap.Add(Node, MakeUniqueNodeName(Node->GetFName()));
+		}
+
+		HalfBoundSize = (BotRight - TopLeft )/2 + FVector2D(100, 100);
+		Center = (BotRight + TopLeft) / 2;
+		
+		const FString GraphPath = GetGraphPath();
+	
+		TArray<TPair<FString, FString>> NewNodePairs;
+		for (const UOptimusNodePair* NodePair : NodePairsToDuplicate)
+		{
+			// Generate new pair
+			const UOptimusNode *FirstNode = NodePair->GetFirst();
+			const UOptimusNode *SecondNode = NodePair->GetSecond();
+
+			if (ensure(NewNodeNameMap.Contains(FirstNode) && NewNodeNameMap.Contains(SecondNode)))
+			{
+				FString NewFirstNodePath = ConstructPath(GraphPath, NewNodeNameMap[FirstNode].ToString(), {});
+				FString NewSecondNodePath = ConstructPath(GraphPath, NewNodeNameMap[SecondNode].ToString(), {});
+
+				NewNodePairs.Add({NewFirstNodePath, NewSecondNodePath});
+			}	
+		}	
+		
+		TArray<TPair<FString, FString>> LinksToAdd;
+
+		// Record all links linking terminal pins and nodes in the outer graph, saving to LinksToAdd array
+		auto AddLink = [&LinksToAdd, GraphPath, SubGraphReferenceNode, NewNodeNameMap](UOptimusNodePin* InReferenceNodePin)
+		{
+			const FOptimusRoutedNodePin PinCounterpart = SubGraphReferenceNode->GetPinCounterpart(InReferenceNodePin, {});
+
+			for (const UOptimusNodePin* ExternalPin : InReferenceNodePin->GetConnectedPins())
+			{
+				FString ExternalPinPath = ExternalPin->GetPinPath();
+				
+				for (const UOptimusNodePin* InternalPin : PinCounterpart.NodePin->GetConnectedPins())
+				{
+					const UOptimusNode* InternalNode = InternalPin->GetOwningNode();
+					FString InternalPinPath = ConstructPath(GraphPath, NewNodeNameMap[InternalNode].ToString(), InternalPin->GetUniqueName().ToString());
+
+					if (ExternalPin->GetDirection() == EOptimusNodePinDirection::Output)
+					{
+						LinksToAdd.Add({ExternalPinPath, InternalPinPath});
+					}
+					else
+					{
+						LinksToAdd.Add({InternalPinPath, ExternalPinPath});
+					}
+				}
+			}	
+		};
+		
+		for (UOptimusNodePin* Pin : SubGraphReferenceNode->GetPins())
+		{
+			AddLink(Pin);
+		}
+		
+		// Unconnected component pins should be linked to outer graph explicitly if the subgraph default component pin is connected externally
+		TArray<UOptimusNodePin*> ExternalComponentPins = SubGraphReferenceNode->DefaultComponentPin->GetConnectedPins();
+		if (ExternalComponentPins.Num() > 0)
+		{
+			UOptimusNodePin* ExternalComponentPin = ExternalComponentPins[0];
+			check(ExternalComponentPin->GetDirection() == EOptimusNodePinDirection::Output);
+			
+			FString ExternalPinPath = ExternalComponentPin->GetPinPath();
+			
+			for (UOptimusNode* Node : SubGraph->Nodes)
+			{
+				if (IOptimusComponentBindingReceiver* ComponentReceiver = Cast<IOptimusComponentBindingReceiver>(Node))
+				{
+					TArray<UOptimusNodePin*> UnconnectedComponentPin = ComponentReceiver->GetUnconnectedComponentPins();
+					for (const UOptimusNodePin* InternalPin : UnconnectedComponentPin)
+					{
+						check(InternalPin->GetDirection() ==EOptimusNodePinDirection::Input);
+						
+						FString InternalPinPath = ConstructPath(GraphPath, NewNodeNameMap[Node].ToString(), InternalPin->GetUniqueName().ToString());
+						
+						LinksToAdd.Add({ExternalPinPath, InternalPinPath});
+					}
+				}
+			}
+		}
+		
+		TArray<UOptimusNodeLink*> InternalLinks = SubGraph->Links;
+		InternalLinks.RemoveAll([](const UOptimusNodeLink* InLink)
+		{
+			return (Cast<UOptimusNode_GraphTerminal>(InLink->NodeInputPin->GetOwningNode()) || 
+				Cast<UOptimusNode_GraphTerminal>(InLink->NodeOutputPin->GetOwningNode()));
+		});
+
+		for (UOptimusNodeLink* InternalLink : InternalLinks)
+		{
+			const UOptimusNodePin* OutputPin = InternalLink->GetNodeOutputPin();
+			const UOptimusNode* OutputNode = OutputPin->GetOwningNode();
+			const UOptimusNodePin* InputPin = InternalLink->GetNodeInputPin();
+			const UOptimusNode* InputNode = InputPin->GetOwningNode();
+			
+			LinksToAdd.Add({
+				ConstructPath(GraphPath, NewNodeNameMap[OutputNode].ToString(), OutputPin->GetUniqueName().ToString()),
+				ConstructPath(GraphPath, NewNodeNameMap[InputNode].ToString(), InputPin->GetUniqueName().ToString())
+			});
+		}
+
+		FOptimusCompoundAction *Action = new FOptimusCompoundAction(TEXT("Expand Collapsed Nodes"));
+
+		RemoveNodesToAction(Action, {SubGraphReferenceNode});
+
+		for (TPair<UOptimusNode*, FVector2D> NodePositionInfo: NodePositionInfos)
+		{
+			FVector2D Offset = NodePositionInfo.Value.GetSignVector() * HalfBoundSize;
+			
+			Action->AddSubAction<FOptimusNodeAction_MoveNode>(NodePositionInfo.Key, NodePositionInfo.Key->GetGraphPosition() + Offset);
+		}
+
+		
+		// Add all the pre-duplicate requirement actions first. This allows certain nodes to set up their
+		// operating environment correctly (e.g. a missing variable description for get variable nodes, etc.)
+		for (UOptimusNode* Node: NodesToDuplicate)
+		{
+			Node->PreDuplicateRequirementActions(this, Action);
+		}
+
+		TArray<TWeakPtr<FOptimusNodeGraphAction_DuplicateNode>> WeakDuplicateActions;
+
+		// Duplicate the nodes and place them correctly
+		for (UOptimusNode* Node: NodesToDuplicate)
+		{
+			FOptimusNodeGraphAction_DuplicateNode *DuplicateNodeAction = new FOptimusNodeGraphAction_DuplicateNode(
+				GetGraphPath(), Node, NewNodeNameMap[Node],
+				[Node, Center, ExpandLocation](UOptimusNode *InNode) {
+					InNode->SetGraphPositionDirect(Node->GraphPosition - Center + ExpandLocation);
+					return true;
+			});
+		
+			WeakDuplicateActions.Add(StaticCastWeakPtr<FOptimusNodeGraphAction_DuplicateNode>(Action->AddSubAction(DuplicateNodeAction)));
+		}
+
+		for (const TPair<FString, FString>& NewPair : NewNodePairs)
+		{
+			Action->AddSubAction<FOptimusNodeGraphAction_AddNodePair>(NewPair.Key, NewPair.Value);
+		}
+
+		for (const TPair<FString, FString>& NewLink : LinksToAdd)
+		{
+			Action->AddSubAction<FOptimusNodeGraphAction_AddLink>(NewLink.Key, NewLink.Value);
+		}
+
+		RemoveNodesToAction(Action, SubGraph->GetAllNodes());
+
+		Action->AddSubAction<FOptimusNodeGraphAction_RemoveGraph>(SubGraph);
+		
+		
+		
+		if (!GetActionStack()->RunAction(Action))
+		{
+			return {};
+		}
+
+		IOptimusPathResolver* PathResolver = GetPathResolver();
+		TArray<UOptimusNode*> NewNodes;
+		for (TWeakPtr<FOptimusNodeGraphAction_DuplicateNode> WeakDuplicateAction : WeakDuplicateActions)
+		{
+			NewNodes.Add(WeakDuplicateAction.Pin()->GetNode(PathResolver));
+		}
+		return NewNodes;	
+	}
+	
 	return {};
 }
 
@@ -1647,9 +1944,67 @@ TArray<FOptimusRoutedNodePin> UOptimusNodeGraph::GetConnectedPinsWithRouting(
 	return RoutedNodePins;
 }
 
+TArray<FOptimusRoutedNodePin> UOptimusNodeGraph::GetConnectedPinsWithRouting(const UOptimusNodePin* InNodePin, const FOptimusPinTraversalContext& InContext, EOptimusNodePinTraversalDirection Direction) const
+{
+	if (Direction == EOptimusNodePinTraversalDirection::Default)
+	{
+		return GetConnectedPinsWithRouting(InNodePin, InContext);
+	}
+	
+	const UOptimusNode* Node = InNodePin->GetOwningNode();
+	UOptimusNodePin* NodePin = const_cast<UOptimusNodePin*>(InNodePin);
+
+	FOptimusRoutedNodePin StartPin;
+
+	if (Direction == EOptimusNodePinTraversalDirection::Upstream)
+	{
+		if (InNodePin->GetDirection() == EOptimusNodePinDirection::Output)
+    	{
+    		if (const IOptimusNodePinRouter* RouterNode = Cast<const IOptimusNodePinRouter>(Node))
+    		{
+    			StartPin = RouterNode->GetPinCounterpart(NodePin, InContext);
+    		}
+    		else
+    		{
+    			return {{NodePin, InContext}};
+    		}
+    	}
+    	else
+    	{
+    		StartPin = {NodePin, InContext};
+    	}	
+	}
+	else
+	{
+		if (InNodePin->GetDirection() == EOptimusNodePinDirection::Input)
+		{
+			if (const IOptimusNodePinRouter* RouterNode = Cast<const IOptimusNodePinRouter>(Node))
+			{
+				StartPin = RouterNode->GetPinCounterpart(NodePin, InContext);
+			}
+			else
+			{
+				return {{NodePin, InContext}};
+			}
+		}
+		else
+		{
+			StartPin = {NodePin, InContext};
+		}
+	}
+	
+	if (StartPin.NodePin)
+	{
+		return GetConnectedPinsWithRouting(StartPin.NodePin, StartPin.TraversalContext);
+	}
+
+	return {};
+}
+
 
 TSet<UOptimusComponentSourceBinding*> UOptimusNodeGraph::GetComponentSourceBindingsForPin(
-	const UOptimusNodePin* InNodePin
+	const UOptimusNodePin* InNodePin,
+	const FOptimusPinTraversalContext& InContext
 	) const
 {
 	TSet<FOptimusRoutedConstNode> VisitedNodes;
@@ -1657,28 +2012,15 @@ TSet<UOptimusComponentSourceBinding*> UOptimusNodeGraph::GetComponentSourceBindi
 	
 	TQueue<FOptimusRoutedConstNode> WorkingSet;
 
-	// If given an input pin, find the other side.
-	const UOptimusNode* StartNode;
-	if (InNodePin->GetDirection() == EOptimusNodePinDirection::Input)
+	TArray<FOptimusRoutedNodePin> ConnectedPins = GetConnectedPinsWithRouting(InNodePin, InContext, EOptimusNodePinTraversalDirection::Upstream);
+	if (ConnectedPins.IsEmpty())
 	{
-		TArray<FOptimusRoutedNodePin> RoutedPins = GetConnectedPinsWithRouting(InNodePin, {});
-		if (RoutedPins.IsEmpty())
-		{
-			// If the input pin is directly on a compute kernel and not connected to anything
-			// we consider it to not have a component source, even if other pins in the same group
-			// has one. Pin group level gathering is caller's responsibility
-			return {};
-		}
-		
-		check(RoutedPins.Num() == 1);
-		StartNode = RoutedPins[0].NodePin->GetOwningNode();
-	}
-	else
-	{
-		StartNode = InNodePin->GetOwningNode();
+		return {};
 	}
 
-	WorkingSet.Enqueue({StartNode, FOptimusPinTraversalContext{}});
+	check(ConnectedPins.Num() == 1);
+	const UOptimusNode* StartNode = ConnectedPins[0].NodePin->GetOwningNode();
+	WorkingSet.Enqueue({StartNode, ConnectedPins[0].TraversalContext});
 	
 	FOptimusRoutedConstNode WorkItem;
 	while (WorkingSet.Dequeue(WorkItem))
@@ -1688,7 +2030,7 @@ TSet<UOptimusComponentSourceBinding*> UOptimusNodeGraph::GetComponentSourceBindi
 
 		if (const IOptimusComponentBindingProvider* ComponentSourceBindingProvider = Cast<const IOptimusComponentBindingProvider>(Node))
 		{
-			if (UOptimusComponentSourceBinding* Binding = ComponentSourceBindingProvider->GetComponentBinding())
+			if (UOptimusComponentSourceBinding* Binding = ComponentSourceBindingProvider->GetComponentBinding(WorkItem.TraversalContext))
 			{
 				Bindings.Add(Binding);
 			}
@@ -1739,16 +2081,17 @@ TSet<UOptimusComponentSourceBinding*> UOptimusNodeGraph::GetComponentSourceBindi
 }
 
 
-bool UOptimusNodeGraph::IsPinMutable(const UOptimusNodePin* InNodePin) const
+bool UOptimusNodeGraph::IsPinMutable(const UOptimusNodePin* InNodePin, const FOptimusPinTraversalContext& InContext) const
 {
 	TSet<FOptimusRoutedConstNodePin> VisitedPins;
 	TQueue<FOptimusRoutedConstNodePin> WorkingSet;
 
 	// If given an input pin, find the other side, since output pin provides mutability definition
 	const UOptimusNodePin* StartNodePin;
+	FOptimusPinTraversalContext StartContext;
 	if (InNodePin->GetDirection() == EOptimusNodePinDirection::Input)
 	{
-		TArray<FOptimusRoutedNodePin> RoutedPins = GetConnectedPinsWithRouting(InNodePin, {});
+		TArray<FOptimusRoutedNodePin> RoutedPins = GetConnectedPinsWithRouting(InNodePin, InContext);
 		if (RoutedPins.IsEmpty())
 		{
 			// Nothing is connected, pin data will never change
@@ -1757,13 +2100,15 @@ bool UOptimusNodeGraph::IsPinMutable(const UOptimusNodePin* InNodePin) const
 		
 		check(RoutedPins.Num() == 1);
 		StartNodePin = RoutedPins[0].NodePin;
+		StartContext = RoutedPins[0].TraversalContext;
 	}
 	else
 	{
 		StartNodePin = InNodePin;
+		StartContext = InContext;
 	}
 
-	WorkingSet.Enqueue({StartNodePin, FOptimusPinTraversalContext{}});
+	WorkingSet.Enqueue({StartNodePin, StartContext});
 	
 	FOptimusRoutedConstNodePin WorkItem;
 	while (WorkingSet.Dequeue(WorkItem))
@@ -1826,11 +2171,11 @@ bool UOptimusNodeGraph::IsPinMutable(const UOptimusNodePin* InNodePin) const
 	return false;
 }
 
-bool UOptimusNodeGraph::DoesNodeHaveMutableInput(const UOptimusNode* InNode) const
+bool UOptimusNodeGraph::DoesNodeHaveMutableInput(const UOptimusNode* InNode, const FOptimusPinTraversalContext& InContext) const
 {
 	for (UOptimusNodePin* Pin : InNode->GetPinsByDirection(EOptimusNodePinDirection::Input, true))
 	{
-		if (IsPinMutable(Pin))
+		if (IsPinMutable(Pin, InContext))
 		{
 			return true;
 		}
diff --git a/Engine/Plugins/Animation/DeformerGraph/Source/OptimusCore/Private/OptimusNodePair.cpp b/Engine/Plugins/Animation/DeformerGraph/Source/OptimusCore/Private/OptimusNodePair.cpp
--- a/Engine/Plugins/Animation/DeformerGraph/Source/OptimusCore/Private/OptimusNodePair.cpp
+++ b/Engine/Plugins/Animation/DeformerGraph/Source/OptimusCore/Private/OptimusNodePair.cpp
@@ -6,16 +6,7 @@
 
 #include UE_INLINE_GENERATED_CPP_BY_NAME(OptimusNodePair)
 
-void UOptimusNodePair::PostDuplicate(bool bDuplicateForPIE)
-{
-	Super::PostDuplicate(bDuplicateForPIE);
-
-	IOptimusNodePairProvider* FirstProvider = CastChecked<IOptimusNodePairProvider>(First);
-	IOptimusNodePairProvider* SecondProvider = CastChecked<IOptimusNodePairProvider>(Second);
 
-	FirstProvider->PairToCounterpartNode(SecondProvider);
-	SecondProvider->PairToCounterpartNode(FirstProvider);
-}
 
 bool UOptimusNodePair::Contains(const UOptimusNode* InFirst, const UOptimusNode* InSecond) const
 {
diff --git a/Engine/Plugins/Animation/DeformerGraph/Source/OptimusCore/Private/OptimusNodePin.cpp b/Engine/Plugins/Animation/DeformerGraph/Source/OptimusCore/Private/OptimusNodePin.cpp
--- a/Engine/Plugins/Animation/DeformerGraph/Source/OptimusCore/Private/OptimusNodePin.cpp
+++ b/Engine/Plugins/Animation/DeformerGraph/Source/OptimusCore/Private/OptimusNodePin.cpp
@@ -185,12 +185,12 @@ TArray<FName> UOptimusNodePin::GetPinNamePathFromString(const FStringView InPinP
 }
 
 
-TSet<UOptimusComponentSourceBinding*> UOptimusNodePin::GetComponentSourceBindings() const
+TSet<UOptimusComponentSourceBinding*> UOptimusNodePin::GetComponentSourceBindings(const FOptimusPinTraversalContext& InContext) const
 {
-	return GetOwningNode()->GetOwningGraph()->GetComponentSourceBindingsForPin(this);
+	return GetOwningNode()->GetOwningGraph()->GetComponentSourceBindingsForPin(this, InContext);
 }
 
-TSet<UOptimusComponentSourceBinding*> UOptimusNodePin::GetComponentSourceBindingsRecursively() const
+TSet<UOptimusComponentSourceBinding*> UOptimusNodePin::GetComponentSourceBindingsRecursively(const FOptimusPinTraversalContext& InContext) const
 {
 	TArray<UOptimusNodePin*> PinsToConsider = GetSubPinsRecursively(true);
 
@@ -198,15 +198,15 @@ TSet<UOptimusComponentSourceBinding*> UOptimusNodePin::GetComponentSourceBinding
 
 	for (const UOptimusNodePin* Pin : PinsToConsider)
 	{
-		Bindings.Append(Pin->GetComponentSourceBindings());
+		Bindings.Append(Pin->GetComponentSourceBindings(InContext));
 	}
 	
 	return Bindings;
 }
 
-bool UOptimusNodePin::IsMutable() const
+bool UOptimusNodePin::IsMutable(const FOptimusPinTraversalContext& InContext) const
 {
-	return GetOwningNode()->GetOwningGraph()->IsPinMutable(this);
+	return GetOwningNode()->GetOwningGraph()->IsPinMutable(this, InContext);
 }
 
 
diff --git a/Engine/Plugins/Animation/DeformerGraph/Source/OptimusCore/Private/OptimusNodeSubGraph.cpp b/Engine/Plugins/Animation/DeformerGraph/Source/OptimusCore/Private/OptimusNodeSubGraph.cpp
--- a/Engine/Plugins/Animation/DeformerGraph/Source/OptimusCore/Private/OptimusNodeSubGraph.cpp
+++ b/Engine/Plugins/Animation/DeformerGraph/Source/OptimusCore/Private/OptimusNodeSubGraph.cpp
@@ -2,10 +2,224 @@
 
 #include "OptimusNodeSubGraph.h"
 
+#include "OptimusDataTypeRegistry.h"
+#include "OptimusHelpers.h"
+#include "ComponentSources/OptimusSkeletalMeshComponentSource.h"
+#include "Nodes/OptimusNode_GraphTerminal.h"
+
 #include UE_INLINE_GENERATED_CPP_BY_NAME(OptimusNodeSubGraph)
 
+FName UOptimusNodeSubGraph::GraphDefaultComponentPinName = TEXT("Default Component Binding");
+FName UOptimusNodeSubGraph::InputBindingsPropertyName = GET_MEMBER_NAME_CHECKED(UOptimusNodeSubGraph, InputBindings);
+FName UOptimusNodeSubGraph::OutputBindingsPropertyName = GET_MEMBER_NAME_CHECKED(UOptimusNodeSubGraph, OutputBindings);
 
 
 UOptimusNodeSubGraph::UOptimusNodeSubGraph()
 {
 }
+
+FString UOptimusNodeSubGraph::GetBindingDeclaration(FName BindingName) const
+{
+	return {};
+}
+
+bool UOptimusNodeSubGraph::GetBindingSupportAtomicCheckBoxVisibility(FName BindingName) const
+{
+	return false;
+}
+
+bool UOptimusNodeSubGraph::GetBindingSupportReadCheckBoxVisibility(FName BindingName) const
+{
+	return false;
+}
+
+EOptimusDataTypeUsageFlags UOptimusNodeSubGraph::GetTypeUsageFlags(const FOptimusDataDomain& InDataDomain) const
+{
+	if (InDataDomain.IsSingleton())
+	{
+		return EOptimusDataTypeUsageFlags::PinType;
+	}
+
+	return EOptimusDataTypeUsageFlags::Resource;
+}
+
+UOptimusComponentSourceBinding* UOptimusNodeSubGraph::GetDefaultComponentBinding(const FOptimusPinTraversalContext& InTraversalContext) const
+{
+	if (!ensure(EntryNode.IsValid()))
+	{
+		return nullptr;
+	}
+
+	return EntryNode->GetDefaultComponentBinding(InTraversalContext);
+}
+
+#if WITH_EDITOR
+void UOptimusNodeSubGraph::PostEditChangeProperty(FPropertyChangedEvent& PropertyChangedEvent)
+{
+	Super::PostEditChangeProperty(PropertyChangedEvent);
+
+	if (PropertyChangedEvent.ChangeType & EPropertyChangeType::ValueSet)
+	{
+		if (CastField<FArrayProperty>(PropertyChangedEvent.Property) &&
+			PropertyChangedEvent.GetArrayIndex(PropertyChangedEvent.GetPropertyName().ToString()) == INDEX_NONE)
+		{
+			PropertyArrayPasted(PropertyChangedEvent);
+		}
+		else
+		{
+			PropertyValueChanged(PropertyChangedEvent);
+		}
+		
+	}
+	else if (PropertyChangedEvent.ChangeType & EPropertyChangeType::ArrayAdd)
+	{
+		PropertyArrayItemAdded(PropertyChangedEvent);
+	}
+	else if (PropertyChangedEvent.ChangeType & EPropertyChangeType::ArrayRemove)
+	{
+		PropertyArrayItemRemoved(PropertyChangedEvent);
+	}
+	else if (PropertyChangedEvent.ChangeType & EPropertyChangeType::ArrayClear)
+	{
+		PropertyArrayCleared(PropertyChangedEvent);
+	}
+	else if (PropertyChangedEvent.ChangeType & EPropertyChangeType::ArrayMove)
+	{
+		PropertyArrayItemMoved(PropertyChangedEvent);
+	}
+}
+#endif
+
+
+void UOptimusNodeSubGraph::SanitizeBinding(FOptimusParameterBinding& InOutBinding, FName InOldName, bool bInAllowParameter)
+{
+	InOutBinding.Name = GetSanitizedBindingName(InOutBinding.Name, InOldName);
+	
+	if (!InOutBinding.DataType.IsValid())
+	{
+		InOutBinding.DataType = FOptimusDataTypeRegistry::Get().FindType(*FFloatProperty::StaticClass());
+	}
+
+	if (!bInAllowParameter && InOutBinding.DataDomain.IsSingleton())
+	{
+		InOutBinding.DataDomain = FOptimusDataDomain(TArray<FName>({UOptimusSkeletalMeshComponentSource::Domains::Vertex}));
+	}
+}
+
+FName UOptimusNodeSubGraph::GetSanitizedBindingName(FName InNewName, FName InOldName)
+{
+	FName Name = InNewName;
+	
+	if (Name == NAME_None)
+	{
+		Name = TEXT("EmptyName");
+	}
+	
+	Name = Optimus::GetSanitizedNameForHlsl(Name);
+
+	if (Name != InOldName)
+	{
+		Name = UOptimusNode::GetAvailablePinNameStable(EntryNode.Get(), Name);
+		Name = UOptimusNode::GetAvailablePinNameStable(ReturnNode.Get(), Name);
+	}
+	
+	return Name;
+}
+
+#if WITH_EDITOR
+void UOptimusNodeSubGraph::PropertyArrayPasted(const FPropertyChangedEvent& InPropertyChangedEvent)
+{
+	FOptimusParameterBindingArray* BindingArrayPtr = nullptr;
+	bool bAllowParameter = false;
+	if (InPropertyChangedEvent.GetMemberPropertyName() == InputBindingsPropertyName)
+	{
+		BindingArrayPtr = &InputBindings;
+		bAllowParameter = true;
+		
+	}
+	else if (InPropertyChangedEvent.GetMemberPropertyName() == OutputBindingsPropertyName)
+	{
+		BindingArrayPtr = &OutputBindings;
+		bAllowParameter = true;
+	}
+
+	if (BindingArrayPtr)
+	{
+		for (FOptimusParameterBinding& Binding : *BindingArrayPtr)
+		{
+			SanitizeBinding(Binding, NAME_None, bAllowParameter);
+		}
+		OnBindingArrayPastedDelegate.Broadcast(InPropertyChangedEvent.GetMemberPropertyName());
+	}
+}
+
+void UOptimusNodeSubGraph::PropertyValueChanged(const FPropertyChangedEvent& InPropertyChangedEvent)
+{
+	FOptimusParameterBindingArray* BindingArrayPtr = nullptr;
+	TArray<UOptimusNodePin*> BindingPins;
+	bool bAllowParameter = false;
+	if (InPropertyChangedEvent.GetMemberPropertyName() == InputBindingsPropertyName)
+	{
+		BindingArrayPtr = &InputBindings;
+		bAllowParameter = true;
+		BindingPins = EntryNode->GetBindingPins();
+	}
+	else if (InPropertyChangedEvent.GetMemberPropertyName() == OutputBindingsPropertyName)
+	{
+		BindingArrayPtr = &OutputBindings;
+		bAllowParameter = false;
+		BindingPins = ReturnNode->GetBindingPins();
+	}
+	
+	if (BindingArrayPtr)
+	{
+		check(BindingArrayPtr->Num() == BindingPins.Num());
+
+		for (int32 Index = 0; Index < BindingArrayPtr->Num(); Index++)
+		{
+			FOptimusParameterBinding& Binding = (*BindingArrayPtr)[Index];
+			SanitizeBinding(Binding, BindingPins[Index]->GetFName(), bAllowParameter);
+		}
+		
+		OnBindingValueChangedDelegate.Broadcast(InPropertyChangedEvent.GetMemberPropertyName());
+	}	
+}
+
+void UOptimusNodeSubGraph::PropertyArrayItemAdded(const FPropertyChangedEvent& InPropertyChangedEvent)
+{
+	FOptimusParameterBindingArray* BindingArrayPtr = nullptr;
+
+	bool bAllowParameter = false;
+	if (InPropertyChangedEvent.GetMemberPropertyName() == InputBindingsPropertyName)
+	{
+		BindingArrayPtr = &InputBindings;
+		bAllowParameter = true;
+	}
+	else if (InPropertyChangedEvent.GetMemberPropertyName() == OutputBindingsPropertyName)
+	{
+		BindingArrayPtr = &OutputBindings;
+		bAllowParameter = false;
+	}
+
+	if (BindingArrayPtr)
+	{
+		SanitizeBinding(BindingArrayPtr->Last(), NAME_None, bAllowParameter);
+		OnBindingArrayItemAddedDelegate.Broadcast(InPropertyChangedEvent.GetMemberPropertyName());
+	}
+}
+
+void UOptimusNodeSubGraph::PropertyArrayItemRemoved(const FPropertyChangedEvent& InPropertyChangedEvent)
+{
+	OnBindingArrayItemRemovedDelegate.Broadcast(InPropertyChangedEvent.GetMemberPropertyName());
+}
+
+void UOptimusNodeSubGraph::PropertyArrayCleared(const FPropertyChangedEvent& InPropertyChangedEvent)
+{
+	OnBindingArrayClearedDelegate.Broadcast(InPropertyChangedEvent.GetMemberPropertyName());
+}
+
+void UOptimusNodeSubGraph::PropertyArrayItemMoved(const FPropertyChangedEvent& InPropertyChangedEvent)
+{
+	OnBindingArrayItemMovedDelegate.Broadcast(InPropertyChangedEvent.GetMemberPropertyName());
+}
+#endif
diff --git a/Engine/Plugins/Animation/DeformerGraph/Source/OptimusCore/Public/IOptimusAlternativeSelectedObjectProvider.h b/Engine/Plugins/Animation/DeformerGraph/Source/OptimusCore/Public/IOptimusAlternativeSelectedObjectProvider.h
new file mode 100644
--- /dev/null
+++ b/Engine/Plugins/Animation/DeformerGraph/Source/OptimusCore/Public/IOptimusAlternativeSelectedObjectProvider.h
@@ -0,0 +1,24 @@
+﻿// Copyright Epic Games, Inc. All Rights Reserved.
+
+#pragma once
+
+#include "UObject/Interface.h"
+
+#include "IOptimusAlternativeSelectedObjectProvider.generated.h"
+
+
+UINTERFACE()
+class OPTIMUSCORE_API UOptimusAlternativeSelectedObjectProvider :
+	public UInterface
+{
+	GENERATED_BODY()
+};
+
+
+class OPTIMUSCORE_API IOptimusAlternativeSelectedObjectProvider
+{
+	GENERATED_BODY()
+
+public:
+	virtual UObject* GetObjectToShowWhenSelected() const = 0;
+};
diff --git a/Engine/Plugins/Animation/DeformerGraph/Source/OptimusCore/Public/IOptimusComponentBindingProvider.h b/Engine/Plugins/Animation/DeformerGraph/Source/OptimusCore/Public/IOptimusComponentBindingProvider.h
--- a/Engine/Plugins/Animation/DeformerGraph/Source/OptimusCore/Public/IOptimusComponentBindingProvider.h
+++ b/Engine/Plugins/Animation/DeformerGraph/Source/OptimusCore/Public/IOptimusComponentBindingProvider.h
@@ -10,6 +10,7 @@
 class AActor;
 class UActorComponent;
 class UOptimusComponentSourceBinding;
+struct FOptimusPinTraversalContext;
 
 
 UINTERFACE()
@@ -26,5 +27,5 @@ class IOptimusComponentBindingProvider
 
 public:
 	/** Returns the component binding for the node that this interface is implemented on */
-	virtual UOptimusComponentSourceBinding *GetComponentBinding() const = 0;
+	virtual UOptimusComponentSourceBinding *GetComponentBinding(const FOptimusPinTraversalContext& InContext) const = 0;
 };
diff --git a/Engine/Plugins/Animation/DeformerGraph/Source/OptimusCore/Public/IOptimusComponentBindingReceiver.h b/Engine/Plugins/Animation/DeformerGraph/Source/OptimusCore/Public/IOptimusComponentBindingReceiver.h
new file mode 100644
--- /dev/null
+++ b/Engine/Plugins/Animation/DeformerGraph/Source/OptimusCore/Public/IOptimusComponentBindingReceiver.h
@@ -0,0 +1,28 @@
+﻿// Copyright Epic Games, Inc. All Rights Reserved.
+
+#pragma once
+
+#include "UObject/Interface.h"
+
+#include "IOptimusComponentBindingReceiver.generated.h"
+
+
+class UOptimusNodePin;
+
+
+UINTERFACE()
+class OPTIMUSCORE_API UOptimusComponentBindingReceiver :
+	public UInterface
+{
+	GENERATED_BODY()
+};
+
+
+class IOptimusComponentBindingReceiver
+{
+	GENERATED_BODY()
+
+public:
+	/** Returns unconnected component binding pins that would be using the default binding implicitly */
+	virtual TArray<UOptimusNodePin*> GetUnconnectedComponentPins() const = 0;
+};
diff --git a/Engine/Plugins/Animation/DeformerGraph/Source/OptimusCore/Public/IOptimusDataInterfaceProvider.h b/Engine/Plugins/Animation/DeformerGraph/Source/OptimusCore/Public/IOptimusDataInterfaceProvider.h
--- a/Engine/Plugins/Animation/DeformerGraph/Source/OptimusCore/Public/IOptimusDataInterfaceProvider.h
+++ b/Engine/Plugins/Animation/DeformerGraph/Source/OptimusCore/Public/IOptimusDataInterfaceProvider.h
@@ -10,7 +10,7 @@
 class UOptimusNodePin;
 class UOptimusComputeDataInterface;
 class UOptimusComponentSourceBinding;
-
+struct FOptimusPinTraversalContext;
 
 UINTERFACE()
 class OPTIMUSCORE_API UOptimusDataInterfaceProvider :
@@ -43,5 +43,5 @@ public:
 	virtual int32 GetDataFunctionIndexFromPin(const UOptimusNodePin* InPin) const = 0;
 
 	/** Returns the component binding that this data interface is bound to (or nullptr if unbound) */
-	virtual UOptimusComponentSourceBinding* GetComponentBinding() const = 0;
+	virtual UOptimusComponentSourceBinding* GetComponentBinding(const FOptimusPinTraversalContext& InContext) const = 0;
 };
diff --git a/Engine/Plugins/Animation/DeformerGraph/Source/OptimusCore/Public/IOptimusNodeGraphProvider.h b/Engine/Plugins/Animation/DeformerGraph/Source/OptimusCore/Public/IOptimusNodeGraphProvider.h
new file mode 100644
--- /dev/null
+++ b/Engine/Plugins/Animation/DeformerGraph/Source/OptimusCore/Public/IOptimusNodeGraphProvider.h
@@ -0,0 +1,28 @@
+﻿// Copyright Epic Games, Inc. All Rights Reserved.
+
+#pragma once
+
+#include "UObject/Interface.h"
+
+#include "IOptimusNodeGraphProvider.generated.h"
+
+class UOptimusNodeGraph;
+
+UINTERFACE()
+class OPTIMUSCORE_API UOptimusNodeGraphProvider :
+	public UInterface
+{
+	GENERATED_BODY()
+};
+
+/**
+* Interface that provides a mechanism to get a node graph to display
+*/
+class OPTIMUSCORE_API IOptimusNodeGraphProvider
+{
+	GENERATED_BODY()
+
+public:
+	virtual UOptimusNodeGraph* GetNodeGraphToShow() = 0;
+};
+
diff --git a/Engine/Plugins/Animation/DeformerGraph/Source/OptimusCore/Public/IOptimusNonCollapsibleNode.h b/Engine/Plugins/Animation/DeformerGraph/Source/OptimusCore/Public/IOptimusNonCollapsibleNode.h
new file mode 100644
--- /dev/null
+++ b/Engine/Plugins/Animation/DeformerGraph/Source/OptimusCore/Public/IOptimusNonCollapsibleNode.h
@@ -0,0 +1,24 @@
+﻿// Copyright Epic Games, Inc. All Rights Reserved.
+
+#pragma once
+
+#include "UObject/Interface.h"
+
+#include "IOptimusNonCollapsibleNode.generated.h"
+
+UINTERFACE()
+class OPTIMUSCORE_API UOptimusNonCollapsibleNode :
+	public UInterface
+{
+	GENERATED_BODY()
+};
+
+/**
+* Interface that provides a mechanism to prevent node from being included in a subgraph
+*/
+class OPTIMUSCORE_API IOptimusNonCollapsibleNode
+{
+	GENERATED_BODY()
+
+};
+
diff --git a/Engine/Plugins/Animation/DeformerGraph/Source/OptimusCore/Public/IOptimusParameterBindingProvider.h b/Engine/Plugins/Animation/DeformerGraph/Source/OptimusCore/Public/IOptimusParameterBindingProvider.h
--- a/Engine/Plugins/Animation/DeformerGraph/Source/OptimusCore/Public/IOptimusParameterBindingProvider.h
+++ b/Engine/Plugins/Animation/DeformerGraph/Source/OptimusCore/Public/IOptimusParameterBindingProvider.h
@@ -25,4 +25,5 @@ public:
 	virtual FString GetBindingDeclaration(FName BindingName) const = 0;
 	virtual bool GetBindingSupportAtomicCheckBoxVisibility(FName BindingName) const = 0;
 	virtual bool GetBindingSupportReadCheckBoxVisibility(FName BindingName) const = 0;
+	virtual EOptimusDataTypeUsageFlags GetTypeUsageFlags(const FOptimusDataDomain& InDataDomain) const = 0;
 };
diff --git a/Engine/Plugins/Animation/DeformerGraph/Source/OptimusCore/Public/OptimusBindingTypes.h b/Engine/Plugins/Animation/DeformerGraph/Source/OptimusCore/Public/OptimusBindingTypes.h
--- a/Engine/Plugins/Animation/DeformerGraph/Source/OptimusCore/Public/OptimusBindingTypes.h
+++ b/Engine/Plugins/Animation/DeformerGraph/Source/OptimusCore/Public/OptimusBindingTypes.h
@@ -34,7 +34,7 @@ struct FOptimusParameterBinding
 	UPROPERTY(EditAnywhere, Category=Binding)
 	FOptimusValidatedName Name;
 
-	UPROPERTY(EditAnywhere, Category = Binding, meta=(UseInResource))
+	UPROPERTY(EditAnywhere, Category = Binding)
 	FOptimusDataTypeRef DataType;
 
 	/** Returns true if the binding is valid and has defined entries */
diff --git a/Engine/Plugins/Animation/DeformerGraph/Source/OptimusCore/Public/OptimusComponentSource.h b/Engine/Plugins/Animation/DeformerGraph/Source/OptimusCore/Public/OptimusComponentSource.h
--- a/Engine/Plugins/Animation/DeformerGraph/Source/OptimusCore/Public/OptimusComponentSource.h
+++ b/Engine/Plugins/Animation/DeformerGraph/Source/OptimusCore/Public/OptimusComponentSource.h
@@ -71,6 +71,9 @@ public:
 	/** Returns all registered component source objects */
 	static TArray<const UOptimusComponentSource*> GetAllSources();
 
+	/** Returns all execution domains from all available sources */
+	static TSet<FName> GetAllExecutionDomains();
+
 	/** Returns a component source that matches a data interface, or nullptr if nothing does */
 	static const UOptimusComponentSource* GetSourceFromDataInterface(
 		const UOptimusComputeDataInterface* InDataInterface
diff --git a/Engine/Plugins/Animation/DeformerGraph/Source/OptimusCore/Public/OptimusDataType.h b/Engine/Plugins/Animation/DeformerGraph/Source/OptimusCore/Public/OptimusDataType.h
--- a/Engine/Plugins/Animation/DeformerGraph/Source/OptimusCore/Public/OptimusDataType.h
+++ b/Engine/Plugins/Animation/DeformerGraph/Source/OptimusCore/Public/OptimusDataType.h
@@ -20,6 +20,8 @@ enum class EOptimusDataTypeUsageFlags : uint8
 	Resource			= 1 << 0,		/** This type can be used in a resource */
 	Variable			= 1 << 1,		/** This type can be used in a variable */
 	AnimAttributes      = 1 << 2,       /** This type can be used to query an anim attribute*/
+	DataInterfaceOutput = 1 << 3,       /** This type can be used as output of a data interface*/
+	PinType				= 1 << 4,       /** This type can be used as pin type*/
 };
 ENUM_CLASS_FLAGS(EOptimusDataTypeUsageFlags)
 
diff --git a/Engine/Plugins/Animation/DeformerGraph/Source/OptimusCore/Public/OptimusDataTypeRegistry.h b/Engine/Plugins/Animation/DeformerGraph/Source/OptimusCore/Public/OptimusDataTypeRegistry.h
--- a/Engine/Plugins/Animation/DeformerGraph/Source/OptimusCore/Public/OptimusDataTypeRegistry.h
+++ b/Engine/Plugins/Animation/DeformerGraph/Source/OptimusCore/Public/OptimusDataTypeRegistry.h
@@ -174,6 +174,10 @@ public:
 	/** A helper function to return the corresponding animation attribute type. */
 	UScriptStruct* FindAttributeType(FName InTypeName) const;
 
+	TArray<FOptimusDataTypeHandle> GetAllTypesWithAtomicSupport() const;
+
+	bool DoesTypeSupportAtomic(FOptimusDataTypeHandle InType);
+	
 	OPTIMUSCORE_API	FOnDataTypeChanged& GetOnDataTypeChanged();
 
 	// -- FGCObject overrides
@@ -236,5 +240,7 @@ private:
 	TMap<FName /* TypeName */, FTypeInfo> RegisteredTypes;
 	TArray<FName> RegistrationOrder;
 
+	TSet<FName> TypeWithAtomicSupport;
+	
 	FOnDataTypeChanged OnDataTypeChanged;
 };
diff --git a/Engine/Plugins/Animation/DeformerGraph/Source/OptimusCore/Public/OptimusHelpers.h b/Engine/Plugins/Animation/DeformerGraph/Source/OptimusCore/Public/OptimusHelpers.h
--- a/Engine/Plugins/Animation/DeformerGraph/Source/OptimusCore/Public/OptimusHelpers.h
+++ b/Engine/Plugins/Animation/DeformerGraph/Source/OptimusCore/Public/OptimusHelpers.h
@@ -107,4 +107,8 @@ namespace Optimus
 		const FString& InKernelName,
 		FIntVector InGroupSize
 	);
+
+	bool FindMovedItemInNameArray(const TArray<FName>& Old, const TArray<FName>& New, FName& OutSubjectName, FName& OutNextName);
+
+	FName GenerateUniqueNameFromExistingNames(FName InBaseName, const TArray<FName>& InExistingNames);
 }
diff --git a/Engine/Plugins/Animation/DeformerGraph/Source/OptimusCore/Public/OptimusNode.h b/Engine/Plugins/Animation/DeformerGraph/Source/OptimusCore/Public/OptimusNode.h
--- a/Engine/Plugins/Animation/DeformerGraph/Source/OptimusCore/Public/OptimusNode.h
+++ b/Engine/Plugins/Animation/DeformerGraph/Source/OptimusCore/Public/OptimusNode.h
@@ -189,6 +189,7 @@ public:
 	
 protected:
 	friend class UOptimusNodeGraph;
+	friend class UOptimusNodeSubGraph;
 	friend class UOptimusNodePin;
 	friend class UOptimusDeformer;
 	friend struct FOptimusNodeAction_AddRemovePin;
@@ -229,7 +230,7 @@ protected:
 	 *  reason as a part of the optional return value. If success return an empty optional
 	 *  object.
 	 */
-	virtual TOptional<FText> ValidateForCompile() const
+	virtual TOptional<FText> ValidateForCompile(const FOptimusPinTraversalContext& InContext) const
 	{
 		return {};
 	}
@@ -257,6 +258,13 @@ protected:
 		UOptimusNodePin* InGroupingPin = nullptr
 		);
 
+	/** Add a new pin based on a parameter binding definition. Only allowed for top-level pins. */
+	UOptimusNodePin* AddPin(
+		const FOptimusParameterBinding& InBinding,
+		EOptimusNodePinDirection InDirection,
+		UOptimusNodePin* InBeforePin = nullptr
+		);
+
 	/** Create a pin and add it to the node in the location specified. */ 
 	UOptimusNodePin* AddPinDirect(
 		FName InName,
diff --git a/Engine/Plugins/Animation/DeformerGraph/Source/OptimusCore/Public/OptimusNodeGraph.h b/Engine/Plugins/Animation/DeformerGraph/Source/OptimusCore/Public/OptimusNodeGraph.h
--- a/Engine/Plugins/Animation/DeformerGraph/Source/OptimusCore/Public/OptimusNodeGraph.h
+++ b/Engine/Plugins/Animation/DeformerGraph/Source/OptimusCore/Public/OptimusNodeGraph.h
@@ -45,6 +45,13 @@ enum class EOptimusNodeGraphType
 	Transient				/** Used to store nodes during duplication. Never serialized. */
 };
 
+enum class EOptimusNodePinTraversalDirection
+{
+	Default,
+	Upstream,
+	Downstream
+};
+
 namespace Optimus
 {
 static bool IsExecutionGraphType(EOptimusNodeGraphType InGraphType)
@@ -297,22 +304,31 @@ public:
 		const UOptimusNodePin* InNodePin,
 		const FOptimusPinTraversalContext& InContext
 		) const;
+	
+	TArray<FOptimusRoutedNodePin> GetConnectedPinsWithRouting(
+		const UOptimusNodePin* InNodePin,
+		const FOptimusPinTraversalContext& InContext,
+		EOptimusNodePinTraversalDirection Direction
+		) const;
 
 	/** Get all unique component bindings that lead to this pin. Note that only pins with a zero or single bindings
 	 *  are considered valid. We return all of them however for error messaging.
 	 */
 	TSet<UOptimusComponentSourceBinding*> GetComponentSourceBindingsForPin(
-		const UOptimusNodePin* InNodePin
+		const UOptimusNodePin* InNodePin,
+		const FOptimusPinTraversalContext& InContext
 		) const;
 
 	/** Check if a pin represents time varying data */
 	bool IsPinMutable(
-		const UOptimusNodePin* InNodePin
+		const UOptimusNodePin* InNodePin,
+		const FOptimusPinTraversalContext& InContext
 		) const;
 
 	/** Check if a Node has mutable input pins */
 	bool DoesNodeHaveMutableInput(
-		const UOptimusNode* InNode 
+		const UOptimusNode* InNode,
+		const FOptimusPinTraversalContext& InContext
 		) const;
 
 	/** Gather connected loop entry terminals */
@@ -482,7 +498,7 @@ private:
 	);
 
 	void AddPairNodesToArray(const TArray<UOptimusNode*>& InNodes, TArray<UOptimusNode*>& OutNodes, TArray<UOptimusNodePair*>& OutNodePairs) const;
-	
+
 	void RemoveNodePairByIndex(int32 NodePairIndex);
 	
 	void RemoveLinkByIndex(int32 LinkIndex);
diff --git a/Engine/Plugins/Animation/DeformerGraph/Source/OptimusCore/Public/OptimusNodePair.h b/Engine/Plugins/Animation/DeformerGraph/Source/OptimusCore/Public/OptimusNodePair.h
--- a/Engine/Plugins/Animation/DeformerGraph/Source/OptimusCore/Public/OptimusNodePair.h
+++ b/Engine/Plugins/Animation/DeformerGraph/Source/OptimusCore/Public/OptimusNodePair.h
@@ -16,9 +16,6 @@ class OPTIMUSCORE_API UOptimusNodePair : public UObject
 public:
 	UOptimusNodePair() = default;
 
-	// UObject overrides
-	void PostDuplicate(bool bDuplicateForPIE) override;
-	
 	/** Returns the output pin on the node this link connects from. */
 	UOptimusNode* GetFirst() const { return First; }
 
diff --git a/Engine/Plugins/Animation/DeformerGraph/Source/OptimusCore/Public/OptimusNodePin.h b/Engine/Plugins/Animation/DeformerGraph/Source/OptimusCore/Public/OptimusNodePin.h
--- a/Engine/Plugins/Animation/DeformerGraph/Source/OptimusCore/Public/OptimusNodePin.h
+++ b/Engine/Plugins/Animation/DeformerGraph/Source/OptimusCore/Public/OptimusNodePin.h
@@ -91,15 +91,15 @@ public:
 	
 	/** Return all component source bindings that flow into pin. 
 	 */
-	TSet<UOptimusComponentSourceBinding*> GetComponentSourceBindings() const;
+	TSet<UOptimusComponentSourceBinding*> GetComponentSourceBindings(const FOptimusPinTraversalContext& InContext) const;
 
 	/** Return all component source bindings that flow into this input pin and its sub-pins.
 	 */
-	TSet<UOptimusComponentSourceBinding*> GetComponentSourceBindingsRecursively() const;
+	TSet<UOptimusComponentSourceBinding*> GetComponentSourceBindingsRecursively(const FOptimusPinTraversalContext& InContext) const;
 
 	/** Whether the data presented by the pin can change over time
 	 */
-	bool IsMutable() const;
+	bool IsMutable(const FOptimusPinTraversalContext& InContext) const;
 	
 	/** Returns the FProperty object for this pin. This can be used to directly address the
 	  * node data represented by this pin. Not all pins have an underlying resource so this can
diff --git a/Engine/Plugins/Animation/DeformerGraph/Source/OptimusCore/Public/OptimusNodeSubGraph.h b/Engine/Plugins/Animation/DeformerGraph/Source/OptimusCore/Public/OptimusNodeSubGraph.h
--- a/Engine/Plugins/Animation/DeformerGraph/Source/OptimusCore/Public/OptimusNodeSubGraph.h
+++ b/Engine/Plugins/Animation/DeformerGraph/Source/OptimusCore/Public/OptimusNodeSubGraph.h
@@ -1,29 +1,88 @@
 ﻿// Copyright Epic Games, Inc. All Rights Reserved.
 
 #pragma once
+#include "IOptimusParameterBindingProvider.h"
 #include "OptimusNodeGraph.h"
 #include "OptimusBindingTypes.h"
 #include "OptimusNodeSubGraph.generated.h"
 
 class UOptimusNode_GraphTerminal;
+
 UCLASS()
 class UOptimusNodeSubGraph :
-	public UOptimusNodeGraph
+	public UOptimusNodeGraph,
+	public IOptimusParameterBindingProvider
 {
+	
 	GENERATED_BODY()
+
 public:
+
+	static FName GraphDefaultComponentPinName;
+	static FName InputBindingsPropertyName;
+	static FName OutputBindingsPropertyName;
+
 	UOptimusNodeSubGraph();
+	
+	// IOptimusParameterBindingProvider 
+	FString GetBindingDeclaration(FName BindingName) const override;
+	bool GetBindingSupportAtomicCheckBoxVisibility(FName BindingName) const override;
+	bool GetBindingSupportReadCheckBoxVisibility(FName BindingName) const override;
+	EOptimusDataTypeUsageFlags GetTypeUsageFlags(const FOptimusDataDomain& InDataDomain) const override;
+
+	UOptimusComponentSourceBinding* GetDefaultComponentBinding(const FOptimusPinTraversalContext& InTraversalContext) const;
+	
 
-	// FIXME: These are uneditable for now.
-	UPROPERTY(VisibleAnywhere, Category=Bindings)
-	TArray<FOptimusParameterBinding> InputBindings;
+	DECLARE_MULTICAST_DELEGATE_OneParam(FOnBindingArrayPasted, FName);
+	DECLARE_MULTICAST_DELEGATE_OneParam(FOnBindingValueChanged, FName);
+	DECLARE_MULTICAST_DELEGATE_OneParam(FOnBindingArrayItemAdded, FName);
+	DECLARE_MULTICAST_DELEGATE_OneParam(FOnBindingArrayItemRemoved, FName);
+	DECLARE_MULTICAST_DELEGATE_OneParam(FOnBindingArrayCleared, FName);
+	DECLARE_MULTICAST_DELEGATE_OneParam(FOnBindingArrayItemMoved, FName);
+	
+	FOnBindingArrayPasted& GetOnBindingArrayPasted() { return OnBindingArrayPastedDelegate; }
+	FOnBindingValueChanged& GetOnBindingValueChanged() { return OnBindingValueChangedDelegate; }
+	FOnBindingArrayItemAdded& GetOnBindingArrayItemAdded() { return OnBindingArrayItemAddedDelegate; }
+	FOnBindingArrayItemRemoved& GetOnBindingArrayItemRemoved() { return OnBindingArrayItemRemovedDelegate; }
+	FOnBindingArrayCleared& GetOnBindingArrayCleared() { return OnBindingArrayClearedDelegate; }
+	FOnBindingArrayItemMoved& GetOnBindingArrayItemMoved() { return OnBindingArrayItemMovedDelegate; }
+	
+#if WITH_EDITOR
+	// UObject overrides
+	void PostEditChangeProperty(FPropertyChangedEvent& PropertyChangedEvent) override;
+#endif
 
-	UPROPERTY(VisibleAnywhere, Category=Bindings)
-	TArray<FOptimusParameterBinding> OutputBindings;
+	UPROPERTY(EditAnywhere, Category=Bindings, meta=(AllowParameters))
+	FOptimusParameterBindingArray InputBindings;
+
+	UPROPERTY(EditAnywhere, Category=Bindings)
+	FOptimusParameterBindingArray OutputBindings;
 
 	UPROPERTY()
 	TWeakObjectPtr<UOptimusNode_GraphTerminal> EntryNode; 
 
 	UPROPERTY()
-	TWeakObjectPtr<UOptimusNode_GraphTerminal> ReturnNode; 
+	TWeakObjectPtr<UOptimusNode_GraphTerminal> ReturnNode;
+private:
+
+	void SanitizeBinding(FOptimusParameterBinding& InOutBinding, FName InOldName, bool bInAllowParameter);
+	FName GetSanitizedBindingName(FName InNewName, FName InOldName);
+	
+
+#if WITH_EDITOR
+	void PropertyArrayPasted(const FPropertyChangedEvent& InPropertyChangedEvent);
+	void PropertyValueChanged(const FPropertyChangedEvent& InPropertyChangedEvent);
+	void PropertyArrayItemAdded(const FPropertyChangedEvent& InPropertyChangedEvent);
+	void PropertyArrayItemRemoved(const FPropertyChangedEvent& InPropertyChangedEvent);
+	void PropertyArrayCleared(const FPropertyChangedEvent& InPropertyChangedEvent);
+	void PropertyArrayItemMoved(const FPropertyChangedEvent& InPropertyChangedEvent);
+#endif
+	
+	FOnBindingArrayPasted OnBindingArrayPastedDelegate;
+	FOnBindingValueChanged OnBindingValueChangedDelegate;
+	FOnBindingArrayItemAdded OnBindingArrayItemAddedDelegate;
+	FOnBindingArrayItemRemoved OnBindingArrayItemRemovedDelegate;
+	FOnBindingArrayCleared OnBindingArrayClearedDelegate;
+	FOnBindingArrayItemMoved OnBindingArrayItemMovedDelegate;
+
 };
diff --git a/Engine/Plugins/Animation/DeformerGraph/Source/OptimusEditor/Private/OptimusDetailsCustomization.cpp b/Engine/Plugins/Animation/DeformerGraph/Source/OptimusEditor/Private/OptimusDetailsCustomization.cpp
--- a/Engine/Plugins/Animation/DeformerGraph/Source/OptimusEditor/Private/OptimusDetailsCustomization.cpp
+++ b/Engine/Plugins/Animation/DeformerGraph/Source/OptimusEditor/Private/OptimusDetailsCustomization.cpp
@@ -63,8 +63,13 @@ void FOptimusDataTypeRefCustomization::CustomizeHeader(
 	)
 {
 	// Usage mask can change on a per-instance basis when the multi-level data domain field changes in a shader parameter binding
-	auto GetUsageMask = [InPropertyHandle]()
+	auto GetUsageMask = [this, InPropertyHandle]()
 	{
+		if (UsageMaskOverride != EOptimusDataTypeUsageFlags::None)
+		{
+			return UsageMaskOverride;
+		}
+		
 		EOptimusDataTypeUsageFlags UsageMask = EOptimusDataTypeUsageFlags::None;
 	
 		if (InPropertyHandle->HasMetaData(FName(TEXT("UseInResource"))))
@@ -80,40 +85,6 @@ void FOptimusDataTypeRefCustomization::CustomizeHeader(
 			UsageMask |= EOptimusDataTypeUsageFlags::AnimAttributes;
 		}
 
-		if (const FString* InstanceMetaData = InPropertyHandle->GetInstanceMetaData(FName(TEXT("UseInResource"))))
-		{
-			if (*InstanceMetaData == "True")
-			{
-				UsageMask |= EOptimusDataTypeUsageFlags::Resource;
-			}
-			else
-			{
-				UsageMask &= ~EOptimusDataTypeUsageFlags::Resource;
-			}
-		}
-		if (const FString* InstanceMetaData = InPropertyHandle->GetInstanceMetaData(FName(TEXT("UseInVariable"))))
-		{
-			if (*InstanceMetaData == "True")
-			{
-				UsageMask |= EOptimusDataTypeUsageFlags::Variable;
-			}
-			else
-			{
-				UsageMask &= ~EOptimusDataTypeUsageFlags::Variable;
-			}
-		}
-		if (const FString* InstanceMetaData = InPropertyHandle->GetInstanceMetaData(FName(TEXT("UseInAnimAttribute"))))
-		{
-			if (*InstanceMetaData == "True")
-			{
-				UsageMask |= EOptimusDataTypeUsageFlags::AnimAttributes;
-			}
-			else
-			{
-				UsageMask &= ~EOptimusDataTypeUsageFlags::AnimAttributes;
-			}
-		}
-		
 		return UsageMask;
 	};
 
@@ -129,6 +100,10 @@ void FOptimusDataTypeRefCustomization::CustomizeHeader(
 	.ValueContent()
 	[
 		SNew(SOptimusDataTypeSelector)
+		.IsEnabled_Lambda([InPropertyHandle]() -> bool
+		{
+			return InPropertyHandle->IsEditable();
+		})
 		.CurrentDataType(this, &FOptimusDataTypeRefCustomization::GetCurrentDataType)
 		.UsageMask_Lambda(GetUsageMask)
 		.Font(InCustomizationUtils.GetRegularFont())
@@ -1111,20 +1086,26 @@ public:
 		FDetailWidgetRow DataDomainHeaderRow;
 		DataDomainCustomizationInstance = FOptimusDataDomainCustomization::MakeInstance();
 		StaticCastSharedPtr<FOptimusDataDomainCustomization>(DataDomainCustomizationInstance)
-			->OnDataDomainChangedDelegate.AddLambda([DataTypeProperty](const FOptimusDataDomain& InDataDomain)
+			->OnDataDomainChangedDelegate.AddLambda([this, SelectedObjects](const FOptimusDataDomain& InDataDomain)
 			{
-				if (InDataDomain.IsSingleton())
-				{
-					DataTypeProperty->SetInstanceMetaData(FName(TEXT("UseInAnimAttribute")), TEXT("True"));
-					DataTypeProperty->SetInstanceMetaData(FName(TEXT("UseInVariable")), TEXT("True"));
-					DataTypeProperty->SetInstanceMetaData(FName(TEXT("UseInResource")), TEXT("False"));
-				}
-				else
+				EOptimusDataTypeUsageFlags AllowedFlags = EOptimusDataTypeUsageFlags::None;
+
+				for (TWeakObjectPtr<UObject> Object : SelectedObjects)
 				{
-					DataTypeProperty->SetInstanceMetaData(FName(TEXT("UseInAnimAttribute")), TEXT("False"));
-					DataTypeProperty->SetInstanceMetaData(FName(TEXT("UseInVariable")), TEXT("False"));
-					DataTypeProperty->SetInstanceMetaData(FName(TEXT("UseInResource")), TEXT("True"));
+					if (const IOptimusParameterBindingProvider* BindingProvider = Cast<const IOptimusParameterBindingProvider>(Object))
+					{
+						if (AllowedFlags == EOptimusDataTypeUsageFlags::None)
+						{
+							AllowedFlags = BindingProvider->GetTypeUsageFlags(InDataDomain);
+						}
+						else
+						{
+							AllowedFlags &= BindingProvider->GetTypeUsageFlags(InDataDomain);
+						}
+					}
 				}
+
+				StaticCastSharedPtr<FOptimusDataTypeRefCustomization>(DataTypeRefCustomizationInstance)->SetUsageMaskOverride(AllowedFlags);
 			});
 		DataDomainCustomizationInstance->CustomizeHeader(DataDomainProperty.ToSharedRef(), DataDomainHeaderRow, InCustomizationUtils);
 		
@@ -1185,7 +1166,11 @@ public:
 									FSimpleDelegate::CreateSP(this, &SOptimusParameterBindingValueWidget::OnDeleteItem)
 						);
 					}),
-					LOCTEXT("OptimusParameterBindingRemoveButton", "Remove this Binding"))
+					LOCTEXT("OptimusParameterBindingRemoveButton", "Remove this Binding"),
+					TAttribute<bool>::CreateLambda([this]()
+					{
+						return BindingPropertyHandle->IsEditable();
+					}))
 			]
 		];
 	}
@@ -1317,6 +1302,10 @@ void FOptimusParameterBindingCustomization::CustomizeHeader(
 		.Padding(0,0,10,0)
 		[
 			SNew(SEditableTextBox)
+			.IsEnabled_Lambda([NameProperty]() -> bool
+			{
+				return NameProperty->IsEditable();
+			})
 			.Font(InCustomizationUtils.GetRegularFont())
 			.Text_Lambda([NameProperty]()
 			{
diff --git a/Engine/Plugins/Animation/DeformerGraph/Source/OptimusEditor/Private/OptimusDetailsCustomization.h b/Engine/Plugins/Animation/DeformerGraph/Source/OptimusEditor/Private/OptimusDetailsCustomization.h
--- a/Engine/Plugins/Animation/DeformerGraph/Source/OptimusEditor/Private/OptimusDetailsCustomization.h
+++ b/Engine/Plugins/Animation/DeformerGraph/Source/OptimusEditor/Private/OptimusDetailsCustomization.h
@@ -41,7 +41,13 @@ public:
 		IPropertyTypeCustomizationUtils& InCustomizationUtils
 		) override;
 
-private:	
+	void SetUsageMaskOverride(EOptimusDataTypeUsageFlags InOverride)
+	{
+		UsageMaskOverride = InOverride;
+	}
+
+private:
+	
 	FOptimusDataTypeHandle GetCurrentDataType() const;
 	void OnDataTypeChanged(FOptimusDataTypeHandle InDataType);
 
@@ -51,6 +57,8 @@ private:
 	TSharedPtr<IPropertyHandle> TypeObjectProperty;
 	TAttribute<FOptimusDataTypeHandle> CurrentDataType;
 
+	EOptimusDataTypeUsageFlags UsageMaskOverride = EOptimusDataTypeUsageFlags::None;
+	
 };
 
 
diff --git a/Engine/Plugins/Animation/DeformerGraph/Source/OptimusEditor/Private/OptimusEditor.cpp b/Engine/Plugins/Animation/DeformerGraph/Source/OptimusEditor/Private/OptimusEditor.cpp
--- a/Engine/Plugins/Animation/DeformerGraph/Source/OptimusEditor/Private/OptimusEditor.cpp
+++ b/Engine/Plugins/Animation/DeformerGraph/Source/OptimusEditor/Private/OptimusEditor.cpp
@@ -26,6 +26,8 @@
 #include "GraphEditor.h"
 #include "IAssetFamily.h"
 #include "IMessageLogListing.h"
+#include "IOptimusAlternativeSelectedObjectProvider.h"
+#include "IOptimusNodeGraphProvider.h"
 #include "IOptimusShaderTextProvider.h"
 #include "IPersonaPreviewScene.h"
 #include "IPersonaToolkit.h"
@@ -102,9 +104,12 @@ void FOptimusEditor::Construct(
 	// Construct a new graph with a default name
 	// TODO: Use a document manager like blueprints.
 	// FIXME: The deformer asset shouldn't really be the owner.
+
+	// We recreate the editor graph all the time based on the active UOptimusNodeGraph that the user is viewing,
+	// so there is no need for editor graph be to Transactional to keep track of its state
 	EditorGraph = NewObject<UOptimusEditorGraph>(
 		DeformerObject, UOptimusEditorGraph::StaticClass(), NAME_None, 
-		RF_Transactional|RF_Transient);
+		RF_Transient);
 	EditorGraph->Schema = UOptimusEditorGraphSchema::StaticClass();
 
 	constexpr bool bCreateDefaultStandaloneMenu = true;
@@ -696,15 +701,23 @@ void FOptimusEditor::OnSelectedNodesChanged(const TSet<UObject*>& NewSelection)
 
 	for (UObject* Object : NewSelection)
 	{
+		UObject* ObjectToShow;
 		if (UOptimusEditorGraphNode* GraphNode = Cast<UOptimusEditorGraphNode>(Object))
 		{
-			SelectedObjects.Add(GraphNode->ModelNode);
+			ObjectToShow = GraphNode->ModelNode;
 			SelectedNodes.Add(GraphNode);
 		}
 		else
 		{
-			SelectedObjects.Add(Object);
+			ObjectToShow = Object;
 		}
+
+		if (const IOptimusAlternativeSelectedObjectProvider* AlternativeObjectProvider = Cast<IOptimusAlternativeSelectedObjectProvider>(ObjectToShow))
+		{
+			ObjectToShow = AlternativeObjectProvider->GetObjectToShowWhenSelected();
+		}
+		
+		SelectedObjects.AddUnique(ObjectToShow);
 	}
 
 	// Make sure the graph knows too.
@@ -740,7 +753,13 @@ void FOptimusEditor::OnSelectedNodesChanged(const TSet<UObject*>& NewSelection)
 
 void FOptimusEditor::OnNodeDoubleClicked(class UEdGraphNode* Node)
 {
-
+	if (UOptimusEditorGraphNode* GraphNode = Cast<UOptimusEditorGraphNode>(Node))
+	{
+		if (IOptimusNodeGraphProvider* GraphProvider = Cast<IOptimusNodeGraphProvider>(GraphNode->ModelNode))
+		{
+			SetEditGraph(GraphProvider->GetNodeGraphToShow());
+		}
+	}
 }
 
 
diff --git a/Engine/Plugins/Animation/DeformerGraph/Source/OptimusEditor/Private/OptimusEditorGraph.cpp b/Engine/Plugins/Animation/DeformerGraph/Source/OptimusEditor/Private/OptimusEditorGraph.cpp
--- a/Engine/Plugins/Animation/DeformerGraph/Source/OptimusEditor/Private/OptimusEditorGraph.cpp
+++ b/Engine/Plugins/Animation/DeformerGraph/Source/OptimusEditor/Private/OptimusEditorGraph.cpp
@@ -78,7 +78,6 @@ void UOptimusEditorGraph::Reset()
 	SelectedNodes.Reset();
 	NodeGraph = nullptr;
 
-	Modify();
 	TArray<UEdGraphNode*> NodesToRemove(Nodes);
 	for (UEdGraphNode* GraphNode : NodesToRemove)
 	{
@@ -111,11 +110,14 @@ const FSlateBrush* UOptimusEditorGraph::GetGraphTypeIcon(
 	switch(InModelGraph->GetGraphType())
 	{
 	case EOptimusNodeGraphType::Setup:
-			return FOptimusEditorStyle::Get().GetBrush(TEXT("GraphType.Setup"));
+		return FOptimusEditorStyle::Get().GetBrush(TEXT("GraphType.Setup"));
 	case EOptimusNodeGraphType::Update:
-			return FOptimusEditorStyle::Get().GetBrush(TEXT("GraphType.Trigger"));
+		return FOptimusEditorStyle::Get().GetBrush(TEXT("GraphType.Update"));
 	case EOptimusNodeGraphType::ExternalTrigger:
-			return FOptimusEditorStyle::Get().GetBrush(TEXT("GraphType.Update"));
+		return FOptimusEditorStyle::Get().GetBrush(TEXT("GraphType.Trigger"));	
+	case EOptimusNodeGraphType::SubGraph:
+		// TODO: need a new icon
+		return FOptimusEditorStyle::Get().GetBrush(TEXT("GraphType.Update"));
 	default:
 		checkNoEntry();
 		break;
@@ -172,8 +174,6 @@ void UOptimusEditorGraph::HandleNodeGraphModified(EOptimusGraphNotifyType InNoti
 
 			if (ensure(ModelNode))
 			{
-				Modify();
-
 				AddGraphNodeFromModelNode(ModelNode);
 			}
 		    break;
@@ -187,7 +187,6 @@ void UOptimusEditorGraph::HandleNodeGraphModified(EOptimusGraphNotifyType InNoti
 
 			if (ensure(GraphNode))
 			{
-				Modify();
 				RemoveNode(GraphNode, true);
 
 				GraphNode->ModelNode = nullptr;
@@ -209,7 +208,6 @@ void UOptimusEditorGraph::HandleNodeGraphModified(EOptimusGraphNotifyType InNoti
 
 				if (ensure(OutputGraphPin) && ensure(InputGraphPin))
 				{
-					Modify();
 					if (InNotifyType == EOptimusGraphNotifyType::LinkAdded)
 					{
 						OutputGraphPin->MakeLinkTo(InputGraphPin);
diff --git a/Engine/Plugins/Animation/DeformerGraph/Source/OptimusEditor/Private/OptimusEditorGraphNode.cpp b/Engine/Plugins/Animation/DeformerGraph/Source/OptimusEditor/Private/OptimusEditorGraphNode.cpp
--- a/Engine/Plugins/Animation/DeformerGraph/Source/OptimusEditor/Private/OptimusEditorGraphNode.cpp
+++ b/Engine/Plugins/Animation/DeformerGraph/Source/OptimusEditor/Private/OptimusEditorGraphNode.cpp
@@ -338,7 +338,7 @@ FLinearColor UOptimusEditorGraphNode::GetNodeTitleColor() const
 				}
 			}
 			
-			if (ModelNode->GetOwningGraph()->DoesNodeHaveMutableInput(ModelNode))
+			if (ModelNode->GetOwningGraph()->DoesNodeHaveMutableInput(ModelNode, {}))
 			{
 				return Super::GetNodeTitleColor();
 			}
diff --git a/Engine/Plugins/Animation/DeformerGraph/Source/OptimusEditor/Private/Widgets/SOptimusEditorGraphNode.cpp b/Engine/Plugins/Animation/DeformerGraph/Source/OptimusEditor/Private/Widgets/SOptimusEditorGraphNode.cpp
--- a/Engine/Plugins/Animation/DeformerGraph/Source/OptimusEditor/Private/Widgets/SOptimusEditorGraphNode.cpp
+++ b/Engine/Plugins/Animation/DeformerGraph/Source/OptimusEditor/Private/Widgets/SOptimusEditorGraphNode.cpp
@@ -272,7 +272,7 @@ private:
 			// behind.
 			if (ModelPin->GetPackage() != GetTransientPackage())
 			{
-				ComponentSourceBindings = ModelPin->GetComponentSourceBindings();
+				ComponentSourceBindings = ModelPin->GetComponentSourceBindings({});
 			}
 
 			if (ComponentSourceBindings.IsEmpty())
@@ -306,7 +306,7 @@ private:
 			// behind.
 			if (ModelPin->GetPackage() != GetTransientPackage())
 			{
-				if (ModelPin->IsMutable())
+				if (ModelPin->IsMutable({}))
 				{
 					return FText::FromName(TEXT("Mutable"));
 				}
