[Deformer Graph] Marked vert attribute node as immutable and supports unified dispatch

#rnx

diff --git a/Engine/Plugins/Animation/DeformerGraph/Source/OptimusCore/Private/DataInterfaces/OptimusDataInterfaceSkinnedMeshVertexAttribute.cpp b/Engine/Plugins/Animation/DeformerGraph/Source/OptimusCore/Private/DataInterfaces/OptimusDataInterfaceSkinnedMeshVertexAttribute.cpp
--- a/Engine/Plugins/Animation/DeformerGraph/Source/OptimusCore/Private/DataInterfaces/OptimusDataInterfaceSkinnedMeshVertexAttribute.cpp
+++ b/Engine/Plugins/Animation/DeformerGraph/Source/OptimusCore/Private/DataInterfaces/OptimusDataInterfaceSkinnedMeshVertexAttribute.cpp
@@ -20,8 +20,8 @@ FString UOptimusSkinnedMeshVertexAttributeDataInterface::GetDisplayName() const
 TArray<FOptimusCDIPinDefinition> UOptimusSkinnedMeshVertexAttributeDataInterface::GetPinDefinitions() const
 {
 	TArray<FOptimusCDIPinDefinition> Defs;
-	Defs.Add({"NumVertices", "ReadNumVertices"});
-	Defs.Add({"Value", "ReadValue", UOptimusSkinnedMeshComponentSource::Domains::Vertex, "ReadNumVertices"});
+	Defs.Add({"NumVertices", "ReadNumVertices", false});
+	Defs.Add({"Value", "ReadValue", UOptimusSkinnedMeshComponentSource::Domains::Vertex, "ReadNumVertices", false});
 	return Defs;
 }
 
diff --git a/Engine/Plugins/Animation/DeformerGraph/Source/OptimusCore/Private/DataInterfaces/OptimusDataInterfaceSkinnedMeshVertexAttribute.h b/Engine/Plugins/Animation/DeformerGraph/Source/OptimusCore/Private/DataInterfaces/OptimusDataInterfaceSkinnedMeshVertexAttribute.h
--- a/Engine/Plugins/Animation/DeformerGraph/Source/OptimusCore/Private/DataInterfaces/OptimusDataInterfaceSkinnedMeshVertexAttribute.h
+++ b/Engine/Plugins/Animation/DeformerGraph/Source/OptimusCore/Private/DataInterfaces/OptimusDataInterfaceSkinnedMeshVertexAttribute.h
@@ -27,6 +27,7 @@ class UOptimusSkinnedMeshVertexAttributeDataInterface :
 
 	//~ Begin UComputeDataInterface Interface
 	TCHAR const* GetClassName() const override { return TEXT("VertexAttribute"); }
+	virtual bool CanSupportUnifiedDispatch() const override { return true; };
 	virtual void GetSupportedInputs(TArray<FShaderFunctionDefinition>& OutFunctions) const override;
 	void GetShaderParameters(TCHAR const* UID, FShaderParametersMetadataBuilder& InOutBuilder, FShaderParametersMetadataAllocations& InOutAllocations) const override;
 	void GetHLSL(FString& OutHLSL, FString const& InDataInterfaceName) const override;
