Optimus: Remove dependency on private headers.

#rnx

diff --git a/Engine/Plugins/Animation/DeformerGraph/Source/OptimusEditor/OptimusEditor.Build.cs b/Engine/Plugins/Animation/DeformerGraph/Source/OptimusEditor/OptimusEditor.Build.cs
--- a/Engine/Plugins/Animation/DeformerGraph/Source/OptimusEditor/OptimusEditor.Build.cs
+++ b/Engine/Plugins/Animation/DeformerGraph/Source/OptimusEditor/OptimusEditor.Build.cs
@@ -6,10 +6,6 @@ namespace UnrealBuildTool.Rules
     {
         public OptimusEditor(ReadOnlyTargetRules Target) : base(Target)
         {
-			PrivateIncludePaths.AddRange(new string[] {
-				System.IO.Path.Combine(GetModuleDirectory("OptimusCore"), "Private"),
-			});
-
             PrivateDependencyModuleNames.AddRange(
                 new string[]
                 {
diff --git a/Engine/Plugins/Animation/DeformerGraph/Source/OptimusEditor/Private/OptimusEditorModule.cpp b/Engine/Plugins/Animation/DeformerGraph/Source/OptimusEditor/Private/OptimusEditorModule.cpp
--- a/Engine/Plugins/Animation/DeformerGraph/Source/OptimusEditor/Private/OptimusEditorModule.cpp
+++ b/Engine/Plugins/Animation/DeformerGraph/Source/OptimusEditor/Private/OptimusEditorModule.cpp
@@ -2,15 +2,12 @@
 
 #include "OptimusEditorModule.h"
 
-#include "Actions/OptimusResourceActions.h"
-#include "Algo/Transform.h"
 #include "AssetRegistry/AssetRegistryModule.h"
 #include "AssetToolsModule.h"
 #include "AssetViewUtils.h"
 #include "EdGraphUtilities.h"
 #include "IAssetTools.h"
 #include "OptimusExecutionDomain.h"
-#include "ObjectTools.h"
 #include "OptimusBindingTypes.h"
 #include "OptimusComponentSource.h"
 #include "OptimusDataType.h"
