[Deformer Graph]: Basic support for creating public functions in deformer graph

diff --git a/Engine/Plugins/Animation/DeformerGraph/Source/OptimusCore/Private/Actions/OptimusNodeGraphActions.cpp b/Engine/Plugins/Animation/DeformerGraph/Source/OptimusCore/Private/Actions/OptimusNodeGraphActions.cpp
--- a/Engine/Plugins/Animation/DeformerGraph/Source/OptimusCore/Private/Actions/OptimusNodeGraphActions.cpp
+++ b/Engine/Plugins/Animation/DeformerGraph/Source/OptimusCore/Private/Actions/OptimusNodeGraphActions.cpp
@@ -2,6 +2,7 @@
 
 #include "OptimusNodeGraphActions.h"
 
+#include "IOptimusCoreModule.h"
 #include "IOptimusPathResolver.h"
 #include "IOptimusUnnamedNodePinProvider.h"
 #include "Nodes/OptimusNode_ComputeKernelFunction.h"
@@ -25,7 +26,7 @@
 // ---- Add graph
 
 FOptimusNodeGraphAction_AddGraph::FOptimusNodeGraphAction_AddGraph(
-	IOptimusNodeGraphCollectionOwner* InGraphOwner, 
+	const IOptimusNodeGraphCollectionOwner* InGraphOwner, 
 	EOptimusNodeGraphType InGraphType, 
 	FName InGraphName, 
 	int32 InGraphIndex,
@@ -64,7 +65,7 @@ bool FOptimusNodeGraphAction_AddGraph::Do(
 		return false;
 	}
 	
-	UOptimusNodeGraph* Graph = GraphOwner->CreateGraph(GraphType, GraphName, {});
+	UOptimusNodeGraph* Graph = GraphOwner->CreateGraphDirect(GraphType, GraphName, {});
 	if (!Graph)
 	{
 		return false;
@@ -77,7 +78,7 @@ bool FOptimusNodeGraphAction_AddGraph::Do(
 	}
 
 	// Add the graph to the collection
-	if (!GraphOwner->AddGraph(Graph, GraphIndex))
+	if (!GraphOwner->AddGraphDirect(Graph, GraphIndex))
 	{
 		Optimus::RemoveObject(Graph);
 		return false;
@@ -103,7 +104,7 @@ bool FOptimusNodeGraphAction_AddGraph::Undo(
 		return false;
 	}
 
-	return Graph->GetCollectionOwner()->RemoveGraph(Graph);
+	return Graph->GetCollectionOwner()->RemoveGraphDirect(Graph);
 }
 
 
@@ -144,7 +145,7 @@ bool FOptimusNodeGraphAction_RemoveGraph::Do(IOptimusPathResolver* InRoot)
 		Optimus::FBinaryObjectWriter GraphArchive(Graph, GraphData);
 	}
 	
-	return GraphOwner->RemoveGraph(Graph);
+	return GraphOwner->RemoveGraphDirect(Graph);
 }
 
 
@@ -160,7 +161,7 @@ bool FOptimusNodeGraphAction_RemoveGraph::Undo(
 	
 	// Create a graph, but don't add it to the list of used graphs. Otherwise interested parties
 	// will be notified with a partially constructed graph.
-	UOptimusNodeGraph* Graph = GraphOwner->CreateGraph(GraphType, GraphName, TOptional<int32>());
+	UOptimusNodeGraph* Graph = GraphOwner->CreateGraphDirect(GraphType, GraphName, TOptional<int32>());
 	if (Graph == nullptr)
 	{
 		return false;
@@ -172,7 +173,7 @@ bool FOptimusNodeGraphAction_RemoveGraph::Undo(
 	}
 	
 	// Now add the graph such that interested parties get notified.
-	if (!GraphOwner->AddGraph(Graph, GraphIndex))
+	if (!GraphOwner->AddGraphDirect(Graph, GraphIndex))
 	{
 		Optimus::RemoveObject(Graph);
 		return false;
@@ -192,6 +193,7 @@ FOptimusNodeGraphAction_RenameGraph::FOptimusNodeGraphAction_RenameGraph(
 	if (ensure(InGraph) && InGraph->GetFName() != InNewName)
 	{
 		GraphPath = InGraph->GetGraphPath();
+		GraphOwnerPath = InGraph->GetCollectionOwner()->GetCollectionPath();
 
 		// Ensure the name is unique within our namespace.
 		if (StaticFindObject(UOptimusNodeGraph::StaticClass(), InGraph->GetOuter(), *InNewName.ToString()) != nullptr)
@@ -216,11 +218,23 @@ bool FOptimusNodeGraphAction_RenameGraph::Do(
 	{
 		return false;
 	}
+
+	IOptimusNodeGraphCollectionOwner* GraphOwner = InRoot->ResolveCollectionPath(GraphOwnerPath);
+	if (!GraphOwner)
+	{
+		return false;
+	}
+
+	if (GraphOwner->RenameGraphDirect(Graph, NewGraphName.ToString()))
+	{
+		GraphPath = Graph->GetGraphPath();
+		return true;
+	}
+
+	return false;
 	
-	return Optimus::RenameObject(Graph, *NewGraphName.ToString(), nullptr);
 }
 
-
 bool FOptimusNodeGraphAction_RenameGraph::Undo(
 	IOptimusPathResolver* InRoot
 	)
@@ -231,8 +245,20 @@ bool FOptimusNodeGraphAction_RenameGraph::Undo(
 		return false;
 	}
 
-	return Optimus::RenameObject(Graph, *OldGraphName.ToString(), nullptr);
-}
+	IOptimusNodeGraphCollectionOwner* GraphOwner = InRoot->ResolveCollectionPath(GraphOwnerPath);
+	if (!GraphOwner)
+	{
+		return false;
+	}
+
+	if (GraphOwner->RenameGraphDirect(Graph, OldGraphName.ToString()))
+	{
+		GraphPath = Graph->GetGraphPath();
+		return true;
+	}
+
+	return false;	
+}	
 
 
 // ---- Add node
diff --git a/Engine/Plugins/Animation/DeformerGraph/Source/OptimusCore/Private/Actions/OptimusNodeGraphActions.h b/Engine/Plugins/Animation/DeformerGraph/Source/OptimusCore/Private/Actions/OptimusNodeGraphActions.h
--- a/Engine/Plugins/Animation/DeformerGraph/Source/OptimusCore/Private/Actions/OptimusNodeGraphActions.h
+++ b/Engine/Plugins/Animation/DeformerGraph/Source/OptimusCore/Private/Actions/OptimusNodeGraphActions.h
@@ -34,7 +34,7 @@ public:
 	FOptimusNodeGraphAction_AddGraph() = default;
 
 	FOptimusNodeGraphAction_AddGraph(
-	    IOptimusNodeGraphCollectionOwner* InGraphOwner,
+	    const IOptimusNodeGraphCollectionOwner* InGraphOwner,
 		EOptimusNodeGraphType InGraphType,
 		FName InGraphName,
 		int32 InGraphIndex,
@@ -126,6 +126,10 @@ protected:
 	bool Undo(IOptimusPathResolver* InRoot) override;
 
 private:
+
+	// The path of the graph owner that will actually perform the rename 
+	FString GraphOwnerPath;
+	
 	// The path of the graph to rename. This value will change after each rename.
 	FString GraphPath;
 
diff --git a/Engine/Plugins/Animation/DeformerGraph/Source/OptimusCore/Private/Nodes/OptimusNode_CustomComputeKernel.cpp b/Engine/Plugins/Animation/DeformerGraph/Source/OptimusCore/Private/Nodes/OptimusNode_CustomComputeKernel.cpp
--- a/Engine/Plugins/Animation/DeformerGraph/Source/OptimusCore/Private/Nodes/OptimusNode_CustomComputeKernel.cpp
+++ b/Engine/Plugins/Animation/DeformerGraph/Source/OptimusCore/Private/Nodes/OptimusNode_CustomComputeKernel.cpp
@@ -185,6 +185,16 @@ void UOptimusNode_CustomComputeKernel::SetShaderText(const FString& NewText)
 	PostEditChangeProperty(PropertyChangedEvent);
 }
 
+bool UOptimusNode_CustomComputeKernel::IsShaderTextReadOnly() const
+{
+	if (GetOwningGraph()->GetGraphType() == EOptimusNodeGraphType::Function)
+	{
+		return true;
+	}
+
+	return false;
+}
+
 #endif // WITH_EDITOR
 
 FString UOptimusNode_CustomComputeKernel::GetBindingDeclaration(
diff --git a/Engine/Plugins/Animation/DeformerGraph/Source/OptimusCore/Private/Nodes/OptimusNode_CustomComputeKernel.h b/Engine/Plugins/Animation/DeformerGraph/Source/OptimusCore/Private/Nodes/OptimusNode_CustomComputeKernel.h
--- a/Engine/Plugins/Animation/DeformerGraph/Source/OptimusCore/Private/Nodes/OptimusNode_CustomComputeKernel.h
+++ b/Engine/Plugins/Animation/DeformerGraph/Source/OptimusCore/Private/Nodes/OptimusNode_CustomComputeKernel.h
@@ -71,6 +71,7 @@ public:
 	FString GetDeclarations() const override;
 	FString GetShaderText() const override;
 	void SetShaderText(const FString& NewText) override;
+	bool IsShaderTextReadOnly() const override;
 	// IOptimusShaderTextProvider overrides
 #endif
 	
diff --git a/Engine/Plugins/Animation/DeformerGraph/Source/OptimusCore/Private/Nodes/OptimusNode_DataInterface.cpp b/Engine/Plugins/Animation/DeformerGraph/Source/OptimusCore/Private/Nodes/OptimusNode_DataInterface.cpp
--- a/Engine/Plugins/Animation/DeformerGraph/Source/OptimusCore/Private/Nodes/OptimusNode_DataInterface.cpp
+++ b/Engine/Plugins/Animation/DeformerGraph/Source/OptimusCore/Private/Nodes/OptimusNode_DataInterface.cpp
@@ -211,15 +211,23 @@ UOptimusComponentSourceBinding* UOptimusNode_DataInterface::GetComponentBinding(
 	}
 
 	// Default to the primary binding, but only if we're at the top-most level of the graph.
-	if (const UOptimusDeformer* Deformer = Cast<UOptimusDeformer>(Graph->GetCollectionOwner()))
+	if (Optimus::IsExecutionGraphType(Graph->GetGraphType()))
 	{
-		return Deformer->GetPrimaryComponentBinding();
+		const UOptimusDeformer* Deformer = Cast<UOptimusDeformer>(Graph->GetCollectionOwner());
+		if (ensure(Deformer))
+		{
+			return Deformer->GetPrimaryComponentBinding();
+		}
 	}
-
-	if (const UOptimusNodeSubGraph* SubGraph = Cast<UOptimusNodeSubGraph>(Graph))
+	else
 	{
-		return SubGraph->GetDefaultComponentBinding(InContext);
+		const UOptimusNodeSubGraph* SubGraph = Cast<UOptimusNodeSubGraph>(Graph);
+		if (ensure(SubGraph))
+		{
+			return SubGraph->GetDefaultComponentBinding(InContext);
+		}
 	}
+
 	
 	return nullptr;
 }
diff --git a/Engine/Plugins/Animation/DeformerGraph/Source/OptimusCore/Private/Nodes/OptimusNode_FunctionReference.cpp b/Engine/Plugins/Animation/DeformerGraph/Source/OptimusCore/Private/Nodes/OptimusNode_FunctionReference.cpp
--- a/Engine/Plugins/Animation/DeformerGraph/Source/OptimusCore/Private/Nodes/OptimusNode_FunctionReference.cpp
+++ b/Engine/Plugins/Animation/DeformerGraph/Source/OptimusCore/Private/Nodes/OptimusNode_FunctionReference.cpp
@@ -2,27 +2,63 @@
 
 #include "OptimusNode_FunctionReference.h"
 
+#include "IOptimusCoreModule.h"
+#include "OptimusDataTypeRegistry.h"
 #include "OptimusFunctionNodeGraph.h"
 #include "OptimusNodePin.h"
 #include "OptimusNode_GraphTerminal.h"
+#include "OptimusComponentSource.h"
+#include "OptimusDeformer.h"
 
 #include UE_INLINE_GENERATED_CPP_BY_NAME(OptimusNode_FunctionReference)
 
 
+FName UOptimusNode_FunctionReference::GetNodeCategory() const
+{
+	return FunctionGraph->Category;
+
+}
+
 FText UOptimusNode_FunctionReference::GetDisplayName() const
 {
-	if (FunctionGraph)
+	if (FunctionGraph.IsValid())
 	{
 		return FText::FromString(FunctionGraph->GetNodeName());
 	}
 
+	
 	return FText::FromString("<graph missing>");
 }
 
 
 void UOptimusNode_FunctionReference::ConstructNode()
 {
-	
+	if (ensure(FunctionGraph.IsValid()))
+	{
+		const FOptimusDataTypeRegistry& TypeRegistry = FOptimusDataTypeRegistry::Get();
+		FOptimusDataTypeRef ComponentSourceType = TypeRegistry.FindType(*UOptimusComponentSourceBinding::StaticClass());
+		DefaultComponentPin = AddPinDirect(UOptimusNodeSubGraph::GraphDefaultComponentPinName, EOptimusNodePinDirection::Input, {}, ComponentSourceType);
+		
+		// After a duplicate, the kernel node has no pins, so we need to reconstruct them from
+		// the bindings. We can assume that all naming clashes have already been dealt with.
+		for (const FOptimusParameterBinding& Binding: FunctionGraph->InputBindings)
+		{
+			AddPinDirect(Binding, EOptimusNodePinDirection::Input);
+		}
+		for (const FOptimusParameterBinding& Binding: FunctionGraph->OutputBindings)
+		{
+			AddPinDirect(Binding, EOptimusNodePinDirection::Output);
+		}
+	}
+}
+
+
+void UOptimusNode_FunctionReference::PostLoad()
+{
+	Super::PostLoad();
+
+	// Load the graph into memory
+	FunctionGraph.LoadSynchronous();
 }
 
 
@@ -36,7 +72,7 @@ FOptimusRoutedNodePin UOptimusNode_FunctionReference::GetPinCounterpart(
 		return {};
 	}
 
-	if (!ensure(FunctionGraph))
+	if (!ensure(FunctionGraph.IsValid()))
 	{
 		return {};
 	}
@@ -64,3 +100,51 @@ FOptimusRoutedNodePin UOptimusNode_FunctionReference::GetPinCounterpart(
 
 	return Result;
 }
+
+UOptimusNodeGraph* UOptimusNode_FunctionReference::GetNodeGraphToShow()
+{
+	return FunctionGraph.Get();
+}
+
+UOptimusComponentSourceBinding* UOptimusNode_FunctionReference::GetDefaultComponentBinding(const FOptimusPinTraversalContext& InTraversalContext) const
+{
+	if (!ensure(DefaultComponentPin.IsValid()))
+	{
+		return nullptr;
+	}
+	
+	const UOptimusNodeGraph* OwningGraph = GetOwningGraph();
+	TSet<UOptimusComponentSourceBinding*> Bindings = OwningGraph->GetComponentSourceBindingsForPin(DefaultComponentPin.Get(), InTraversalContext);
+	
+	if (!Bindings.IsEmpty() && ensure(Bindings.Num() == 1))
+	{
+		return Bindings.Array()[0];
+	}
+
+	// Default to the primary binding, but only if we're at the top-most level of the graph.
+	if (const UOptimusDeformer* Deformer = Cast<UOptimusDeformer>(OwningGraph->GetCollectionOwner()))
+	{
+		return Deformer->GetPrimaryComponentBinding();
+	}
+
+	if (const UOptimusNodeSubGraph* OwningSubGraph = Cast<UOptimusNodeSubGraph>(OwningGraph))
+	{
+		return OwningSubGraph->GetDefaultComponentBinding(InTraversalContext);
+	}
+
+	return nullptr;		
+}
+
+FSoftObjectPath UOptimusNode_FunctionReference::GetSerializedGraphPath() const
+{
+	return FunctionGraph.ToSoftObjectPath();
+}
+
+void UOptimusNode_FunctionReference::SetSerializedGraphPath(const FSoftObjectPath& InNewGraphPath)
+{
+	FunctionGraph = InNewGraphPath;
+	// Making sure the pointer is alive such that we avoid loading on demand everywhere
+	FunctionGraph.LoadSynchronous();
+
+	SetDisplayName(GetDisplayName());
+}
diff --git a/Engine/Plugins/Animation/DeformerGraph/Source/OptimusCore/Private/Nodes/OptimusNode_FunctionReference.h b/Engine/Plugins/Animation/DeformerGraph/Source/OptimusCore/Private/Nodes/OptimusNode_FunctionReference.h
--- a/Engine/Plugins/Animation/DeformerGraph/Source/OptimusCore/Private/Nodes/OptimusNode_FunctionReference.h
+++ b/Engine/Plugins/Animation/DeformerGraph/Source/OptimusCore/Private/Nodes/OptimusNode_FunctionReference.h
@@ -1,7 +1,11 @@
 ﻿// Copyright Epic Games, Inc. All Rights Reserved.
 #pragma once
 
+#include "IOptimusNodeGraphProvider.h"
 #include "IOptimusNodePinRouter.h"
+#include "IOptimusNodeSubGraphReferencer.h"
+#include "OptimusFunctionNodeGraph.h"
+#include "OptimusFunctionNodeGraphHeader.h"
 #include "OptimusNode.h"
 
 #include "OptimusNode_FunctionReference.generated.h"
@@ -9,30 +13,51 @@
 
 class UOptimusFunctionNodeGraph;
 
-
 UCLASS(Hidden)
 class UOptimusNode_FunctionReference :
 	public UOptimusNode,
-	public IOptimusNodePinRouter
+	public IOptimusNodePinRouter,
+	public IOptimusNodeGraphProvider,
+	public IOptimusNodeSubGraphReferencer
 {
 	GENERATED_BODY()
 
 public:
 	// UOptimusNode overrides
-	FName GetNodeCategory() const override { return NAME_None; }
+	FName GetNodeCategory() const override;
 	FText GetDisplayName() const override;
 	void ConstructNode() override;
 
+	// UObject
+	void PostLoad() override;
+	
 	// IOptimusNodePinRouter implementation
 	FOptimusRoutedNodePin GetPinCounterpart(
 		UOptimusNodePin* InNodePin,
 		const FOptimusPinTraversalContext& InTraversalContext
 	) const override;
 
+	// IOptimusNodeGraphProvider
+	UOptimusNodeGraph* GetNodeGraphToShow() override;
+	
+	// IOptimusNodeSubGraphReferencer
+	UOptimusComponentSourceBinding* GetDefaultComponentBinding(const FOptimusPinTraversalContext& InTraversalContext) const override;
+
+	FSoftObjectPath GetSerializedGraphPath() const;
+
+	void SetSerializedGraphPath(const FSoftObjectPath& InNewGraphPath);
+	
 protected:
+	friend class UOptimusNodeGraph;
+	friend class FOptimusCoreModule;
+
+	
 	/** The graph that owns us. This contains all the necessary pin information to add on
 	 * the terminal node.
 	 */
 	UPROPERTY()
-	TObjectPtr<UOptimusFunctionNodeGraph> FunctionGraph;
+	TSoftObjectPtr<UOptimusFunctionNodeGraph> FunctionGraph;
+
+	UPROPERTY()
+	TWeakObjectPtr<UOptimusNodePin> DefaultComponentPin;
 };
diff --git a/Engine/Plugins/Animation/DeformerGraph/Source/OptimusCore/Private/Nodes/OptimusNode_GraphTerminal.cpp b/Engine/Plugins/Animation/DeformerGraph/Source/OptimusCore/Private/Nodes/OptimusNode_GraphTerminal.cpp
--- a/Engine/Plugins/Animation/DeformerGraph/Source/OptimusCore/Private/Nodes/OptimusNode_GraphTerminal.cpp
+++ b/Engine/Plugins/Animation/DeformerGraph/Source/OptimusCore/Private/Nodes/OptimusNode_GraphTerminal.cpp
@@ -16,6 +16,8 @@
 
 #define LOCTEXT_NAMESPACE "OptimusNodeGraphTerminal"
 
+FName UOptimusNode_GraphTerminal::EntryNodeName = TEXT("Entry");
+FName UOptimusNode_GraphTerminal::ReturnNodeName= TEXT("Return");
 
 UOptimusNode_GraphTerminal::UOptimusNode_GraphTerminal()
 {
@@ -142,7 +144,7 @@ UOptimusComponentSourceBinding* UOptimusNode_GraphTerminal::GetDefaultComponentB
 		return nullptr;
 	}
 	
-	const UOptimusNode_SubGraphReference* ReferenceNode = CastChecked<UOptimusNode_SubGraphReference>(PinCounterpart.NodePin->GetOwningNode());
+	const IOptimusNodeSubGraphReferencer* ReferenceNode = CastChecked<IOptimusNodeSubGraphReferencer>(PinCounterpart.NodePin->GetOwningNode());
 
 	return ReferenceNode->GetDefaultComponentBinding(PinCounterpart.TraversalContext);
 }
diff --git a/Engine/Plugins/Animation/DeformerGraph/Source/OptimusCore/Private/Nodes/OptimusNode_GraphTerminal.h b/Engine/Plugins/Animation/DeformerGraph/Source/OptimusCore/Private/Nodes/OptimusNode_GraphTerminal.h
--- a/Engine/Plugins/Animation/DeformerGraph/Source/OptimusCore/Private/Nodes/OptimusNode_GraphTerminal.h
+++ b/Engine/Plugins/Animation/DeformerGraph/Source/OptimusCore/Private/Nodes/OptimusNode_GraphTerminal.h
@@ -31,6 +31,8 @@ class UOptimusNode_GraphTerminal :
 	GENERATED_BODY()
 	
 public:
+	static FName EntryNodeName;
+	static FName ReturnNodeName;
 	UOptimusNode_GraphTerminal();
 	
 	// UOptimusNode overrides
diff --git a/Engine/Plugins/Animation/DeformerGraph/Source/OptimusCore/Private/Nodes/OptimusNode_SubGraphReference.h b/Engine/Plugins/Animation/DeformerGraph/Source/OptimusCore/Private/Nodes/OptimusNode_SubGraphReference.h
--- a/Engine/Plugins/Animation/DeformerGraph/Source/OptimusCore/Private/Nodes/OptimusNode_SubGraphReference.h
+++ b/Engine/Plugins/Animation/DeformerGraph/Source/OptimusCore/Private/Nodes/OptimusNode_SubGraphReference.h
@@ -4,6 +4,7 @@
 
 #include "IOptimusNodeGraphProvider.h"
 #include "IOptimusNodePinRouter.h"
+#include "IOptimusNodeSubGraphReferencer.h"
 #include "OptimusNode.h"
 
 #include "OptimusNode_SubGraphReference.generated.h"
@@ -13,10 +14,11 @@ class UOptimusNodeSubGraph;
 
 
 UCLASS(Hidden)
-class UOptimusNode_SubGraphReference :
+class OPTIMUSCORE_API UOptimusNode_SubGraphReference :
 	public UOptimusNode,
 	public IOptimusNodePinRouter,
-	public IOptimusNodeGraphProvider
+	public IOptimusNodeGraphProvider,
+	public IOptimusNodeSubGraphReferencer
 {
 	GENERATED_BODY()
 
@@ -38,12 +40,15 @@ public:
 		const FOptimusPinTraversalContext& InTraversalContext
 	) const override;
 
+	// IOptimusNodeGraphProvider
 	UOptimusNodeGraph* GetNodeGraphToShow() override;
 
-	UOptimusComponentSourceBinding* GetDefaultComponentBinding(const FOptimusPinTraversalContext& InTraversalContext) const;
-	
+	// IOptimusNodeSubGraphReferencer
+	UOptimusComponentSourceBinding* GetDefaultComponentBinding(const FOptimusPinTraversalContext& InTraversalContext) const override;
+
 protected:
 	friend class UOptimusNodeGraph;
+	friend class UOptimusDeformer;
 
 	void SubscribeToSubGraph();
 	void UnsubscribeFromSubGraph() const;
diff --git a/Engine/Plugins/Animation/DeformerGraph/Source/OptimusCore/Private/OptimusCoreModule.cpp b/Engine/Plugins/Animation/DeformerGraph/Source/OptimusCore/Private/OptimusCoreModule.cpp
--- a/Engine/Plugins/Animation/DeformerGraph/Source/OptimusCore/Private/OptimusCoreModule.cpp
+++ b/Engine/Plugins/Animation/DeformerGraph/Source/OptimusCore/Private/OptimusCoreModule.cpp
@@ -5,11 +5,15 @@
 #include "Interfaces/IPluginManager.h"
 #include "OptimusComputeDataInterface.h"
 #include "OptimusDataTypeRegistry.h"
+#include "OptimusDeformer.h"
 #include "OptimusObjectVersion.h"
 #include "Misc/Paths.h"
 #include "ShaderCore.h"
 #include "UObject/DevObjectVersion.h"
 #include "Modules/ModuleManager.h"
+#include "AssetRegistry/AssetRegistryModule.h"
+#include "AssetRegistry/AssetData.h"
+#include "Nodes/OptimusNode_FunctionReference.h"
 
 // Unique serialization id for Optimus .
 const FGuid FOptimusObjectVersion::GUID(0x93ede1aa, 0x10ca7375, 0x4df98a28, 0x49b157a0);
@@ -25,12 +29,14 @@ void FOptimusCoreModule::StartupModule()
 	FOptimusDataTypeRegistry::RegisterBuiltinTypes();
 	FOptimusDataTypeRegistry::RegisterEngineCallbacks();
 	UOptimusComputeDataInterface::RegisterAllTypes();
+	
 }
 
 void FOptimusCoreModule::ShutdownModule()
 {
 	FOptimusDataTypeRegistry::UnregisterEngineCallbacks();
 	FOptimusDataTypeRegistry::UnregisterAllTypes();
+	
 }
 
 
@@ -48,6 +54,54 @@ bool FOptimusCoreModule::RegisterDataInterfaceClass(TSubclassOf<UOptimusComputeD
 	return false;
 }
 
+void FOptimusCoreModule::UpdateFunctionReferences(const FSoftObjectPath& InOldGraphPath, const FSoftObjectPath& InNewGraphPath)
+{
+	const FAssetRegistryModule& AssetRegistryModule = FModuleManager::LoadModuleChecked<FAssetRegistryModule>("AssetRegistry");
+	
+	// find all assets in the project
+	TArray<FAssetData> AssetDatas;
+	AssetRegistryModule.Get().GetAssetsByClass(UOptimusDeformer::StaticClass()->GetClassPathName(), AssetDatas);
+
+	FOptimusFunctionReferenceData FunctionReferenceData;
+	
+	// loop over all found assets
+	for(const FAssetData& AssetData : AssetDatas)
+	{
+		RegisterFunctionReferencesFromAsset(FunctionReferenceData, AssetData);
+	}
+	
+	if (FOptimusFunctionReferenceNodeSet* FunctionNodeArray = FunctionReferenceData.FunctionReferences.Find(InOldGraphPath))
+	{
+		for (TSoftObjectPtr<UOptimusNode_FunctionReference> FunctionNodePtr : FunctionNodeArray->Nodes)
+		{
+			UOptimusNode_FunctionReference* FunctionNode = FunctionNodePtr.LoadSynchronous();
+
+			if (FunctionNode)
+			{
+				FunctionNode->Modify();
+				FunctionNode->SetSerializedGraphPath(InNewGraphPath);
+				FunctionNode->MarkPackageDirty();
+			}
+		}
+	}
+}
+
+void FOptimusCoreModule::RegisterFunctionReferencesFromAsset(FOptimusFunctionReferenceData& InOutData, const FAssetData& AssetData)
+{
+	FString FunctionReferenceString = AssetData.GetTagValueRef<FString>(UOptimusDeformer::FunctionReferencesAssetTagName);
+
+	if (!FunctionReferenceString.IsEmpty())
+	{
+		FOptimusFunctionReferenceData AssetFunctionReferenceData;
+
+		FOptimusFunctionReferenceData::StaticStruct()->ImportText(*FunctionReferenceString, &AssetFunctionReferenceData, nullptr, PPF_None, nullptr, {});
+
+		for(const TPair<FSoftObjectPath, FOptimusFunctionReferenceNodeSet>& Pair : AssetFunctionReferenceData.FunctionReferences)
+		{
+			InOutData.FunctionReferences.FindOrAdd(Pair.Key).Nodes.Append(Pair.Value.Nodes);
+		}
+	}
+}
 
 IMPLEMENT_MODULE(FOptimusCoreModule, OptimusCore)
 
diff --git a/Engine/Plugins/Animation/DeformerGraph/Source/OptimusCore/Private/OptimusCoreModule.h b/Engine/Plugins/Animation/DeformerGraph/Source/OptimusCore/Private/OptimusCoreModule.h
--- a/Engine/Plugins/Animation/DeformerGraph/Source/OptimusCore/Private/OptimusCoreModule.h
+++ b/Engine/Plugins/Animation/DeformerGraph/Source/OptimusCore/Private/OptimusCoreModule.h
@@ -3,6 +3,7 @@
 #pragma once
 
 #include "IOptimusCoreModule.h"
+#include "OptimusFunctionNodeGraphHeader.h"
 
 
 class FOptimusCoreModule : public IOptimusCoreModule
@@ -12,8 +13,16 @@ public:
 	virtual void StartupModule() override;
 	virtual void ShutdownModule() override;
 
+
 	/** IOptimusCoreModule implementation */
 	bool RegisterDataInterfaceClass(TSubclassOf<UOptimusComputeDataInterface> InDataInterfaceClass) override;
+
+	void UpdateFunctionReferences(const FSoftObjectPath& InOldGraphPath, const FSoftObjectPath& InNewGraphPath) override;
+	
+protected:
+	void RegisterFunctionReferencesFromAsset(FOptimusFunctionReferenceData& InOutData, const FAssetData& AssetData);
+	
+	
 };
 
 DECLARE_LOG_CATEGORY_EXTERN(LogOptimusCore, Log, All);
diff --git a/Engine/Plugins/Animation/DeformerGraph/Source/OptimusCore/Private/OptimusDeformer.cpp b/Engine/Plugins/Animation/DeformerGraph/Source/OptimusCore/Private/OptimusDeformer.cpp
--- a/Engine/Plugins/Animation/DeformerGraph/Source/OptimusCore/Private/OptimusDeformer.cpp
+++ b/Engine/Plugins/Animation/DeformerGraph/Source/OptimusCore/Private/OptimusDeformer.cpp
@@ -57,6 +57,10 @@
 
 #include <limits>
 
+#include "OptimusFunctionNodeGraphHeader.h"
+#include "Nodes/OptimusNode_FunctionReference.h"
+#include "Nodes/OptimusNode_SubGraphReference.h"
+
 #define PRINT_COMPILED_OUTPUT 1
 
 #define LOCTEXT_NAMESPACE "OptimusDeformer"
@@ -64,7 +68,8 @@
 static const FName DefaultResourceName("Resource");
 static const FName DefaultVariableName("Variable");
 
-
+const FName UOptimusDeformer::PublicFunctionsAssetTagName = TEXT("PublicFunctions");
+const FName UOptimusDeformer::FunctionReferencesAssetTagName = TEXT("FunctionReferences");
 
 UOptimusDeformer::UOptimusDeformer()
 {
@@ -145,7 +150,38 @@ UOptimusNodeGraph* UOptimusDeformer::GetUpdateGraph() const
 
 bool UOptimusDeformer::RemoveGraph(UOptimusNodeGraph* InGraph)
 {
-    return GetActionStack()->RunAction<FOptimusNodeGraphAction_RemoveGraph>(InGraph);
+	FOptimusCompoundAction *Action = new FOptimusCompoundAction(TEXT("Remove Graph"));
+
+	if (UOptimusNodeSubGraph* SubGraph = Cast<UOptimusNodeSubGraph>(InGraph))
+	{
+		if (UOptimusNode* Node = GetSubGraphReferenceNode(SubGraph))
+		{
+			// Remove node also triggers the removal of the graph
+			Node->GetOwningGraph()->RemoveNodesToAction(Action, {Node} );
+			return GetActionStack()->RunAction(Action);
+		}
+	}
+	
+	Action->AddSubAction<FOptimusNodeGraphAction_RemoveGraph>(InGraph);
+	
+    return GetActionStack()->RunAction(Action);
+}
+
+UOptimusNode* UOptimusDeformer::GetSubGraphReferenceNode(const UOptimusNodeSubGraph* InSubGraph) const
+{
+	UOptimusNode* UsedNode = nullptr;
+	
+	TArray<UOptimusNode*> AllSubGraphNodes = GetAllNodesOfClass(UOptimusNode_SubGraphReference::StaticClass());
+	for (UOptimusNode* Node: AllSubGraphNodes)
+	{
+		const UOptimusNode_SubGraphReference* SubGraphReference = Cast<UOptimusNode_SubGraphReference>(Node);
+		if (SubGraphReference->SubGraph == InSubGraph)
+		{
+			UsedNode = Node;
+			break;
+		}
+	}
+	return UsedNode;
 }
 
 
@@ -1462,8 +1498,11 @@ bool UOptimusDeformer::Compile()
 
 	for (const UOptimusNodeGraph* Graph: Graphs)
 	{
-		TArray<FOptimusComputeGraphInfo> ComputeGraphInfos = CompileNodeGraphToComputeGraphs(Graph, ErrorReporter);
-		ComputeGraphs.Append(ComputeGraphInfos);
+		if (Graph->GetGraphType() != EOptimusNodeGraphType::Function)
+		{
+			TArray<FOptimusComputeGraphInfo> ComputeGraphInfos = CompileNodeGraphToComputeGraphs(Graph, ErrorReporter);
+			ComputeGraphs.Append(ComputeGraphInfos);
+		}
 	}
 
 	CompileEndDelegate.Broadcast(this);
@@ -2822,11 +2861,24 @@ UOptimusNodeGraph* UOptimusDeformer::ResolveGraphPath(
 	{
 		return nullptr;
 	}
-
+	int32 SubGraphStartIndex = 0;
+	
 	UOptimusNodeGraph* Graph = nullptr;
 	if (Path[0] == UOptimusNodeGraph::LibraryRoot)
 	{
 		// FIXME: Search the library graphs.
+		for (UOptimusNodeGraph* RootGraph : Graphs)
+		{
+			if (RootGraph->GetGraphType() == EOptimusNodeGraphType::Function)
+			{
+				if (Path[1].Equals(RootGraph->GetName(), ESearchCase::IgnoreCase))
+				{
+					SubGraphStartIndex = 2;
+					Graph = RootGraph;
+					break;
+				}
+			}
+		}
 	}
 	else
 	{
@@ -2834,6 +2886,7 @@ UOptimusNodeGraph* UOptimusDeformer::ResolveGraphPath(
 		{
 			if (Path[0].Equals(RootGraph->GetName(), ESearchCase::IgnoreCase))
 			{
+				SubGraphStartIndex = 1;
 				Graph = RootGraph;
 				break;
 			}
@@ -2846,13 +2899,12 @@ UOptimusNodeGraph* UOptimusDeformer::ResolveGraphPath(
 	}
 
 	// See if we need to traverse any sub-graphs
-	int32 GraphIndex = 1;
-	for (; GraphIndex < Path.Num(); GraphIndex++)
+	for (; SubGraphStartIndex < Path.Num(); SubGraphStartIndex++)
 	{
 		bool bFoundSubGraph = false;
 		for (UOptimusNodeGraph* SubGraph: Graph->GetGraphs())
 		{
-			if (Path[GraphIndex].Equals(SubGraph->GetName(), ESearchCase::IgnoreCase))
+			if (Path[SubGraphStartIndex].Equals(SubGraph->GetName(), ESearchCase::IgnoreCase))
 			{
 				Graph = SubGraph;
 				bFoundSubGraph = true;
@@ -2865,11 +2917,11 @@ UOptimusNodeGraph* UOptimusDeformer::ResolveGraphPath(
 		}
 	}
 
-	if (GraphIndex < Path.Num())
+	if (SubGraphStartIndex < Path.Num())
 	{
 		OutRemainingPath = FStringView(
-			Path[GraphIndex].GetData(),
-			static_cast<int32>(Path.Last().GetData() - Path[GraphIndex].GetData()) + Path.Last().Len());
+			Path[SubGraphStartIndex].GetData(),
+			static_cast<int32>(Path.Last().GetData() - Path[SubGraphStartIndex].GetData()) + Path.Last().Len());
 	}
 	else
 	{
@@ -3400,6 +3452,45 @@ void UOptimusDeformer::GetAssetRegistryTags(FAssetRegistryTagsContext Context) c
 		FSoftClassPath ClassPath(BindingClass);
 		Context.AddTag(FAssetRegistryTag(TEXT("PrimaryBindingClass"), *ClassPath.ToString(), FAssetRegistryTag::TT_Hidden));
 	}
+
+	// Expose all the public functions
+	{
+		FOptimusFunctionNodeGraphHeaderArray PublicFunctionHeadersArray;
+		for (UOptimusNodeGraph* Graph : Graphs)
+		{
+			if (UOptimusFunctionNodeGraph* FunctionNodeGraph = Cast<UOptimusFunctionNodeGraph>(Graph))
+			{
+				if (FunctionNodeGraph->AccessSpecifier == UOptimusFunctionNodeGraph::AccessSpecifierPublicName)
+				{
+					PublicFunctionHeadersArray.Headers.Add(FunctionNodeGraph->GetHeader());
+				}
+			}
+		}
+
+	
+		FString PublicFunctionString;
+	
+		FOptimusFunctionNodeGraphHeaderArray::StaticStruct()->ExportText(PublicFunctionString, &PublicFunctionHeadersArray, nullptr, nullptr, PPF_None, nullptr);
+
+		Context.AddTag(FAssetRegistryTag(PublicFunctionsAssetTagName, *PublicFunctionString, FAssetRegistryTag::TT_Hidden));	
+	}
+
+	// Record all nodes referencing a function, such that function nodes can react when the referencing function changes
+	{
+		FOptimusFunctionReferenceData FunctionReferenceData;
+		TArray<UOptimusNode*> FunctionNodes = GetAllNodesOfClass(UOptimusNode_FunctionReference::StaticClass());
+
+		for (UOptimusNode* Node : FunctionNodes)
+		{
+			UOptimusNode_FunctionReference* FunctionNode = CastChecked<UOptimusNode_FunctionReference>(Node);
+			FunctionReferenceData.FunctionReferences.FindOrAdd(FunctionNode->GetSerializedGraphPath()).Nodes.Add(FunctionNode);
+		}
+
+		FString FunctionReferenceString;
+		FOptimusFunctionReferenceData::StaticStruct()->ExportText(FunctionReferenceString, &FunctionReferenceData, nullptr, nullptr, PPF_None, nullptr);
+
+		Context.AddTag(FAssetRegistryTag(FunctionReferencesAssetTagName, *FunctionReferenceString, FAssetRegistryTag::TT_Hidden));	
+	}
 }
 
 #if WITH_EDITOR
@@ -3525,7 +3616,7 @@ UOptimusNodePin* UOptimusDeformer::ResolvePinPath(const FString& InPinPath)
 
 
 
-UOptimusNodeGraph* UOptimusDeformer::CreateGraph(
+UOptimusNodeGraph* UOptimusDeformer::CreateGraphDirect(
 	EOptimusNodeGraphType InType, 
 	FName InName, 
 	TOptional<int32> InInsertBefore
@@ -3564,9 +3655,15 @@ UOptimusNodeGraph* UOptimusDeformer::CreateGraph(
 	}
 	else if (InType == EOptimusNodeGraphType::Function)
 	{
+		if (!UOptimusNodeGraph::IsValidUserGraphName(InName.ToString()))
+		{
+			return nullptr;
+		}
 		// Not fully implemented yet.
-		checkNoEntry();
 		GraphClass = UOptimusFunctionNodeGraph::StaticClass();
+
+		// If there's already an object with this name, then attempt to make the name unique.
+		InName = Optimus::GetUniqueNameForScope(this, InName);
 	}
 
 	UOptimusNodeGraph* Graph = NewObject<UOptimusNodeGraph>(this, GraphClass, InName, RF_Transactional);
@@ -3575,7 +3672,7 @@ UOptimusNodeGraph* UOptimusDeformer::CreateGraph(
 
 	if (InInsertBefore.IsSet())
 	{
-		if (!AddGraph(Graph, InInsertBefore.GetValue()))
+		if (!AddGraphDirect(Graph, InInsertBefore.GetValue()))
 		{
 			Optimus::RemoveObject(Graph);
 			return nullptr;
@@ -3586,7 +3683,7 @@ UOptimusNodeGraph* UOptimusDeformer::CreateGraph(
 }
 
 
-bool UOptimusDeformer::AddGraph(
+bool UOptimusDeformer::AddGraphDirect(
 	UOptimusNodeGraph* InGraph,
 	int32 InInsertBefore
 	)
@@ -3647,7 +3744,7 @@ bool UOptimusDeformer::AddGraph(
 }
 
 
-bool UOptimusDeformer::RemoveGraph(
+bool UOptimusDeformer::RemoveGraphDirect(
 	UOptimusNodeGraph* InGraph,
 	bool bInDeleteGraph
 	)
@@ -3679,7 +3776,7 @@ bool UOptimusDeformer::RemoveGraph(
 
 
 
-bool UOptimusDeformer::MoveGraph(
+bool UOptimusDeformer::MoveGraphDirect(
 	UOptimusNodeGraph* InGraph, 
 	int32 InInsertBefore
 	)
@@ -3720,6 +3817,29 @@ bool UOptimusDeformer::MoveGraph(
 	return true;
 }
 
+bool UOptimusDeformer::RenameGraphDirect(UOptimusNodeGraph* InGraph, const FString& InNewName)
+{
+	FSoftObjectPath PreviousFunctionGraphPath;
+	if (InGraph->GetGraphType() == EOptimusNodeGraphType::Function)
+	{
+		PreviousFunctionGraphPath = FSoftObjectPath(InGraph);
+	}	
+	
+	if (Optimus::RenameObject(InGraph, *InNewName, nullptr))
+	{
+		if (InGraph->GetGraphType() == EOptimusNodeGraphType::Function)
+		{
+			IOptimusCoreModule& OptimusCoreModule = IOptimusCoreModule::Get();
+			OptimusCoreModule.UpdateFunctionReferences(PreviousFunctionGraphPath, FSoftObjectPath(InGraph));
+		}
+
+		Notify(EOptimusGlobalNotifyType::GraphRenamed, InGraph);
+		return true;
+	}
+
+	return false;	
+}
+
 
 bool UOptimusDeformer::RenameGraph(UOptimusNodeGraph* InGraph, const FString& InNewName)
 {
@@ -3742,12 +3862,7 @@ bool UOptimusDeformer::RenameGraph(UOptimusNodeGraph* InGraph, const FString& In
 		return false;
 	}
 
-	const bool bSuccess = GetActionStack()->RunAction<FOptimusNodeGraphAction_RenameGraph>(InGraph, FName(*InNewName));
-	if (bSuccess)
-	{
-		Notify(EOptimusGlobalNotifyType::GraphRenamed, InGraph);
-	}
-	return bSuccess;
+	return GetActionStack()->RunAction<FOptimusNodeGraphAction_RenameGraph>(InGraph, FName(*InNewName));
 }
 
 
diff --git a/Engine/Plugins/Animation/DeformerGraph/Source/OptimusCore/Private/OptimusFunctionNodeGraph.cpp b/Engine/Plugins/Animation/DeformerGraph/Source/OptimusCore/Private/OptimusFunctionNodeGraph.cpp
--- a/Engine/Plugins/Animation/DeformerGraph/Source/OptimusCore/Private/OptimusFunctionNodeGraph.cpp
+++ b/Engine/Plugins/Animation/DeformerGraph/Source/OptimusCore/Private/OptimusFunctionNodeGraph.cpp
@@ -1,11 +1,55 @@
 ﻿// Copyright Epic Games, Inc. All Rights Reserved.
 
 #include "OptimusFunctionNodeGraph.h"
+#include "OptimusFunctionNodeGraphHeader.h"
 
 #include UE_INLINE_GENERATED_CPP_BY_NAME(OptimusFunctionNodeGraph)
 
+FName UOptimusFunctionNodeGraph::AccessSpecifierPublicName = TEXT("Public");
+FName UOptimusFunctionNodeGraph::AccessSpecifierPrivateName = TEXT("Private");
+
+
+UOptimusFunctionNodeGraph::UOptimusFunctionNodeGraph()
+{
+}
 
 FString UOptimusFunctionNodeGraph::GetNodeName() const
 {
 	return GetName();
 }
+
+TArray<FName> UOptimusFunctionNodeGraph::GetAccessSpecifierOptions() const
+{
+	return {AccessSpecifierPublicName, AccessSpecifierPrivateName};
+}
+
+FOptimusFunctionNodeGraphHeader UOptimusFunctionNodeGraph::GetHeader() const
+{
+	FOptimusFunctionNodeGraphHeader Header;
+
+	Header.GraphPath = this; 
+	Header.FunctionName = GetFName();
+	Header.Category = Category;
+
+	return Header;
+}
+
+
+#if WITH_EDITOR
+
+
+bool UOptimusFunctionNodeGraph::CanEditChange(const FProperty* InProperty) const
+{
+	if (InProperty == StaticClass()->FindPropertyByName(GET_MEMBER_NAME_CHECKED(UOptimusNodeSubGraph, InputBindings)))
+	{
+		return false;
+	}
+	else if (InProperty == StaticClass()->FindPropertyByName(GET_MEMBER_NAME_CHECKED(UOptimusNodeSubGraph, OutputBindings)))
+	{
+		return false;
+	}
+
+	return true;
+}
+
+#endif
diff --git a/Engine/Plugins/Animation/DeformerGraph/Source/OptimusCore/Private/OptimusFunctionNodeGraphHeader.cpp b/Engine/Plugins/Animation/DeformerGraph/Source/OptimusCore/Private/OptimusFunctionNodeGraphHeader.cpp
new file mode 100644
--- /dev/null
+++ b/Engine/Plugins/Animation/DeformerGraph/Source/OptimusCore/Private/OptimusFunctionNodeGraphHeader.cpp
@@ -0,0 +1,7 @@
+﻿// Copyright Epic Games, Inc. All Rights Reserved.
+
+#include "OptimusFunctionNodeGraphHeader.h"
+#include "OptimusFunctionNodeGraph.h"
+
+#include UE_INLINE_GENERATED_CPP_BY_NAME(OptimusFunctionNodeGraphHeader)
+
diff --git a/Engine/Plugins/Animation/DeformerGraph/Source/OptimusCore/Private/OptimusNodeGraph.cpp b/Engine/Plugins/Animation/DeformerGraph/Source/OptimusCore/Private/OptimusNodeGraph.cpp
--- a/Engine/Plugins/Animation/DeformerGraph/Source/OptimusCore/Private/OptimusNodeGraph.cpp
+++ b/Engine/Plugins/Animation/DeformerGraph/Source/OptimusCore/Private/OptimusNodeGraph.cpp
@@ -31,6 +31,8 @@
 #include "IOptimusComponentBindingReceiver.h"
 #include "IOptimusNodePairProvider.h"
 #include "IOptimusPinMutabilityDefiner.h"
+#include "OptimusFunctionNodeGraph.h"
+#include "OptimusFunctionNodeGraphHeader.h"
 #include "OptimusNodeSubGraph.h"
 #include "Nodes/OptimusNode_CustomComputeKernel.h"
 #include "Nodes/OptimusNode_FunctionReference.h"
@@ -55,6 +57,10 @@ const FName UOptimusNodeGraph::SetupGraphName("SetupGraph");
 const FName UOptimusNodeGraph::UpdateGraphName("UpdateGraph");
 const TCHAR* UOptimusNodeGraph::LibraryRoot = TEXT("@Library");
 
+FString UOptimusNodeGraph::GetFunctionGraphPath(const FString& InFunctionName)
+{
+	return FString::Printf(TEXT("%s/%s"), LibraryRoot, *InFunctionName);
+}
 
 void UOptimusNodeGraph::PostDuplicate(EDuplicateMode::Type DuplicateMode)
 {
@@ -196,8 +202,7 @@ FString UOptimusNodeGraph::GetGraphPath() const
 		return GetName();
 
 	case EOptimusNodeGraphType::Function:
-		// FIXME: Check if we're internal or external function graph
-		return FString::Printf(TEXT("%s/%s"), LibraryRoot, *GetName());
+		return GetFunctionGraphPath(GetName());
 		
 	case EOptimusNodeGraphType::SubGraph:
 		{
@@ -366,7 +371,7 @@ void UOptimusNodeGraph::AddPairNodesToArray(
 	const TArray<UOptimusNode*>& InNodes,
 	TArray<UOptimusNode*>& OutNodes,
 	TArray<UOptimusNodePair*>& OutNodePairs
-	) const
+	)
 {
 	OutNodes = InNodes;
 	for (UOptimusNode* Node : InNodes)
@@ -467,6 +472,15 @@ TArray<UOptimusNode*> UOptimusNodeGraph::AddLoopTerminalNodes(const FVector2D& I
 	return AddedNodes;
 }
 
+UOptimusNode* UOptimusNodeGraph::AddFunctionReferenceNode(TSoftObjectPtr<UOptimusFunctionNodeGraph> InFunctionGraph, const FVector2D& InPosition)
+{
+	return AddNodeInternal(UOptimusNode_FunctionReference::StaticClass(), InPosition,
+	[InFunctionGraph](UOptimusNode *InNode)
+	{
+		UOptimusNode_FunctionReference* FunctionNode = CastChecked<UOptimusNode_FunctionReference>(InNode);
+		FunctionNode->FunctionGraph = InFunctionGraph;
+	});
+}
 
 UOptimusNode* UOptimusNodeGraph::AddResourceNode(
 	UOptimusResourceDescription* InResourceDesc,
@@ -576,7 +590,7 @@ bool UOptimusNodeGraph::RemoveNodes(
 
 
 bool UOptimusNodeGraph::RemoveNodesToAction(
-	FOptimusCompoundAction* InAction,
+	FOptimusCompoundAction* InOutAction,
 	const TArray<UOptimusNode*>& InNodes
 	) const
 {
@@ -597,6 +611,18 @@ bool UOptimusNodeGraph::RemoveNodesToAction(
 			return false;
 		}
 	}
+
+	TArray<UOptimusNodeSubGraph*> SubGraphsToRemove;
+	for (UOptimusNode* Node : NodesToRemove)
+	{
+		if (UOptimusNode_SubGraphReference* SubGraphReference = Cast<UOptimusNode_SubGraphReference>(Node))
+		{
+			if (ensure(SubGraphReference->SubGraph.IsValid()))
+			{
+				SubGraphsToRemove.Add(SubGraphReference->SubGraph.Get());
+			}
+		}
+	}
 	
 	TSet<int32> AllLinkIndexes;
 
@@ -608,24 +634,27 @@ bool UOptimusNodeGraph::RemoveNodesToAction(
 
 	for (const int32 LinkIndex : AllLinkIndexes)
 	{
-		InAction->AddSubAction<FOptimusNodeGraphAction_RemoveLink>(Links[LinkIndex]);
+		InOutAction->AddSubAction<FOptimusNodeGraphAction_RemoveLink>(Links[LinkIndex]);
 	}
 
 	for (const UOptimusNodePair* NodePair: NodePairsToRemove)
 	{
-		InAction->AddSubAction<FOptimusNodeGraphAction_RemoveNodePair>(NodePair);
+		InOutAction->AddSubAction<FOptimusNodeGraphAction_RemoveNodePair>(NodePair);
 	}	
 	
-	
 	for (UOptimusNode* Node : NodesToRemove)
 	{
-		InAction->AddSubAction<FOptimusNodeGraphAction_RemoveNode>(Node);
+		InOutAction->AddSubAction<FOptimusNodeGraphAction_RemoveNode>(Node);
+	}
+
+	for (UOptimusNodeSubGraph* SubGraph : SubGraphsToRemove)
+	{
+		InOutAction->AddSubAction<FOptimusNodeGraphAction_RemoveGraph>(SubGraph);
 	}
 
 	return true;
 }
 
-
 UOptimusNode* UOptimusNodeGraph::DuplicateNode(
 	UOptimusNode* InNode,
 	const FVector2D& InPosition
@@ -1174,7 +1203,7 @@ UOptimusNode* UOptimusNodeGraph::CollapseNodesToSubGraph(
 	}
 	
 	Action->AddSubAction<FOptimusNodeGraphAction_AddNode>(
-		SubGraphPath, UOptimusNode_GraphTerminal::StaticClass(), "Entry",
+		SubGraphPath, UOptimusNode_GraphTerminal::StaticClass(), UOptimusNode_GraphTerminal::EntryNodeName,
 		[NodeBox, SubGraphPath, PathResolver](UOptimusNode* InNode)
 		{
 			UOptimusNode_GraphTerminal* EntryNode = Cast<UOptimusNode_GraphTerminal>(InNode); 
@@ -1188,7 +1217,7 @@ UOptimusNode* UOptimusNodeGraph::CollapseNodesToSubGraph(
 		});
 
 	Action->AddSubAction<FOptimusNodeGraphAction_AddNode>(
-		SubGraphPath, UOptimusNode_GraphTerminal::StaticClass(), "Return",
+		SubGraphPath, UOptimusNode_GraphTerminal::StaticClass(), UOptimusNode_GraphTerminal::ReturnNodeName,
 		[NodeBox, SubGraphPath, PathResolver](UOptimusNode* InNode)
 		{
 			UOptimusNode_GraphTerminal* ReturnNode = Cast<UOptimusNode_GraphTerminal>(InNode); 
@@ -1272,8 +1301,7 @@ UOptimusNode* UOptimusNodeGraph::CollapseNodesToSubGraph(
 	{
 		FString EntryNodePinName = EntryNodePinNames[Link].ToString(); 
 		// Once for the Entry -> Sub-graph nodes, and another for Outer graph nodes -> Sub-graph ref node inputs.
-		FString NodeOutputPinPath = FString::Printf(TEXT("%s/Entry.%s"),
-			*SubGraphPath, *EntryNodePinName);
+		FString NodeOutputPinPath = ConstructPath(SubGraphPath, UOptimusNode_GraphTerminal::EntryNodeName.ToString(), EntryNodePinName);
 		FString NodeInputPinPath = FString::Printf(TEXT("%s/%s.%s"),
 			*SubGraphPath, *Link->NodeInputPin->GetOwningNode()->GetName(), *Link->NodeInputPin->GetUniqueName().ToString());
 		Action->AddSubAction<FOptimusNodeGraphAction_AddLink>(NodeOutputPinPath, NodeInputPinPath);
@@ -1290,8 +1318,8 @@ UOptimusNode* UOptimusNodeGraph::CollapseNodesToSubGraph(
 		// Once for the Sub-graph nodes -> Return, and another for Sub-graph ref node outputs -> Outer graph nodes.
 		FString NodeOutputPinPath = FString::Printf(TEXT("%s/%s.%s"),
 			*SubGraphPath, *Link->NodeOutputPin->GetOwningNode()->GetName(), *Link->NodeOutputPin->GetUniqueName().ToString());
-		FString NodeInputPinPath = FString::Printf(TEXT("%s/Return.%s"),
-			*SubGraphPath, *ReturnNodePinName);
+		FString NodeInputPinPath =
+		ConstructPath(SubGraphPath, UOptimusNode_GraphTerminal::ReturnNodeName.ToString(), ReturnNodePinName);
 		Action->AddSubAction<FOptimusNodeGraphAction_AddLink>(NodeOutputPinPath, NodeInputPinPath);
 		
 		NodeOutputPinPath = FString::Printf(TEXT("%s/%s.%s"),
@@ -1504,8 +1532,6 @@ TArray<UOptimusNode*> UOptimusNodeGraph::ExpandCollapsedNodes(
 
 		FOptimusCompoundAction *Action = new FOptimusCompoundAction(TEXT("Expand Collapsed Nodes"));
 
-		RemoveNodesToAction(Action, {SubGraphReferenceNode});
-
 		for (TPair<UOptimusNode*, FVector2D> NodePositionInfo: NodePositionInfos)
 		{
 			FVector2D Offset = NodePositionInfo.Value.GetSignVector() * HalfBoundSize;
@@ -1548,10 +1574,8 @@ TArray<UOptimusNode*> UOptimusNodeGraph::ExpandCollapsedNodes(
 
 		RemoveNodesToAction(Action, SubGraph->GetAllNodes());
 
-		Action->AddSubAction<FOptimusNodeGraphAction_RemoveGraph>(SubGraph);
-		
-		
-		
+		RemoveNodesToAction(Action, {SubGraphReferenceNode});
+
 		if (!GetActionStack()->RunAction(Action))
 		{
 			return {};
@@ -1569,6 +1593,238 @@ TArray<UOptimusNode*> UOptimusNodeGraph::ExpandCollapsedNodes(
 	return {};
 }
 
+bool UOptimusNodeGraph::ConvertToFunction(UOptimusNode* InSubGraphNode)
+{
+	const bool bIsSubGraph = IsSubGraphReference(InSubGraphNode);
+	if (!bIsSubGraph)
+	{
+		return false;
+	}
+	
+	UOptimusNode_SubGraphReference* SubGraphNode = CastChecked<UOptimusNode_SubGraphReference>(InSubGraphNode);
+	UOptimusNodeSubGraph* SubGraph = SubGraphNode->SubGraph.Get();
+
+	UOptimusDeformer* Deformer = Cast<UOptimusDeformer>(GetCollectionRoot());
+	if (!ensure(Deformer))
+	{
+		return false;
+	}
+
+	IOptimusPathResolver* PathResolver = GetPathResolver();
+
+	// Ensure the name is unique here instead of relying on AddGraph action
+	// because subsequent actions needs to use the final graph path
+	// Note: we could change the actions to accept a lambda instead of specific value in the future
+	FName FunctionName = TEXT("NewFunction");
+	FunctionName = Optimus::GetUniqueNameForScope(Deformer, FunctionName);
+	
+	const FString FunctionGraphPath = GetFunctionGraphPath(FunctionName.ToString());
+	const FVector2D EntryNodePosition = SubGraph->EntryNode->GetGraphPosition();
+	const FVector2D ReturnNodePosition = SubGraph->ReturnNode->GetGraphPosition();
+	const FVector2D ReferenceNodePosition = SubGraphNode->GetGraphPosition();
+	const FName FunctionNodeName = Optimus::GetUniqueNameForScope(this, TEXT("FunctionReferenceNode"));
+	const FString FunctionNodeNameString = FunctionNodeName.ToString();
+	const FString CurrentGraphPath = GetGraphPath();
+
+	TMap<FString, TArray<FString>> ExternalLinksToAdd;
+
+	for (UOptimusNodeLink* Link : GetAllLinks())
+	{
+		if (Link->GetNodeInputPin()->GetOwningNode() == SubGraphNode)
+		{
+			FString FunctionNodeInputPinPath = ConstructPath(
+				CurrentGraphPath,
+				FunctionNodeNameString,
+				Link->GetNodeInputPin()->GetUniqueName().ToString()
+				);
+
+			FString ExistingOutputPinPath = Link->GetNodeOutputPin()->GetPinPath();
+			ExternalLinksToAdd.FindOrAdd(ExistingOutputPinPath).Add(FunctionNodeInputPinPath);
+		}
+		else if (Link->GetNodeOutputPin()->GetOwningNode() == SubGraphNode)
+		{
+			FString FunctionNodeOutputPinPath = ConstructPath(
+				CurrentGraphPath,
+				FunctionNodeNameString,
+				Link->GetNodeOutputPin()->GetUniqueName().ToString()
+				);
+
+			FString ExistingInputPinPath = Link->GetNodeInputPin()->GetPinPath();
+			ExternalLinksToAdd.FindOrAdd(FunctionNodeOutputPinPath).Add(ExistingInputPinPath);
+		}
+	}
+
+	TMap<FString, TArray<FString>> InternalLinksToAdd;
+	
+	for (UOptimusNodeLink* Link : SubGraph->GetAllLinks())
+	{
+		if (Link->GetNodeInputPin()->GetOwningNode() == SubGraph->ReturnNode)
+		{
+			FString ReturnNodeInputPinPath = ConstructPath(
+				FunctionGraphPath,
+				UOptimusNode_GraphTerminal::ReturnNodeName.ToString(),
+				Link->GetNodeInputPin()->GetUniqueName().ToString()
+				);
+
+			FString InternalNodeOutputPinPath = ConstructPath(
+				FunctionGraphPath,
+				Link->GetNodeOutputPin()->GetOwningNode()->GetName(),
+				Link->GetNodeOutputPin()->GetUniqueName().ToString()
+				);
+
+			InternalLinksToAdd.FindOrAdd(InternalNodeOutputPinPath).Add(ReturnNodeInputPinPath);
+		}
+		else if (Link->GetNodeOutputPin()->GetOwningNode() == SubGraph->EntryNode)
+		{
+			FString EntryNodeOutputPinPath = ConstructPath(
+				FunctionGraphPath,
+				UOptimusNode_GraphTerminal::EntryNodeName.ToString(),
+				Link->GetNodeOutputPin()->GetUniqueName().ToString()
+				);
+
+			FString InternalNodeInputPinPath = ConstructPath(
+				FunctionGraphPath,
+				Link->GetNodeInputPin()->GetOwningNode()->GetName(),
+				Link->GetNodeInputPin()->GetUniqueName().ToString()
+				);
+			InternalLinksToAdd.FindOrAdd(EntryNodeOutputPinPath).Add(InternalNodeInputPinPath);
+		}
+		else
+		{
+			FString OutputPinPath = ConstructPath(
+					FunctionGraphPath,
+					Link->GetNodeOutputPin()->GetOwningNode()->GetName(),
+					Link->GetNodeOutputPin()->GetUniqueName().ToString()
+					);
+
+			FString InputPinPath = ConstructPath(
+				FunctionGraphPath,
+				Link->GetNodeInputPin()->GetOwningNode()->GetName(),
+				Link->GetNodeInputPin()->GetUniqueName().ToString()
+				);
+			
+			InternalLinksToAdd.FindOrAdd(OutputPinPath).Add(InputPinPath);
+		}
+	}	
+	
+	
+	TArray<UOptimusNode*> SubGraphNodes = SubGraph->GetAllNodes();
+	
+	SubGraphNodes.RemoveAll([](UOptimusNode* InNode)
+	{
+		return Cast<UOptimusNode_GraphTerminal>(InNode) != nullptr;
+	});
+
+	TArray<UOptimusNode*> NodesToDuplicate;
+	TArray<UOptimusNodePair*> NodePairsToDuplicate;
+	AddPairNodesToArray(SubGraphNodes, NodesToDuplicate, NodePairsToDuplicate);
+
+	TArray<TPair<FString, FString>> NewNodePairs;
+	for (const UOptimusNodePair* NodePair : NodePairsToDuplicate)
+	{
+		// Generate new pair
+		const UOptimusNode *FirstNode = NodePair->GetFirst();
+		const UOptimusNode *SecondNode = NodePair->GetSecond();
+
+		FString NewFirstNodePath = ConstructPath(FunctionGraphPath, FirstNode->GetName(), {});
+		FString NewSecondNodePath = ConstructPath(FunctionGraphPath, SecondNode->GetName(), {});
+
+		NewNodePairs.Add({NewFirstNodePath, NewSecondNodePath});
+	}	
+
+	{
+		FOptimusActionScope(*GetActionStack(), TEXT("Convert To Function"));
+		GetActionStack()->RunAction<FOptimusNodeGraphAction_AddGraph>(Deformer, EOptimusNodeGraphType::Function, FunctionName, INDEX_NONE,
+			[SubGraph](UOptimusNodeGraph* InGraph)
+			{
+				UOptimusFunctionNodeGraph* FunctionGraph = CastChecked<UOptimusFunctionNodeGraph>(InGraph);
+				FunctionGraph->InputBindings = SubGraph->InputBindings;
+				FunctionGraph->OutputBindings = SubGraph->OutputBindings;
+
+				return true;
+			});
+	
+		GetActionStack()->RunAction<FOptimusNodeGraphAction_AddNode>(
+			FunctionGraphPath, UOptimusNode_GraphTerminal::StaticClass(), UOptimusNode_GraphTerminal::EntryNodeName,
+			[FunctionGraphPath, EntryNodePosition, PathResolver](UOptimusNode* InNode)
+			{
+				UOptimusNode_GraphTerminal* EntryNode = Cast<UOptimusNode_GraphTerminal>(InNode); 
+				UOptimusNodeSubGraph* SubGraph = Cast<UOptimusNodeSubGraph>(PathResolver->ResolveGraphPath(FunctionGraphPath));
+
+				SubGraph->EntryNode = EntryNode;
+		
+				EntryNode->TerminalType = EOptimusTerminalType::Entry;
+				EntryNode->OwningGraph = SubGraph;
+				return EntryNode->SetGraphPositionDirect(EntryNodePosition);
+			});
+
+		GetActionStack()->RunAction<FOptimusNodeGraphAction_AddNode>(
+			FunctionGraphPath, UOptimusNode_GraphTerminal::StaticClass(), UOptimusNode_GraphTerminal::ReturnNodeName,
+			[FunctionGraphPath, ReturnNodePosition, PathResolver](UOptimusNode* InNode)
+			{
+				UOptimusNode_GraphTerminal* ReturnNode = Cast<UOptimusNode_GraphTerminal>(InNode); 
+				UOptimusNodeSubGraph* SubGraph = Cast<UOptimusNodeSubGraph>(PathResolver->ResolveGraphPath(FunctionGraphPath));
+
+				SubGraph->ReturnNode = ReturnNode;
+			
+				ReturnNode->TerminalType = EOptimusTerminalType::Return;
+				ReturnNode->OwningGraph = SubGraph;
+				return ReturnNode->SetGraphPositionDirect(ReturnNodePosition);
+			});
+
+		
+		// Copy the nodes to the function graph
+		for (UOptimusNode* Node: NodesToDuplicate)
+		{
+			GetActionStack()->RunAction<FOptimusNodeGraphAction_DuplicateNode>(
+				FunctionGraphPath, Node, Node->GetFName(), [](UOptimusNode*) { return true; });
+		}
+
+		for (const TPair<FString, FString>& NewPair : NewNodePairs)
+		{
+			GetActionStack()->RunAction<FOptimusNodeGraphAction_AddNodePair>(NewPair.Key, NewPair.Value);
+		}	
+		
+		// Restore links
+		for (const TPair<FString, TArray<FString>>& InternalLinks : InternalLinksToAdd )
+		{
+			const FString& OutputPinPath = InternalLinks.Key;
+			for (const FString& InputPinPath : InternalLinks.Value)
+			{
+				GetActionStack()->RunAction<FOptimusNodeGraphAction_AddLink>(OutputPinPath,	InputPinPath);
+			}
+		}
+		
+		// Create the function ref node
+		GetActionStack()->RunAction<FOptimusNodeGraphAction_AddNode>(
+			GetGraphPath(), UOptimusNode_FunctionReference::StaticClass(), FunctionNodeName,
+			[ReferenceNodePosition, FunctionGraphPath, PathResolver](UOptimusNode* InNode)
+			{
+				UOptimusNode_FunctionReference* FunctionNode = Cast<UOptimusNode_FunctionReference>(InNode); 
+				UOptimusFunctionNodeGraph* FunctionGraph = Cast<UOptimusFunctionNodeGraph>(PathResolver->ResolveGraphPath(FunctionGraphPath));
+				FunctionNode->FunctionGraph = FunctionGraph;
+				
+				return FunctionNode->SetGraphPositionDirect(ReferenceNodePosition);
+			});
+
+		// Remove the subgraph node
+		RemoveNodes({SubGraphNode});
+		
+		// Restore links
+		for (const TPair<FString, TArray<FString>>& ExternalLinks : ExternalLinksToAdd )
+		{
+			const FString& OutputPinPath = ExternalLinks.Key;
+			for (const FString& InputPinPath : ExternalLinks.Value)
+			{
+				GetActionStack()->RunAction<FOptimusNodeGraphAction_AddLink>(OutputPinPath,	InputPinPath);
+			}
+		}
+	}
+	
+	
+	return true;
+}
+
 
 bool UOptimusNodeGraph::IsCustomKernel(UOptimusNode* InNode) const
 {
@@ -2512,7 +2768,7 @@ FString UOptimusNodeGraph::GetCollectionPath() const
 }
 
 
-UOptimusNodeGraph* UOptimusNodeGraph::CreateGraph(
+UOptimusNodeGraph* UOptimusNodeGraph::CreateGraphDirect(
 	EOptimusNodeGraphType InType,
 	FName InName,
 	TOptional<int32> InInsertBefore
@@ -2532,7 +2788,7 @@ UOptimusNodeGraph* UOptimusNodeGraph::CreateGraph(
 
 	if (InInsertBefore.IsSet())
 	{
-		if (!AddGraph(Graph, InInsertBefore.GetValue()))
+		if (!AddGraphDirect(Graph, InInsertBefore.GetValue()))
 		{
 			Optimus::RemoveObject(Graph);
 			return nullptr;
@@ -2543,7 +2799,7 @@ UOptimusNodeGraph* UOptimusNodeGraph::CreateGraph(
 }
 
 
-bool UOptimusNodeGraph::AddGraph(
+bool UOptimusNodeGraph::AddGraphDirect(
 	UOptimusNodeGraph* InGraph, 
 	int32 InInsertBefore
 	)
@@ -2571,7 +2827,7 @@ bool UOptimusNodeGraph::AddGraph(
 }
 
 
-bool UOptimusNodeGraph::RemoveGraph(
+bool UOptimusNodeGraph::RemoveGraphDirect(
 	UOptimusNodeGraph* InGraph,
 	bool bInDeleteGraph
 	)
@@ -2602,7 +2858,7 @@ bool UOptimusNodeGraph::RemoveGraph(
 }
 
 
-bool UOptimusNodeGraph::MoveGraph(
+bool UOptimusNodeGraph::MoveGraphDirect(
 	UOptimusNodeGraph* InGraph,
 	int32 InInsertBefore
 	)
@@ -2610,6 +2866,17 @@ bool UOptimusNodeGraph::MoveGraph(
 	return false;
 }
 
+bool UOptimusNodeGraph::RenameGraphDirect(UOptimusNodeGraph* InGraph, const FString& InNewName)
+{
+	if (Optimus::RenameObject(InGraph, *InNewName, nullptr))
+	{
+		GlobalNotify(EOptimusGlobalNotifyType::GraphRenamed, InGraph);
+		return true;
+	}
+
+	return false;
+}
+
 
 bool UOptimusNodeGraph::RenameGraph(
 	UOptimusNodeGraph* InGraph,
@@ -2627,13 +2894,8 @@ bool UOptimusNodeGraph::RenameGraph(
 	{
 		return false;
 	}
-
-	const bool bSuccess = GetActionStack()->RunAction<FOptimusNodeGraphAction_RenameGraph>(InGraph, FName(*InNewName));
-	if (bSuccess)
-	{
-		GlobalNotify(EOptimusGlobalNotifyType::GraphRenamed, InGraph);
-	}
-	return bSuccess;
+	
+	return GetActionStack()->RunAction<FOptimusNodeGraphAction_RenameGraph>(InGraph, FName(*InNewName));;
 }
 
 #if WITH_EDITOR
diff --git a/Engine/Plugins/Animation/DeformerGraph/Source/OptimusCore/Private/OptimusSource.cpp b/Engine/Plugins/Animation/DeformerGraph/Source/OptimusCore/Private/OptimusSource.cpp
--- a/Engine/Plugins/Animation/DeformerGraph/Source/OptimusCore/Private/OptimusSource.cpp
+++ b/Engine/Plugins/Animation/DeformerGraph/Source/OptimusCore/Private/OptimusSource.cpp
@@ -26,4 +26,9 @@ FString UOptimusSource::GetNameForShaderTextEditor() const
 	return GetFName().ToString();
 }
 
+bool UOptimusSource::IsShaderTextReadOnly() const
+{
+	return false;
+}
+
 #endif // WITH_EDITOR	
diff --git a/Engine/Plugins/Animation/DeformerGraph/Source/OptimusCore/Public/IOptimusCoreModule.h b/Engine/Plugins/Animation/DeformerGraph/Source/OptimusCore/Public/IOptimusCoreModule.h
--- a/Engine/Plugins/Animation/DeformerGraph/Source/OptimusCore/Public/IOptimusCoreModule.h
+++ b/Engine/Plugins/Animation/DeformerGraph/Source/OptimusCore/Public/IOptimusCoreModule.h
@@ -4,7 +4,10 @@
 
 #include "CoreMinimal.h"
 #include "Modules/ModuleInterface.h"
+#include "Modules/ModuleManager.h"
 #include "Templates/SubclassOf.h"
+#include "OptimusFunctionNodeGraph.h"
+#include "OptimusNode.h"
 
 
 class UOptimusComputeDataInterface;
@@ -13,6 +16,11 @@ class UOptimusComputeDataInterface;
 class IOptimusCoreModule : public IModuleInterface
 {
 public:
+	static IOptimusCoreModule& Get()
+	{
+		return FModuleManager::LoadModuleChecked< IOptimusCoreModule >(TEXT("OptimusCore"));
+	}
+	
 	/** Register a data interface after the OptimusCore module has loaded */
 	template<typename T>
 	bool RegisterDataInterfaceClass()
@@ -20,6 +28,8 @@ public:
 		return RegisterDataInterfaceClass(T::StaticClass());
 	}
 	
+	virtual void UpdateFunctionReferences(const FSoftObjectPath& InOldGraphPath, const FSoftObjectPath& InNewGraphPath) = 0;
+	
 protected:
 	virtual bool RegisterDataInterfaceClass(TSubclassOf<UOptimusComputeDataInterface> InDataInterfaceClass) = 0;
 };
diff --git a/Engine/Plugins/Animation/DeformerGraph/Source/OptimusCore/Public/IOptimusNodeGraphCollectionOwner.h b/Engine/Plugins/Animation/DeformerGraph/Source/OptimusCore/Public/IOptimusNodeGraphCollectionOwner.h
--- a/Engine/Plugins/Animation/DeformerGraph/Source/OptimusCore/Public/IOptimusNodeGraphCollectionOwner.h
+++ b/Engine/Plugins/Animation/DeformerGraph/Source/OptimusCore/Public/IOptimusNodeGraphCollectionOwner.h
@@ -58,7 +58,7 @@ public:
 	/// not be inserted into the list of editable graphs, although it will be owned by this 
 	/// collection.
 	/// @return The newly created graph.
-	virtual UOptimusNodeGraph* CreateGraph(
+	virtual UOptimusNodeGraph* CreateGraphDirect(
 		EOptimusNodeGraphType InType, 
 		FName InName = NAME_None, 
 		TOptional<int32> InInsertBefore = TOptional<int32>(INDEX_NONE)
@@ -70,7 +70,7 @@ public:
 	/// @param InInsertBefore The index at which the insert the graph at. This inserts the given
 	/// graph before the graph already occupying this location.
 	/// @return true if the graph was successfully added to the graph.
-	virtual bool AddGraph(
+	virtual bool AddGraphDirect(
 		UOptimusNodeGraph* InGraph,
 		int32 InInsertBefore = INDEX_NONE
 		) = 0;
@@ -79,7 +79,7 @@ public:
 	/// @param InGraph The graph to remove, owned by this collection. The graph will be unowned 
 	/// by this graph collection and marked for GC.
 	/// @return true if the graph was successfully removed.
-	virtual bool RemoveGraph(
+	virtual bool RemoveGraphDirect(
 		UOptimusNodeGraph* InGraph,
 		bool bDeleteGraph = true
 	) = 0;
@@ -91,10 +91,19 @@ public:
 	/// @param InInsertBefore The order at which the move the graph to. This inserts the given 
 	/// graph before the graph already occupying this location.
 	/// @return true if the graph was successfully moved.
-	virtual bool MoveGraph(
+	virtual bool MoveGraphDirect(
 	    UOptimusNodeGraph* InGraph,
 		int32 InInsertBefore) = 0;
 
+	/// Rename the given graph, subject to validation of the name.
+	/// @param InGraph The graph to rename, owned by this collection.
+	/// @param InNewName The new name to give the graph.
+	/// @return true if the graph was successfully renamed.
+	virtual bool RenameGraphDirect(
+		UOptimusNodeGraph *InGraph,
+		const FString &InNewName
+		) = 0;
+	
 	/// Rename the given graph, subject to validation of the name.
 	/// @param InGraph The graph to rename, owned by this collection.
 	/// @param InNewName The new name to give the graph.
diff --git a/Engine/Plugins/Animation/DeformerGraph/Source/OptimusCore/Public/IOptimusNodeSubGraphReferencer.h b/Engine/Plugins/Animation/DeformerGraph/Source/OptimusCore/Public/IOptimusNodeSubGraphReferencer.h
new file mode 100644
--- /dev/null
+++ b/Engine/Plugins/Animation/DeformerGraph/Source/OptimusCore/Public/IOptimusNodeSubGraphReferencer.h
@@ -0,0 +1,30 @@
+// Copyright Epic Games, Inc. All Rights Reserved.
+
+#pragma once
+
+#include "UObject/Interface.h"
+
+#include "IOptimusNodeSubGraphReferencer.generated.h"
+
+class UOptimusNodeGraph;
+class UOptimusComponentSourceBinding;
+struct FOptimusPinTraversalContext;
+
+UINTERFACE()
+class OPTIMUSCORE_API UOptimusNodeSubGraphReferencer :
+	public UInterface
+{
+	GENERATED_BODY()
+};
+
+/**
+* Interface that provides a mechanism to get a node graph to display
+*/
+class OPTIMUSCORE_API IOptimusNodeSubGraphReferencer
+{
+	GENERATED_BODY()
+
+public:
+	virtual UOptimusComponentSourceBinding* GetDefaultComponentBinding(const FOptimusPinTraversalContext& InTraversalContext) const = 0;
+};
+
diff --git a/Engine/Plugins/Animation/DeformerGraph/Source/OptimusCore/Public/IOptimusShaderTextProvider.h b/Engine/Plugins/Animation/DeformerGraph/Source/OptimusCore/Public/IOptimusShaderTextProvider.h
--- a/Engine/Plugins/Animation/DeformerGraph/Source/OptimusCore/Public/IOptimusShaderTextProvider.h
+++ b/Engine/Plugins/Animation/DeformerGraph/Source/OptimusCore/Public/IOptimusShaderTextProvider.h
@@ -28,5 +28,7 @@ public:
 	virtual FString GetShaderText() const = 0;
 	/** Set shader text after edit modifications. */
 	virtual void SetShaderText(const FString& NewText) = 0;
+	/** Return true is the shader text can be edited */
+	virtual bool IsShaderTextReadOnly() const = 0;
 #endif
 };
diff --git a/Engine/Plugins/Animation/DeformerGraph/Source/OptimusCore/Public/OptimusDeformer.h b/Engine/Plugins/Animation/DeformerGraph/Source/OptimusCore/Public/OptimusDeformer.h
--- a/Engine/Plugins/Animation/DeformerGraph/Source/OptimusCore/Public/OptimusDeformer.h
+++ b/Engine/Plugins/Animation/DeformerGraph/Source/OptimusCore/Public/OptimusDeformer.h
@@ -10,6 +10,7 @@
 #include "OptimusNodeGraph.h"
 #include "OptimusComponentSource.h"
 #include "OptimusDataDomain.h"
+#include "OptimusNodeSubGraph.h"
 
 #include "Animation/MeshDeformer.h"
 #include "Interfaces/Interface_PreviewMeshProvider.h"
@@ -27,6 +28,7 @@ class UOptimusComputeGraph;
 class UOptimusDeformer;
 class UOptimusResourceDescription;
 class UOptimusVariableDescription;
+class UOptimusFunctionNodeGraph;
 enum class EOptimusDiagnosticLevel : uint8;
 struct FOptimusCompilerDiagnostic;
 struct FOptimusCompoundAction;
@@ -122,6 +124,11 @@ class OPTIMUSCORE_API UOptimusDeformer :
 	GENERATED_BODY()
 
 public:
+
+	static const FName PublicFunctionsAssetTagName;
+	static const FName FunctionReferencesAssetTagName;
+
+	
 	UOptimusDeformer();
 
 	/** Get the action stack for this deformer graph */
@@ -157,6 +164,11 @@ public:
 	/** Remove a graph and delete it. */
 	bool RemoveGraph(UOptimusNodeGraph* InGraph);
 
+	/** Returns the sub graph reference node that is uniquely associated with the given subgraph */
+	UOptimusNode* GetSubGraphReferenceNode(
+		const UOptimusNodeSubGraph* InSubGraph
+		) const;
+	
 	// Variables
 	UOptimusVariableDescription* AddVariable(
 		FOptimusDataTypeRef InDataTypeRef,
@@ -327,29 +339,34 @@ public:
 	UOptimusNodeGraph* CreateGraph(
 		EOptimusNodeGraphType InType,
 		FName InName)
-	{ return CreateGraph(InType, InName, TOptional<int32>()); }
+	{ return CreateGraphDirect(InType, InName, TOptional<int32>()); }
 	
-	UOptimusNodeGraph* CreateGraph(
+	UOptimusNodeGraph* CreateGraphDirect(
 	    EOptimusNodeGraphType InType,
 	    FName InName,
 	    TOptional<int32> InInsertBefore
 	    ) override;
 	
-	bool AddGraph(
+	bool AddGraphDirect(
 	    UOptimusNodeGraph* InGraph,
 		int32 InInsertBefore
 		) override;
 	
-	bool RemoveGraph(
+	bool RemoveGraphDirect(
 	    UOptimusNodeGraph* InGraph,
 		bool bDeleteGraph
 		) override;
 
-	bool MoveGraph(
+	bool MoveGraphDirect(
 	    UOptimusNodeGraph* InGraph,
 	    int32 InInsertBefore
 	    ) override;
 
+	bool RenameGraphDirect(
+		UOptimusNodeGraph* InGraph,
+		const FString& InNewName
+		) override;
+	
 	bool RenameGraph(
 	    UOptimusNodeGraph* InGraph,
 	    const FString& InNewName
diff --git a/Engine/Plugins/Animation/DeformerGraph/Source/OptimusCore/Public/OptimusFunctionNodeGraph.h b/Engine/Plugins/Animation/DeformerGraph/Source/OptimusCore/Public/OptimusFunctionNodeGraph.h
--- a/Engine/Plugins/Animation/DeformerGraph/Source/OptimusCore/Public/OptimusFunctionNodeGraph.h
+++ b/Engine/Plugins/Animation/DeformerGraph/Source/OptimusCore/Public/OptimusFunctionNodeGraph.h
@@ -4,11 +4,12 @@
 
 #include "OptimusNode.h"
 #include "OptimusNodeSubGraph.h"
-#include "Nodes/OptimusNode_ComputeKernelBase.h"
 
 
 #include "OptimusFunctionNodeGraph.generated.h"
 
+
+struct FOptimusFunctionNodeGraphHeader;
 /**
  * 
  */
@@ -19,10 +20,30 @@ class OPTIMUSCORE_API UOptimusFunctionNodeGraph :
 	GENERATED_BODY()
 
 public:
+	static FName AccessSpecifierPublicName;
+	static FName AccessSpecifierPrivateName;
+
+	UOptimusFunctionNodeGraph();
+	
 	/** The name to give the node based off of this graph */
 	FString GetNodeName() const; 
 
 	/** The category of the node based of of this graph for listing purposes */ 
 	UPROPERTY(EditAnywhere, Category=Settings)
 	FName Category = UOptimusNode::CategoryName::Deformers;
+
+	UPROPERTY(EditAnywhere, Category=Settings, meta=(GetOptions="GetAccessSpecifierOptions"))
+	FName AccessSpecifier = AccessSpecifierPrivateName;
+
+	UFUNCTION()
+	TArray<FName> GetAccessSpecifierOptions() const;
+
+	FOptimusFunctionNodeGraphHeader GetHeader() const;
+
+protected:
+	
+#if WITH_EDITOR
+	// UObject override
+	bool CanEditChange(const FProperty* InProperty) const override;
+#endif
 };
diff --git a/Engine/Plugins/Animation/DeformerGraph/Source/OptimusCore/Public/OptimusFunctionNodeGraphHeader.h b/Engine/Plugins/Animation/DeformerGraph/Source/OptimusCore/Public/OptimusFunctionNodeGraphHeader.h
new file mode 100644
--- /dev/null
+++ b/Engine/Plugins/Animation/DeformerGraph/Source/OptimusCore/Public/OptimusFunctionNodeGraphHeader.h
@@ -0,0 +1,54 @@
+﻿// Copyright Epic Games, Inc. All Rights Reserved.
+
+#pragma once
+
+#include "CoreMinimal.h"
+
+#include "OptimusFunctionNodeGraphHeader.generated.h"
+
+class UOptimusFunctionNodeGraph;
+class UOptimusNode_FunctionReference;
+
+USTRUCT()
+struct OPTIMUSCORE_API FOptimusFunctionNodeGraphHeader
+{
+	GENERATED_BODY()
+
+	UPROPERTY()
+	TSoftObjectPtr<UOptimusFunctionNodeGraph> GraphPath;
+
+	UPROPERTY()
+	FName FunctionName;
+
+	UPROPERTY()
+	FName Category;
+};
+
+USTRUCT()
+struct OPTIMUSCORE_API FOptimusFunctionNodeGraphHeaderArray
+{
+	GENERATED_BODY()
+
+	UPROPERTY()
+	TArray<FOptimusFunctionNodeGraphHeader> Headers;
+};
+
+USTRUCT()
+struct FOptimusFunctionReferenceNodeSet
+{
+	GENERATED_BODY()
+
+	UPROPERTY()
+	TSet<TSoftObjectPtr<UOptimusNode_FunctionReference>> Nodes;
+};
+
+USTRUCT()
+struct FOptimusFunctionReferenceData
+{
+	GENERATED_BODY()
+
+	// Group function reference nodes by function node graph path
+	// using FSoftObjectPath instead of TSoftObjectPtr such that ExportText(...) is deterministic
+	UPROPERTY()
+	TMap<FSoftObjectPath, FOptimusFunctionReferenceNodeSet> FunctionReferences;
+};
diff --git a/Engine/Plugins/Animation/DeformerGraph/Source/OptimusCore/Public/OptimusNodeGraph.h b/Engine/Plugins/Animation/DeformerGraph/Source/OptimusCore/Public/OptimusNodeGraph.h
--- a/Engine/Plugins/Animation/DeformerGraph/Source/OptimusCore/Public/OptimusNodeGraph.h
+++ b/Engine/Plugins/Animation/DeformerGraph/Source/OptimusCore/Public/OptimusNodeGraph.h
@@ -75,7 +75,11 @@ public:
 	static const FName SetupGraphName;
 	static const FName UpdateGraphName;
 	static const TCHAR* LibraryRoot;
-
+	
+	// Function Graphs are addressed in a special way
+	static FString GetFunctionGraphPath(const FString& InFunctionName);
+	
+	
 	// Check if the duplication took place at the asset level
 	// if so, we have to recreate all constant/attribute nodes such that their class pointers
 	// don't point to classes in the source asset. This can happen because generated class
@@ -146,6 +150,12 @@ public:
 		const FVector2D& InPosition
 	);
 
+	UFUNCTION(BlueprintCallable, Category = OptimusNodeGraph)
+	UOptimusNode* AddFunctionReferenceNode(
+		TSoftObjectPtr<UOptimusFunctionNodeGraph> InFunctionGraph,
+		const FVector2D& InPosition
+	);
+
 	UFUNCTION(BlueprintCallable, Category = OptimusNodeGraph)
 	UOptimusNode* AddResourceNode(
 		UOptimusResourceDescription* InResourceDesc,
@@ -258,6 +268,11 @@ public:
 	UFUNCTION(BlueprintCallable, Category = OptimusNodeGraph)
 	TArray<UOptimusNode *> ExpandCollapsedNodes(UOptimusNode* InFunctionNode);
 	
+	/** Take a subgraph node convert it to a function in-place
+	 */
+	UFUNCTION(BlueprintCallable, Category = OptimusNodeGraph)
+	bool ConvertToFunction(UOptimusNode* InSubGraphNode);
+	
 	/** Returns true if the node in question is a custom kernel node that can be converted to
 	  * a kernel function with ConvertCustomKernelToFunction.
 	  */
@@ -378,22 +393,27 @@ public:
 	UFUNCTION(BlueprintCallable, Category = OptimusNodeGraph)
 	const TArray<UOptimusNodeGraph*> &GetGraphs() const override { return SubGraphs; }
 
-	UOptimusNodeGraph* CreateGraph(
+	UOptimusNodeGraph* CreateGraphDirect(
 		EOptimusNodeGraphType InType,
 		FName InName,
 		TOptional<int32> InInsertBefore) override;
-	bool AddGraph(
+	bool AddGraphDirect(
 		UOptimusNodeGraph* InGraph,
 		int32 InInsertBefore) override;
-	bool RemoveGraph(
+	bool RemoveGraphDirect(
 		UOptimusNodeGraph* InGraph,
 		bool bInDeleteGraph) override;
 
 	UFUNCTION(BlueprintCallable, Category = OptimusNodeGraph)
-	bool MoveGraph(
+	bool MoveGraphDirect(
 		UOptimusNodeGraph* InGraph,
 		int32 InInsertBefore) override;
 
+	UFUNCTION(BlueprintCallable, Category = OptimusNodeGraph)
+	bool RenameGraphDirect(
+		UOptimusNodeGraph* InGraph,
+		const FString& InNewName) override;
+	
 	UFUNCTION(BlueprintCallable, Category = OptimusNodeGraph)
 	bool RenameGraph(
 		UOptimusNodeGraph* InGraph,
@@ -436,10 +456,10 @@ protected:
 	);
 	
 	bool RemoveNodesToAction(
-			FOptimusCompoundAction* InAction,
-			const TArray<UOptimusNode*>& InNodes
-			) const;
-
+		FOptimusCompoundAction* InOutAction,
+		const TArray<UOptimusNode*>& InNodes
+		) const;
+	
 	// Remove a node directly. If a node still has connections this call will fail. 
 	bool RemoveNodeDirect(
 		UOptimusNode* InNode,		
@@ -497,7 +517,7 @@ private:
 		TFunction<void(UOptimusNode*)> InSecondNodeConfigFunc
 	);
 
-	void AddPairNodesToArray(const TArray<UOptimusNode*>& InNodes, TArray<UOptimusNode*>& OutNodes, TArray<UOptimusNodePair*>& OutNodePairs) const;
+	static void AddPairNodesToArray(const TArray<UOptimusNode*>& InNodes, TArray<UOptimusNode*>& OutNodes, TArray<UOptimusNodePair*>& OutNodePairs);
 
 	void RemoveNodePairByIndex(int32 NodePairIndex);
 	
diff --git a/Engine/Plugins/Animation/DeformerGraph/Source/OptimusCore/Public/OptimusNodeSubGraph.h b/Engine/Plugins/Animation/DeformerGraph/Source/OptimusCore/Public/OptimusNodeSubGraph.h
--- a/Engine/Plugins/Animation/DeformerGraph/Source/OptimusCore/Public/OptimusNodeSubGraph.h
+++ b/Engine/Plugins/Animation/DeformerGraph/Source/OptimusCore/Public/OptimusNodeSubGraph.h
@@ -9,7 +9,7 @@
 class UOptimusNode_GraphTerminal;
 
 UCLASS()
-class UOptimusNodeSubGraph :
+class OPTIMUSCORE_API UOptimusNodeSubGraph :
 	public UOptimusNodeGraph,
 	public IOptimusParameterBindingProvider
 {
@@ -63,6 +63,8 @@ public:
 
 	UPROPERTY()
 	TWeakObjectPtr<UOptimusNode_GraphTerminal> ReturnNode;
+
+
 private:
 
 	void SanitizeBinding(FOptimusParameterBinding& InOutBinding, FName InOldName, bool bInAllowParameter);
diff --git a/Engine/Plugins/Animation/DeformerGraph/Source/OptimusCore/Public/OptimusSource.h b/Engine/Plugins/Animation/DeformerGraph/Source/OptimusCore/Public/OptimusSource.h
--- a/Engine/Plugins/Animation/DeformerGraph/Source/OptimusCore/Public/OptimusSource.h
+++ b/Engine/Plugins/Animation/DeformerGraph/Source/OptimusCore/Public/OptimusSource.h
@@ -27,6 +27,7 @@ public:
 	FString GetDeclarations() const override { return FString(); }
 	FString GetShaderText() const override { return SourceText; }
 	void SetShaderText(const FString& InText) override { SetSource(InText); }
+	bool IsShaderTextReadOnly() const override;
 #endif
 	// End IOptimusShaderTextProvider interface.
 
diff --git a/Engine/Plugins/Animation/DeformerGraph/Source/OptimusEditor/Private/OptimusEditor.cpp b/Engine/Plugins/Animation/DeformerGraph/Source/OptimusEditor/Private/OptimusEditor.cpp
--- a/Engine/Plugins/Animation/DeformerGraph/Source/OptimusEditor/Private/OptimusEditor.cpp
+++ b/Engine/Plugins/Animation/DeformerGraph/Source/OptimusEditor/Private/OptimusEditor.cpp
@@ -18,6 +18,7 @@
 #include "OptimusDeformer.h"
 #include "OptimusNode.h"
 #include "OptimusNodeGraph.h"
+#include "OptimusFunctionNodeGraph.h"
 
 #include "Animation/DebugSkelMeshComponent.h"
 #include "AnimationEditorPreviewActor.h"
@@ -50,6 +51,7 @@
 #include "ISourceCodeAccessModule.h"
 #include "ISourceCodeAccessor.h"
 #include "OptimusEditorStyle.h"
+#include "Toolkits/ToolkitManager.h"
 
 
 #define LOCTEXT_NAMESPACE "OptimusEditor"
@@ -462,6 +464,11 @@ void FOptimusEditor::DeleteSelectedNodes()
 
 bool FOptimusEditor::CanDeleteSelectedNodes() const
 {
+	if (IsGraphReadOnly())
+	{
+		return false;
+	}
+	
 	if (GraphEditorWidget->GetSelectedNodes().IsEmpty())
 	{
 		return false;
@@ -475,7 +482,7 @@ bool FOptimusEditor::CanDeleteSelectedNodes() const
 			return false;
 		}
 	}
-
+	
 	return true;
 }
 
@@ -529,6 +536,11 @@ void FOptimusEditor::PasteNodes() const
 
 bool FOptimusEditor::CanPasteNodes()
 {
+	if (IsGraphReadOnly())
+	{
+		return false;	
+	}
+	
 	return FOptimusEditorModule::Get().GetClipboard().HasValidClipboardContent();
 }
 
@@ -544,6 +556,11 @@ void FOptimusEditor::DuplicateNodes() const
 
 bool FOptimusEditor::CanDuplicateNodes() const
 {
+	if (IsGraphReadOnly())
+	{
+		return false;
+	}
+	
 	return !GetSelectedModelNodes().IsEmpty();
 }
 
@@ -567,6 +584,11 @@ void FOptimusEditor::PackageNodes()
 
 bool FOptimusEditor::CanPackageNodes() const
 {
+	if (IsGraphReadOnly())
+	{
+		return false;
+	}
+	
 	const TArray<UOptimusNode*> ModelNodes = GetSelectedModelNodes();
 	const UOptimusNodeGraph* ModelGraph = EditorGraph->GetModelGraph();	
 	for (UOptimusNode* ModelNode: ModelNodes)
@@ -600,6 +622,11 @@ void FOptimusEditor::UnpackageNodes()
 
 bool FOptimusEditor::CanUnpackageNodes() const
 {
+	if (IsGraphReadOnly())
+	{
+		return false;
+	}
+	
 	const TArray<UOptimusNode*> ModelNodes = GetSelectedModelNodes();	
 	const UOptimusNodeGraph* ModelGraph = EditorGraph->GetModelGraph();	
 	for (UOptimusNode* ModelNode: ModelNodes)
@@ -632,6 +659,10 @@ void FOptimusEditor::CollapseNodesToSubGraph()
 
 bool FOptimusEditor::CanCollapseNodes() const
 {
+	if (IsGraphReadOnly())
+	{
+		return false;
+	}
 	return CanCopyNodes();
 }
 
@@ -643,14 +674,18 @@ void FOptimusEditor::ExpandCollapsedNode()
 	TArray<UObject*> NewNodes;
 	for (UOptimusNode* ModelNode: GetSelectedModelNodes())
 	{
-		NewNodes.Append(ModelGraph->ExpandCollapsedNodes(ModelNode));
+		ModelGraph->ExpandCollapsedNodes(ModelNode);
 	}
-	InspectObjects(NewNodes);
 }
 
 
 bool FOptimusEditor::CanExpandCollapsedNode() const
 {
+	if (IsGraphReadOnly())
+	{
+		return false;
+	}
+	
 	const TArray<UOptimusNode*> ModelNodes = GetSelectedModelNodes();	
 	const UOptimusNodeGraph* ModelGraph = EditorGraph->GetModelGraph();	
 	for (UOptimusNode* ModelNode: ModelNodes)
@@ -665,6 +700,60 @@ bool FOptimusEditor::CanExpandCollapsedNode() const
 	return !ModelNodes.IsEmpty();
 }
 
+void FOptimusEditor::ConvertToFunction()
+{
+	FOptimusActionScope ActionScope(*GetActionStack(), TEXT("Convert to Functions"));
+	UOptimusNodeGraph* ModelGraph = EditorGraph->GetModelGraph();
+	for (UOptimusNode* ModelNode: GetSelectedModelNodes())
+	{
+		ModelGraph->ConvertToFunction(ModelNode);
+	}
+}
+
+bool FOptimusEditor::CanConvertToFunction() const
+{
+	if (IsGraphReadOnly())
+	{
+		return false;
+	}
+	
+	const TArray<UOptimusNode*> ModelNodes = GetSelectedModelNodes();	
+	const UOptimusNodeGraph* ModelGraph = EditorGraph->GetModelGraph();	
+	for (UOptimusNode* ModelNode: ModelNodes)
+	{
+		if(!ModelGraph->IsSubGraphReference(ModelNode))
+		{
+			return false;
+		}
+	}
+
+	return !ModelNodes.IsEmpty();
+}
+
+void FOptimusEditor::ConvertToSubGraph()
+{
+}
+
+bool FOptimusEditor::CanConvertToSubGraph() const
+{
+	if (IsGraphReadOnly())
+	{
+		return false;
+	}
+	
+	const TArray<UOptimusNode*> ModelNodes = GetSelectedModelNodes();	
+	const UOptimusNodeGraph* ModelGraph = EditorGraph->GetModelGraph();	
+	for (UOptimusNode* ModelNode: ModelNodes)
+	{
+		if (!ModelGraph->IsFunctionReference(ModelNode))
+		{
+			return false;
+		}
+	}
+
+	return !ModelNodes.IsEmpty();
+}
+
 void FOptimusEditor::SubscribeToGraphNotifies(UOptimusNodeGraph* Graph, const FOnGraphNotified& OnGraphNotifiedDelegate)
 {
 	if (IsValid(Graph))
@@ -757,7 +846,29 @@ void FOptimusEditor::OnNodeDoubleClicked(class UEdGraphNode* Node)
 	{
 		if (IOptimusNodeGraphProvider* GraphProvider = Cast<IOptimusNodeGraphProvider>(GraphNode->ModelNode))
 		{
-			SetEditGraph(GraphProvider->GetNodeGraphToShow());
+			UOptimusNodeGraph* GraphToShow = GraphProvider->GetNodeGraphToShow();
+			UOptimusNodeGraph* CurrentGraph = GraphNode->ModelNode->GetOwningGraph();
+
+			if (GraphToShow)
+			{
+				if (CurrentGraph->GetCollectionRoot() == GraphToShow->GetCollectionRoot())
+				{
+					SetEditGraph(GraphToShow);
+				}
+				else if (UOptimusDeformer* Deformer = Cast<UOptimusDeformer>(GraphToShow->GetCollectionRoot()))
+				{
+					GEditor->GetEditorSubsystem<UAssetEditorSubsystem>()->OpenEditorForAsset(Deformer);
+
+					TSharedPtr< IToolkit > FoundAssetEditor = FToolkitManager::Get().FindEditorForAsset(Deformer);
+					
+					if (FoundAssetEditor.IsValid())
+					{
+						TSharedPtr<FOptimusEditor> OptimusEditorForExternalGraph = StaticCastSharedPtr<FOptimusEditor>(FoundAssetEditor);
+
+						OptimusEditorForExternalGraph->SetEditGraph(GraphToShow);
+					}	
+				}
+			}
 		}
 	}
 }
@@ -917,6 +1028,7 @@ void FOptimusEditor::HandleDetailsCreated(
 	)
 {
 	PropertyDetailsWidget = InDetailsView;
+	PropertyDetailsWidget->SetIsPropertyReadOnlyDelegate(FIsPropertyReadOnly::CreateSP(this, &FOptimusEditor::IsPropertyReadOnly));
 }
 
 
@@ -971,7 +1083,7 @@ TSharedRef<SGraphEditor> FOptimusEditor::CreateGraphEditorWidget()
 			FExecuteAction::CreateSP(this, &FOptimusEditor::DuplicateNodes),
 			FCanExecuteAction::CreateSP(this, &FOptimusEditor::CanDuplicateNodes)
 		);
-
+#if 0
 		// Packaging commands
 		GraphEditorCommands->MapAction(FOptimusEditorGraphCommands::Get().ConvertToKernelFunction,
 			FExecuteAction::CreateSP(this, &FOptimusEditor::PackageNodes),
@@ -987,7 +1099,8 @@ TSharedRef<SGraphEditor> FOptimusEditor::CreateGraphEditorWidget()
 			FExecuteAction::CreateSP(this, &FOptimusEditor::CollapseNodesToFunction),
 			FCanExecuteAction::CreateSP(this, &FOptimusEditor::CanCollapseNodes)
 		);
-
+#endif
+		
 		GraphEditorCommands->MapAction(FOptimusEditorGraphCommands::Get().CollapseNodesToSubGraph,
 			FExecuteAction::CreateSP(this, &FOptimusEditor::CollapseNodesToSubGraph),
 			FCanExecuteAction::CreateSP(this, &FOptimusEditor::CanCollapseNodes)
@@ -997,6 +1110,16 @@ TSharedRef<SGraphEditor> FOptimusEditor::CreateGraphEditorWidget()
 			FExecuteAction::CreateSP(this, &FOptimusEditor::ExpandCollapsedNode),
 			FCanExecuteAction::CreateSP(this, &FOptimusEditor::CanExpandCollapsedNode)
 		);
+
+		GraphEditorCommands->MapAction(FOptimusEditorGraphCommands::Get().ConvertToFunction,
+			FExecuteAction::CreateSP(this, &FOptimusEditor::ConvertToFunction),
+			FCanExecuteAction::CreateSP(this, &FOptimusEditor::CanConvertToFunction)
+		);
+		
+		GraphEditorCommands->MapAction(FOptimusEditorGraphCommands::Get().ConvertToSubGraph,
+			FExecuteAction::CreateSP(this, &FOptimusEditor::ConvertToSubGraph),
+			FCanExecuteAction::CreateSP(this, &FOptimusEditor::CanConvertToSubGraph)
+		);
 		
 #if 0
 
@@ -1063,9 +1186,12 @@ TSharedRef<SGraphEditor> FOptimusEditor::CreateGraphEditorWidget()
 		.OptimusEditor(SharedThis(this))
 		.OnGraphCrumbClickedEvent_Lambda([this](UOptimusNodeGraph* InNodeGraph) { SetEditGraph(InNodeGraph);});
 
+	
+
 	return SNew(SGraphEditor)
 		.AdditionalCommands(GraphEditorCommands)
-		.IsEditable(true)
+		.IsEditable(this, &FOptimusEditor::IsGraphEditable)
+		.DisplayAsReadOnly(this, &FOptimusEditor::IsGraphReadOnly)
 		.TitleBar(TitleBarWidget)
 		.Appearance(this, &FOptimusEditor::GetGraphAppearance)
 		.GraphToEdit(EditorGraph)
@@ -1081,6 +1207,42 @@ FGraphAppearanceInfo FOptimusEditor::GetGraphAppearance() const
 	return Appearance;
 }
 
+bool FOptimusEditor::IsGraphReadOnly() const
+{
+	if (EditorGraph->GetModelGraph()->GetGraphType() == EOptimusNodeGraphType::Function)
+	{
+		return true;
+	}
+
+	return false;
+}
+
+bool FOptimusEditor::IsGraphEditable() const
+{
+	return !IsGraphReadOnly();
+}
+
+bool FOptimusEditor::IsPropertyReadOnly(const FPropertyAndParent& InPropertyAndParent) const
+{
+	const TArray<TWeakObjectPtr<UObject>>& DetailsSelectedObjects = PropertyDetailsWidget->GetSelectedObjects();
+	for (const TWeakObjectPtr<UObject>& SelectedObject : DetailsSelectedObjects)
+	{
+		if (!SelectedObject.IsValid())
+		{
+			continue;
+		}
+
+		// Everything a function graph owns should be read only
+		if (SelectedObject->GetTypedOuter<UOptimusFunctionNodeGraph>())
+		{
+			return true;
+		}
+		
+	}
+
+	return false;	
+}
+
 
 TArray<UOptimusNode*> FOptimusEditor::GetSelectedModelNodes() const
 {
diff --git a/Engine/Plugins/Animation/DeformerGraph/Source/OptimusEditor/Private/OptimusEditor.h b/Engine/Plugins/Animation/DeformerGraph/Source/OptimusEditor/Private/OptimusEditor.h
--- a/Engine/Plugins/Animation/DeformerGraph/Source/OptimusEditor/Private/OptimusEditor.h
+++ b/Engine/Plugins/Animation/DeformerGraph/Source/OptimusEditor/Private/OptimusEditor.h
@@ -161,6 +161,12 @@ private:
 
 	void ExpandCollapsedNode();
 	bool CanExpandCollapsedNode() const;
+
+	void ConvertToFunction();
+	bool CanConvertToFunction() const;
+	
+	void ConvertToSubGraph();
+	bool CanConvertToSubGraph() const;
 	
 	// ----------------------------------------------------------------------------------------
 	// Graph event listeners
@@ -211,6 +217,9 @@ private:
 	void CreateWidgets();
 	TSharedRef<SGraphEditor> CreateGraphEditorWidget();
 	FGraphAppearanceInfo GetGraphAppearance() const;
+	bool IsGraphReadOnly() const;
+	bool IsGraphEditable() const;
+	bool IsPropertyReadOnly(const FPropertyAndParent& InPropertyAndParent) const;
 
 	TArray<UOptimusNode*> GetSelectedModelNodes() const;
 
diff --git a/Engine/Plugins/Animation/DeformerGraph/Source/OptimusEditor/Private/OptimusEditorGraph.cpp b/Engine/Plugins/Animation/DeformerGraph/Source/OptimusEditor/Private/OptimusEditorGraph.cpp
--- a/Engine/Plugins/Animation/DeformerGraph/Source/OptimusEditor/Private/OptimusEditorGraph.cpp
+++ b/Engine/Plugins/Animation/DeformerGraph/Source/OptimusEditor/Private/OptimusEditorGraph.cpp
@@ -116,8 +116,9 @@ const FSlateBrush* UOptimusEditorGraph::GetGraphTypeIcon(
 	case EOptimusNodeGraphType::ExternalTrigger:
 		return FOptimusEditorStyle::Get().GetBrush(TEXT("GraphType.Trigger"));	
 	case EOptimusNodeGraphType::SubGraph:
-		// TODO: need a new icon
-		return FOptimusEditorStyle::Get().GetBrush(TEXT("GraphType.Update"));
+		return FOptimusEditorStyle::Get().GetBrush(TEXT("GraphType.SubGraph"));
+	case EOptimusNodeGraphType::Function:
+		return FAppStyle::Get().GetBrush(TEXT("GraphEditor.Function_16x"));
 	default:
 		checkNoEntry();
 		break;
diff --git a/Engine/Plugins/Animation/DeformerGraph/Source/OptimusEditor/Private/OptimusEditorGraphCommands.cpp b/Engine/Plugins/Animation/DeformerGraph/Source/OptimusEditor/Private/OptimusEditorGraphCommands.cpp
--- a/Engine/Plugins/Animation/DeformerGraph/Source/OptimusEditor/Private/OptimusEditorGraphCommands.cpp
+++ b/Engine/Plugins/Animation/DeformerGraph/Source/OptimusEditor/Private/OptimusEditorGraphCommands.cpp
@@ -27,6 +27,9 @@ void FOptimusEditorGraphCommands::RegisterCommands()
 	UI_COMMAND(CollapseNodesToFunction, "Collapse to a Function", "Collapse the selected nodes to a reusable function graph, replacing the selected nodes with a single node.", EUserInterfaceActionType::Button, FInputChord());
 	UI_COMMAND(CollapseNodesToSubGraph, "Collapse to a Sub-Graph", "Collapse the selected nodes to a non-reusable sub-graph, replacing the selected nodes with a single node.", EUserInterfaceActionType::Button, FInputChord());
 	UI_COMMAND(ExpandCollapsedNode, "Expand Collapsed Node", "Replace the selected node with the collapsed nodes from inside the selected function or sub-graph node.", EUserInterfaceActionType::Button, FInputChord());
+	
+	UI_COMMAND(ConvertToFunction, "Convert to a Function", "Convert the selected SubGraph to a reusable Function.", EUserInterfaceActionType::Button, FInputChord());
+	UI_COMMAND(ConvertToSubGraph, "Convert to a SubGraph", "Convert the selected Function to a non-reusable SubGraph.", EUserInterfaceActionType::Button, FInputChord());
 }
 
 #undef LOCTEXT_NAMESPACE
diff --git a/Engine/Plugins/Animation/DeformerGraph/Source/OptimusEditor/Private/OptimusEditorGraphCommands.h b/Engine/Plugins/Animation/DeformerGraph/Source/OptimusEditor/Private/OptimusEditorGraphCommands.h
--- a/Engine/Plugins/Animation/DeformerGraph/Source/OptimusEditor/Private/OptimusEditorGraphCommands.h
+++ b/Engine/Plugins/Animation/DeformerGraph/Source/OptimusEditor/Private/OptimusEditorGraphCommands.h
@@ -19,4 +19,7 @@ public:
 	TSharedPtr<FUICommandInfo> CollapseNodesToFunction;
 	TSharedPtr<FUICommandInfo> CollapseNodesToSubGraph;
 	TSharedPtr<FUICommandInfo> ExpandCollapsedNode;
+	
+	TSharedPtr<FUICommandInfo> ConvertToFunction;
+	TSharedPtr<FUICommandInfo> ConvertToSubGraph;
 };
diff --git a/Engine/Plugins/Animation/DeformerGraph/Source/OptimusEditor/Private/OptimusEditorGraphNode.cpp b/Engine/Plugins/Animation/DeformerGraph/Source/OptimusEditor/Private/OptimusEditorGraphNode.cpp
--- a/Engine/Plugins/Animation/DeformerGraph/Source/OptimusEditor/Private/OptimusEditorGraphNode.cpp
+++ b/Engine/Plugins/Animation/DeformerGraph/Source/OptimusEditor/Private/OptimusEditorGraphNode.cpp
@@ -300,10 +300,11 @@ void UOptimusEditorGraphNode::GetNodeContextMenuActions(
 		Clipboard.AddMenuEntry(FGenericCommands::Get().Duplicate);
 
 		FToolMenuSection& PackagingSection = InMenu->AddSection("OptimusNodePackaging", LOCTEXT("NodeMenuPackagingHeader", "Packaging"));
-
+#if 0
+		// NOTE: Unsupported for 5.4 
 		PackagingSection.AddMenuEntry(FOptimusEditorGraphCommands::Get().ConvertToKernelFunction);
 		PackagingSection.AddMenuEntry(FOptimusEditorGraphCommands::Get().ConvertFromKernelFunction);
-
+#endif 
 #if 0
 		// NOTE: Disabled for 5.0
 		PackagingSection.AddMenuEntry(FOptimusEditorGraphCommands::Get().CollapseNodesToFunction);
@@ -311,6 +312,8 @@ void UOptimusEditorGraphNode::GetNodeContextMenuActions(
 		PackagingSection.AddMenuEntry(FOptimusEditorGraphCommands::Get().CollapseNodesToSubGraph);
 		PackagingSection.AddMenuEntry(FOptimusEditorGraphCommands::Get().ExpandCollapsedNode);
 
+		PackagingSection.AddMenuEntry(FOptimusEditorGraphCommands::Get().ConvertToFunction);
+		PackagingSection.AddMenuEntry(FOptimusEditorGraphCommands::Get().ConvertToSubGraph);
 		// FIXME: Add alignment.
 	}
 }
diff --git a/Engine/Plugins/Animation/DeformerGraph/Source/OptimusEditor/Private/OptimusEditorGraphSchema.cpp b/Engine/Plugins/Animation/DeformerGraph/Source/OptimusEditor/Private/OptimusEditorGraphSchema.cpp
--- a/Engine/Plugins/Animation/DeformerGraph/Source/OptimusEditor/Private/OptimusEditorGraphSchema.cpp
+++ b/Engine/Plugins/Animation/DeformerGraph/Source/OptimusEditor/Private/OptimusEditorGraphSchema.cpp
@@ -19,7 +19,10 @@
 #include "OptimusActionStack.h"
 #include "OptimusComponentSource.h"
 #include "OptimusComputeDataInterface.h"
+#include "OptimusDeformer.h"
 #include "OptimusEditorGraphConnectionDrawingPolicy.h"
+#include "OptimusFunctionNodeGraph.h"
+#include "AssetRegistry/AssetRegistryModule.h"
 #include "Styling/SlateIconFinder.h"
 
 #include "Framework/Application/SlateApplication.h"
@@ -163,9 +166,46 @@ void UOptimusEditorGraphSchema::GetGraphActions(
 		IoActionBuilder.AddAction(Action);
 	}
 
+	FAssetRegistryModule& AssetRegistryModule = FModuleManager::LoadModuleChecked<FAssetRegistryModule>("AssetRegistry");
+
+	TArray<FAssetData> AssetDatas;
+
+	// This triggers the gathering of asset tags across assets
+	AssetRegistryModule.Get().GetAssetsByClass(UOptimusDeformer::StaticClass()->GetClassPathName(), AssetDatas);
+
+	for (FAssetData& AssetData : AssetDatas)
+	{
+		FString PublicFunctionString = AssetData.GetTagValueRef<FString>(UOptimusDeformer::PublicFunctionsAssetTagName);
+
+		if (!PublicFunctionString.IsEmpty())
+		{
+			FOptimusFunctionNodeGraphHeaderArray PublicFunctionHeaderArray;
+
+			FOptimusFunctionNodeGraphHeaderArray::StaticStruct()->ImportText(*PublicFunctionString, &PublicFunctionHeaderArray, nullptr, PPF_None, nullptr, {});
+
+			for (const FOptimusFunctionNodeGraphHeader& Header : PublicFunctionHeaderArray.Headers)
+			{
+				const FText FunctionName = FText::FromName(Header.FunctionName);
+
+				const FText Category = FText::FromName(Header.Category);
+
+				TSharedPtr< FOptimusGraphSchemaAction_NewFunctionReferenceNode> Action(
+					new FOptimusGraphSchemaAction_NewFunctionReferenceNode(
+						Category,
+						FunctionName,
+						/* Tooltip */{}, 0, /* Keywords */{}
+				));
+
+				Action->GraphPath = Header.GraphPath;
+
+				IoActionBuilder.AddAction(Action);	
+			}
+		}
+	}
+
 	{
-		TSharedPtr< FOptimusSchemaAction_NewLoopTerminalNodes> Action(
-			new FOptimusSchemaAction_NewLoopTerminalNodes(
+		TSharedPtr< FOptimusGraphSchemaAction_NewLoopTerminalNodes> Action(
+			new FOptimusGraphSchemaAction_NewLoopTerminalNodes(
 				{},
 				FText::FromName(TEXT("Loop")),
 				/* Tooltip */{}, 0, /* Keywords */{}
diff --git a/Engine/Plugins/Animation/DeformerGraph/Source/OptimusEditor/Private/OptimusEditorGraphSchemaActions.cpp b/Engine/Plugins/Animation/DeformerGraph/Source/OptimusEditor/Private/OptimusEditorGraphSchemaActions.cpp
--- a/Engine/Plugins/Animation/DeformerGraph/Source/OptimusEditor/Private/OptimusEditorGraphSchemaActions.cpp
+++ b/Engine/Plugins/Animation/DeformerGraph/Source/OptimusEditor/Private/OptimusEditorGraphSchemaActions.cpp
@@ -5,6 +5,7 @@
 #include "OptimusComponentSource.h"
 #include "OptimusEditorGraph.h"
 #include "OptimusEditorGraphNode.h"
+#include "OptimusFunctionNodeGraph.h"
 
 #include "OptimusNode.h"
 #include "OptimusResourceDescription.h"
@@ -94,7 +95,7 @@ UEdGraphNode* FOptimusGraphSchemaAction_NewDataInterfaceNode::PerformAction(
 	return nullptr;
 }
 
-UEdGraphNode* FOptimusSchemaAction_NewLoopTerminalNodes::PerformAction(UEdGraph* InParentGraph, UEdGraphPin* FromPin,
+UEdGraphNode* FOptimusGraphSchemaAction_NewLoopTerminalNodes::PerformAction(UEdGraph* InParentGraph, UEdGraphPin* FromPin,
 	const FVector2D Location, bool bInSelectNewNode)
 {
 	UOptimusEditorGraph* Graph = Cast<UOptimusEditorGraph>(InParentGraph);
@@ -117,6 +118,27 @@ UEdGraphNode* FOptimusSchemaAction_NewLoopTerminalNodes::PerformAction(UEdGraph*
 	return nullptr;
 }
 
+UEdGraphNode* FOptimusGraphSchemaAction_NewFunctionReferenceNode::PerformAction(UEdGraph* InParentGraph, UEdGraphPin* InFromPin, const FVector2D InLocation, bool bInSelectNewNode)
+{
+	UOptimusEditorGraph* Graph = Cast<UOptimusEditorGraph>(InParentGraph);
+	
+	if (ensure(Graph != nullptr))
+	{
+		UOptimusNode* ModelNode = Graph->GetModelGraph()->AddFunctionReferenceNode(GraphPath, InLocation);
+
+		// FIXME: Automatic connection from the given pin.
+
+		UOptimusEditorGraphNode* GraphNode = Graph->FindGraphNodeFromModelNode(ModelNode);
+		if (GraphNode && bInSelectNewNode)
+		{
+			Graph->SelectNodeSet({GraphNode});
+		}
+		return GraphNode;
+	}
+
+	return nullptr;
+}
+
 
 static FText GetGraphTooltip(UOptimusNodeGraph* InGraph)
 {
@@ -126,13 +148,12 @@ static FText GetGraphTooltip(UOptimusNodeGraph* InGraph)
 
 FOptimusSchemaAction_Graph::FOptimusSchemaAction_Graph(
 	UOptimusNodeGraph* InGraph,
-	int32 InGrouping,
 	const FText& InCategory) : 
 		FEdGraphSchemaAction(
 			InCategory, 
 			FText::FromString(InGraph->GetName()), 
 			GetGraphTooltip(InGraph), 
-			InGrouping, 
+			0, 
 			FText(), 
 			int32(EOptimusSchemaItemGroup::Graphs) 
 		), 
@@ -143,14 +164,13 @@ FOptimusSchemaAction_Graph::FOptimusSchemaAction_Graph(
 
 
 FOptimusSchemaAction_Binding::FOptimusSchemaAction_Binding(
-	UOptimusComponentSourceBinding* InBinding,
-	int32 InGrouping
+	UOptimusComponentSourceBinding* InBinding
 	) :
 	FEdGraphSchemaAction(
 			FText::GetEmpty(),
 			FText::FromString(InBinding->GetName()),
 			FText::GetEmpty(),
-			InGrouping,
+			0,
 			FText(),
 			int32(EOptimusSchemaItemGroup::Bindings)
 		)
@@ -161,13 +181,13 @@ FOptimusSchemaAction_Binding::FOptimusSchemaAction_Binding(
 
 
 FOptimusSchemaAction_Resource::FOptimusSchemaAction_Resource(
-	UOptimusResourceDescription* InResource, 
-	int32 InGrouping ) :
+	UOptimusResourceDescription* InResource
+	) :
 	FEdGraphSchemaAction(
 			FText::GetEmpty(),
 			FText::FromString(InResource->GetName()),
 			FText::GetEmpty(),
-			InGrouping,
+			0,
 			FText(),
 			int32(EOptimusSchemaItemGroup::Resources)
 		)
@@ -177,13 +197,13 @@ FOptimusSchemaAction_Resource::FOptimusSchemaAction_Resource(
 
 
 FOptimusSchemaAction_Variable::FOptimusSchemaAction_Variable(
-	UOptimusVariableDescription* InVariable, 
-	int32 InGrouping ) : 
+	UOptimusVariableDescription* InVariable 
+	) : 
 	FEdGraphSchemaAction(
           FText::GetEmpty(),
           FText::FromString(InVariable->GetName()),
           FText::GetEmpty(),
-          InGrouping,
+          0,
           FText(),
           int32(EOptimusSchemaItemGroup::Variables))
 {
diff --git a/Engine/Plugins/Animation/DeformerGraph/Source/OptimusEditor/Private/OptimusEditorGraphSchemaActions.h b/Engine/Plugins/Animation/DeformerGraph/Source/OptimusEditor/Private/OptimusEditorGraphSchemaActions.h
--- a/Engine/Plugins/Animation/DeformerGraph/Source/OptimusEditor/Private/OptimusEditorGraphSchemaActions.h
+++ b/Engine/Plugins/Animation/DeformerGraph/Source/OptimusEditor/Private/OptimusEditorGraphSchemaActions.h
@@ -7,6 +7,7 @@
 
 #include "OptimusComputeDataInterface.h"
 #include "OptimusDataType.h"
+#include "OptimusFunctionNodeGraphHeader.h"
 
 #include "OptimusEditorGraphSchemaActions.generated.h"
 
@@ -89,20 +90,37 @@ struct FOptimusGraphSchemaAction_NewDataInterfaceNode :
 
 
 USTRUCT()
-struct FOptimusSchemaAction_NewLoopTerminalNodes : public FEdGraphSchemaAction
+struct FOptimusGraphSchemaAction_NewLoopTerminalNodes : public FEdGraphSchemaAction
 {
 	GENERATED_BODY()
 
 	// Inherit the base class's constructors
 	using FEdGraphSchemaAction::FEdGraphSchemaAction;
 	
-	static FName StaticGetTypeId() { static FName Type("FOptimusSchemaAction_LoopTerminal"); return Type; }
+	static FName StaticGetTypeId() { static FName Type("FOptimusGraphSchemaAction_LoopTerminal"); return Type; }
 	FName GetTypeId() const override { return StaticGetTypeId(); }
 
 	// FEdGraphSchemaAction overrides
 	UEdGraphNode* PerformAction(class UEdGraph* InParentGraph, UEdGraphPin* FromPin, const FVector2D Location, bool bInSelectNewNode = true) override;
 };
 
+USTRUCT()
+struct FOptimusGraphSchemaAction_NewFunctionReferenceNode: 
+	public FEdGraphSchemaAction
+{
+	GENERATED_BODY()
+
+	// Inherit the base class's constructors
+	using FEdGraphSchemaAction::FEdGraphSchemaAction;
+
+	TSoftObjectPtr<UOptimusFunctionNodeGraph> GraphPath;
+
+	static FName StaticGetTypeId() { static FName Type("FOptimusGraphSchemaAction_NewFunctionReferenceNode"); return Type; }
+	FName GetTypeId() const override { return StaticGetTypeId(); }
+
+	// FEdGraphSchemaAction overrides
+	UEdGraphNode* PerformAction(class UEdGraph* InParentGraph, UEdGraphPin* InFromPin, const FVector2D InLocation, bool bInSelectNewNode = true) override;
+};
 
 /// Reference to a UOptimusNodeGraph.
 USTRUCT()
@@ -126,7 +144,6 @@ struct FOptimusSchemaAction_Graph :
 
 	FOptimusSchemaAction_Graph(
 	    UOptimusNodeGraph* InGraph,
-	    int32 InGrouping,
 	    const FText& InCategory = FText::GetEmpty());
 
 	// FEdGraphSchemaAction overrides
@@ -151,10 +168,7 @@ struct FOptimusSchemaAction_Binding : public FEdGraphSchemaAction
 
 	FOptimusSchemaAction_Binding() = default;
 
-	FOptimusSchemaAction_Binding(
-		UOptimusComponentSourceBinding* InBinding,
-		int32 InGrouping);
-
+	FOptimusSchemaAction_Binding( UOptimusComponentSourceBinding* InBinding );
 	// FEdGraphSchemaAction overrides
 	bool IsParentable() const override { return false; }
 };
@@ -177,9 +191,7 @@ struct FOptimusSchemaAction_Resource : public FEdGraphSchemaAction
 
 	FOptimusSchemaAction_Resource() = default;
 
-	FOptimusSchemaAction_Resource(
-	    UOptimusResourceDescription* InResource,
-	    int32 InGrouping);
+	FOptimusSchemaAction_Resource( UOptimusResourceDescription* InResource );
 
 	// FEdGraphSchemaAction overrides
 	bool IsParentable() const override { return false; }
@@ -203,9 +215,7 @@ struct FOptimusSchemaAction_Variable : public FEdGraphSchemaAction
 
 	FOptimusSchemaAction_Variable() = default;
 
-	FOptimusSchemaAction_Variable(
-	    UOptimusVariableDescription* InVariable,
-	    int32 InGrouping);
+	FOptimusSchemaAction_Variable( UOptimusVariableDescription* InVariable );
 
 	// FEdGraphSchemaAction overrides
 	bool IsParentable() const override { return false; }
diff --git a/Engine/Plugins/Animation/DeformerGraph/Source/OptimusEditor/Private/OptimusEditorStyle.cpp b/Engine/Plugins/Animation/DeformerGraph/Source/OptimusEditor/Private/OptimusEditorStyle.cpp
--- a/Engine/Plugins/Animation/DeformerGraph/Source/OptimusEditor/Private/OptimusEditorStyle.cpp
+++ b/Engine/Plugins/Animation/DeformerGraph/Source/OptimusEditor/Private/OptimusEditorStyle.cpp
@@ -93,9 +93,10 @@ FOptimusEditorStyle::FOptimusEditorStyle() :
 
 	// Graph type icons
 	{
-		Set("GraphType.Setup", new IMAGE_BRUSH_SVG("Icons/Graph_Setup", IconSize20x20, DefaultForeground));
-		Set("GraphType.Trigger", new IMAGE_BRUSH_SVG("Icons/Graph_Trigger", IconSize20x20, DefaultForeground));
-		Set("GraphType.Update", new IMAGE_BRUSH_SVG("Icons/Graph_Update", IconSize20x20, DefaultForeground));
+		Set("GraphType.Setup", new IMAGE_BRUSH_SVG("Icons/Graph_Setup", IconSize16x16, DefaultForeground));
+		Set("GraphType.Trigger", new IMAGE_BRUSH_SVG("Icons/Graph_Trigger", IconSize16x16, DefaultForeground));
+		Set("GraphType.Update", new IMAGE_BRUSH_SVG("Icons/Graph_Update", IconSize16x16, DefaultForeground));
+		Set("GraphType.SubGraph", new IMAGE_BRUSH_SVG("Icons/Graph_SubGraph", IconSize16x16, DefaultForeground));
 	}
 
 	// Graph styles
diff --git a/Engine/Plugins/Animation/DeformerGraph/Source/OptimusEditor/Private/Widgets/SOptimusEditorGraphExplorer.cpp b/Engine/Plugins/Animation/DeformerGraph/Source/OptimusEditor/Private/Widgets/SOptimusEditorGraphExplorer.cpp
--- a/Engine/Plugins/Animation/DeformerGraph/Source/OptimusEditor/Private/Widgets/SOptimusEditorGraphExplorer.cpp
+++ b/Engine/Plugins/Animation/DeformerGraph/Source/OptimusEditor/Private/Widgets/SOptimusEditorGraphExplorer.cpp
@@ -27,6 +27,7 @@
 
 // Ick.
 #include "GraphActionNode.h"
+#include "OptimusFunctionNodeGraph.h"
 
 
 #define LOCTEXT_NAMESPACE "OptimusGraphExplorer"
@@ -337,7 +338,7 @@ void SOptimusEditorGraphExplorer::CollectAllActions(FGraphActionListBuilderBase&
 	for (UOptimusNodeGraph* Graph : Deformer->GetGraphs())
 	{
 		FText GraphCategory = GetGraphSubCategory(Graph);
-		TSharedPtr<FOptimusSchemaAction_Graph> GraphAction = MakeShared<FOptimusSchemaAction_Graph>(Graph, /*Grouping=*/1);
+		TSharedPtr<FOptimusSchemaAction_Graph> GraphAction = MakeShared<FOptimusSchemaAction_Graph>(Graph);
 		OutAllActions.AddAction(GraphAction);
 
 		CollectChildGraphActions(OutAllActions, Graph, GraphCategory);
@@ -345,19 +346,19 @@ void SOptimusEditorGraphExplorer::CollectAllActions(FGraphActionListBuilderBase&
 	
 	for (UOptimusComponentSourceBinding* Binding : Deformer->GetComponentBindings())
 	{
-		TSharedPtr<FOptimusSchemaAction_Binding> BindingAction = MakeShared<FOptimusSchemaAction_Binding>(Binding, /*Grouping=*/2);
+		TSharedPtr<FOptimusSchemaAction_Binding> BindingAction = MakeShared<FOptimusSchemaAction_Binding>(Binding);
 		OutAllActions.AddAction(BindingAction);
 	}
 
 	for (UOptimusResourceDescription* Resource : Deformer->GetResources())
 	{
-		TSharedPtr<FOptimusSchemaAction_Resource> ResourceAction = MakeShared<FOptimusSchemaAction_Resource>(Resource, /*Grouping=*/3);
+		TSharedPtr<FOptimusSchemaAction_Resource> ResourceAction = MakeShared<FOptimusSchemaAction_Resource>(Resource);
 		OutAllActions.AddAction(ResourceAction);
 	}
 
 	for (UOptimusVariableDescription* Variable : Deformer->GetVariables())
 	{
-		TSharedPtr<FOptimusSchemaAction_Variable> VariableAction = MakeShared<FOptimusSchemaAction_Variable>(Variable, /*Grouping=*/4);
+		TSharedPtr<FOptimusSchemaAction_Variable> VariableAction = MakeShared<FOptimusSchemaAction_Variable>(Variable);
 		OutAllActions.AddAction(VariableAction);
 	}
 }
@@ -382,7 +383,7 @@ void SOptimusEditorGraphExplorer::CollectChildGraphActions(
 
 	for (UOptimusNodeGraph* SubGraph: InParentGraph->GetGraphs())
 	{
-		TSharedPtr<FOptimusSchemaAction_Graph> GraphAction = MakeShared<FOptimusSchemaAction_Graph>(SubGraph, /*Grouping=*/1, Category);
+		TSharedPtr<FOptimusSchemaAction_Graph> GraphAction = MakeShared<FOptimusSchemaAction_Graph>(SubGraph, Category);
 		OutAllActions.AddAction(GraphAction);
 	}
 
diff --git a/Engine/Plugins/Animation/DeformerGraph/Source/OptimusEditor/Private/Widgets/SOptimusEditorGraphExplorerItem.cpp b/Engine/Plugins/Animation/DeformerGraph/Source/OptimusEditor/Private/Widgets/SOptimusEditorGraphExplorerItem.cpp
--- a/Engine/Plugins/Animation/DeformerGraph/Source/OptimusEditor/Private/Widgets/SOptimusEditorGraphExplorerItem.cpp
+++ b/Engine/Plugins/Animation/DeformerGraph/Source/OptimusEditor/Private/Widgets/SOptimusEditorGraphExplorerItem.cpp
@@ -200,13 +200,14 @@ void SOptimusEditorGraphExplorerItem::Construct(
 			SNew(SHorizontalBox)
 			+ SHorizontalBox::Slot()
 			.AutoWidth()
+			.VAlign(VAlign_Center)
 			[
 				CreateIconWidget(InCreateData, bIsReadOnlyCreate)
 			]
 			+ SHorizontalBox::Slot()
 			.FillWidth(1.f)
 			.VAlign(VAlign_Center)
-			.Padding(/* horizontal */ 3.0f, /* vertical */ 0.0f)
+			.Padding(/* horizontal */ 3.0f, /* vertical */ 3.0f)
 			[
 				CreateTextSlotWidget(InCreateData, bIsReadOnlyCreate)
 			]		
diff --git a/Engine/Plugins/Animation/DeformerGraph/Source/OptimusEditor/Private/Widgets/SOptimusShaderTextDocumentTab.cpp b/Engine/Plugins/Animation/DeformerGraph/Source/OptimusEditor/Private/Widgets/SOptimusShaderTextDocumentTab.cpp
--- a/Engine/Plugins/Animation/DeformerGraph/Source/OptimusEditor/Private/Widgets/SOptimusShaderTextDocumentTab.cpp
+++ b/Engine/Plugins/Animation/DeformerGraph/Source/OptimusEditor/Private/Widgets/SOptimusShaderTextDocumentTab.cpp
@@ -89,7 +89,7 @@ void SOptimusShaderTextDocumentTab::Construct(const FArguments& InArgs, UObject*
 			[
 				SAssignNew(ShaderTextTextBox, SOptimusShaderTextDocumentTextBox)
 				.Text(this, &SOptimusShaderTextDocumentTab::GetShaderTextAsText)
-				.IsReadOnly(false)
+				.IsReadOnly(this, &SOptimusShaderTextDocumentTab::IsShaderTextReadOnly)
 				.Marshaller(SyntaxHighlighterShaderText)
 				.OnTextChanged(this, &SOptimusShaderTextDocumentTab::OnShaderTextChanged)	
 			]
@@ -206,4 +206,14 @@ void SOptimusShaderTextDocumentTab::OnDiagnosticsUpdated() const
 	}
 }
 
+bool SOptimusShaderTextDocumentTab::IsShaderTextReadOnly() const
+{
+	if (IOptimusShaderTextProvider* Provider = Cast<IOptimusShaderTextProvider>(ShaderTextProviderObject))
+	{
+		return Provider->IsShaderTextReadOnly();
+	}
+
+	return false;
+}
+
 #undef LOCTEXT_NAMESPACE
diff --git a/Engine/Plugins/Animation/DeformerGraph/Source/OptimusEditor/Private/Widgets/SOptimusShaderTextDocumentTab.h b/Engine/Plugins/Animation/DeformerGraph/Source/OptimusEditor/Private/Widgets/SOptimusShaderTextDocumentTab.h
--- a/Engine/Plugins/Animation/DeformerGraph/Source/OptimusEditor/Private/Widgets/SOptimusShaderTextDocumentTab.h
+++ b/Engine/Plugins/Animation/DeformerGraph/Source/OptimusEditor/Private/Widgets/SOptimusShaderTextDocumentTab.h
@@ -39,6 +39,7 @@ private:
 	FText GetShaderTextAsText() const;
 	void OnShaderTextChanged(const FText& InText) const;
 	void OnDiagnosticsUpdated() const;
+	bool IsShaderTextReadOnly() const;
 
 	TSharedRef<FOptimusHLSLSyntaxHighlighter> SyntaxHighlighterDeclarations;
 	TSharedRef<FOptimusHLSLSyntaxHighlighter> SyntaxHighlighterShaderText;
