[Deformer Graph] added "ConvertToSubGraph"

also fix subgraph rename not refreshing nodes in the graph editor

diff --git a/Engine/Plugins/Animation/DeformerGraph/Source/OptimusCore/Private/OptimusNodeGraph.cpp b/Engine/Plugins/Animation/DeformerGraph/Source/OptimusCore/Private/OptimusNodeGraph.cpp
--- a/Engine/Plugins/Animation/DeformerGraph/Source/OptimusCore/Private/OptimusNodeGraph.cpp
+++ b/Engine/Plugins/Animation/DeformerGraph/Source/OptimusCore/Private/OptimusNodeGraph.cpp
@@ -56,6 +56,8 @@
 const FName UOptimusNodeGraph::SetupGraphName("SetupGraph");
 const FName UOptimusNodeGraph::UpdateGraphName("UpdateGraph");
 const TCHAR* UOptimusNodeGraph::LibraryRoot = TEXT("@Library");
+const FName UOptimusNodeGraph::DefaultSubGraphName("SubGraph");
+const FName UOptimusNodeGraph::DefaultSubGraphRefNodeName("SubGraphNode");
 
 FString UOptimusNodeGraph::GetFunctionGraphPath(const FString& InFunctionName)
 {
@@ -1179,7 +1181,7 @@ UOptimusNode* UOptimusNodeGraph::CollapseNodesToSubGraph(
 	FOptimusCompoundAction *Action = new FOptimusCompoundAction(TEXT("Collapse Nodes to Sub-graph"));
 	IOptimusPathResolver* PathResolver = GetPathResolver();
 
-	FName SubGraphName("SubGraph");
+	FName SubGraphName = DefaultSubGraphName;
 	SubGraphName = Optimus::GetUniqueNameForScope(this, SubGraphName);
 	
 	FOptimusNodeGraphAction_AddGraph* CreateGraph = new FOptimusNodeGraphAction_AddGraph(
@@ -1193,7 +1195,7 @@ UOptimusNode* UOptimusNodeGraph::CollapseNodesToSubGraph(
 		});
 	Action->AddSubAction(CreateGraph);
 
-	FString SubGraphPath = GetGraphPath() + TEXT("/") + SubGraphName.ToString();
+	FString SubGraphPath = ConstructSubGraphPath(SubGraphName.ToString());
 
 	// Create the entry and return nodes.
 	FBox2D NodeBox(ForceInit);
@@ -1256,7 +1258,7 @@ UOptimusNode* UOptimusNodeGraph::CollapseNodesToSubGraph(
 	}
 
 	// Create the reference node and connect it.
-	FName GraphNodeRefName("SubGraphNode");
+	FName GraphNodeRefName = DefaultSubGraphRefNodeName;
 	GraphNodeRefName = Optimus::GetUniqueNameForScope(this, GraphNodeRefName);
 	FOptimusNodeGraphAction_AddNode* AddSubGraphRefNodeAction = new FOptimusNodeGraphAction_AddNode(
 		GetGraphPath(), UOptimusNode_SubGraphReference::StaticClass(), GraphNodeRefName,
@@ -1825,6 +1827,230 @@ bool UOptimusNodeGraph::ConvertToFunction(UOptimusNode* InSubGraphNode)
 	return true;
 }
 
+bool UOptimusNodeGraph::ConvertToSubGraph(UOptimusNode* InFunctionNode)
+{
+	const bool bIsFunctionReference = IsFunctionReference(InFunctionNode);
+	if (!bIsFunctionReference)
+	{
+		return false;
+	}
+	
+	UOptimusNode_FunctionReference* FunctionNode = CastChecked<UOptimusNode_FunctionReference>(InFunctionNode);
+	UOptimusFunctionNodeGraph* FunctionGraph = FunctionNode->FunctionGraph.Get();
+
+	IOptimusPathResolver* PathResolver = GetPathResolver();
+
+	// Ensure the name is unique here instead of relying on AddGraph action
+	// because subsequent actions needs to use the final graph path
+	// Note: we could change the actions to accept a lambda instead of specific value in the future
+	FName SubGraphName = DefaultSubGraphName;
+	SubGraphName = Optimus::GetUniqueNameForScope(this, SubGraphName);
+	
+	const FString SubGraphPath = ConstructSubGraphPath(SubGraphName.ToString());
+	const FVector2D EntryNodePosition = FunctionGraph->EntryNode->GetGraphPosition();
+	const FVector2D ReturnNodePosition = FunctionGraph->ReturnNode->GetGraphPosition();
+	const FVector2D ReferenceNodePosition = FunctionNode->GetGraphPosition();
+	const FName SubGraphNodeName = Optimus::GetUniqueNameForScope(this, DefaultSubGraphRefNodeName);
+	const FString SubGraphNodeNameString = SubGraphNodeName.ToString();
+	const FString CurrentGraphPath = GetGraphPath();
+
+	TMap<FString, TArray<FString>> ExternalLinksToAdd;
+
+	for (UOptimusNodeLink* Link : GetAllLinks())
+	{
+		if (Link->GetNodeInputPin()->GetOwningNode() == FunctionNode)
+		{
+			FString FunctionNodeInputPinPath = ConstructPath(
+				CurrentGraphPath,
+				SubGraphNodeNameString,
+				Link->GetNodeInputPin()->GetUniqueName().ToString()
+				);
+
+			FString ExistingOutputPinPath = Link->GetNodeOutputPin()->GetPinPath();
+			ExternalLinksToAdd.FindOrAdd(ExistingOutputPinPath).Add(FunctionNodeInputPinPath);
+		}
+		else if (Link->GetNodeOutputPin()->GetOwningNode() == FunctionNode)
+		{
+			FString FunctionNodeOutputPinPath = ConstructPath(
+				CurrentGraphPath,
+				SubGraphNodeNameString,
+				Link->GetNodeOutputPin()->GetUniqueName().ToString()
+				);
+
+			FString ExistingInputPinPath = Link->GetNodeInputPin()->GetPinPath();
+			ExternalLinksToAdd.FindOrAdd(FunctionNodeOutputPinPath).Add(ExistingInputPinPath);
+		}
+	}
+
+	TMap<FString, TArray<FString>> InternalLinksToAdd;
+	
+	for (UOptimusNodeLink* Link : FunctionGraph->GetAllLinks())
+	{
+		if (Link->GetNodeInputPin()->GetOwningNode() == FunctionGraph->ReturnNode)
+		{
+			FString ReturnNodeInputPinPath = ConstructPath(
+				SubGraphPath,
+				UOptimusNode_GraphTerminal::ReturnNodeName.ToString(),
+				Link->GetNodeInputPin()->GetUniqueName().ToString()
+				);
+
+			FString InternalNodeOutputPinPath = ConstructPath(
+				SubGraphPath,
+				Link->GetNodeOutputPin()->GetOwningNode()->GetName(),
+				Link->GetNodeOutputPin()->GetUniqueName().ToString()
+				);
+
+			InternalLinksToAdd.FindOrAdd(InternalNodeOutputPinPath).Add(ReturnNodeInputPinPath);
+		}
+		else if (Link->GetNodeOutputPin()->GetOwningNode() == FunctionGraph->EntryNode)
+		{
+			FString EntryNodeOutputPinPath = ConstructPath(
+				SubGraphPath,
+				UOptimusNode_GraphTerminal::EntryNodeName.ToString(),
+				Link->GetNodeOutputPin()->GetUniqueName().ToString()
+				);
+
+			FString InternalNodeInputPinPath = ConstructPath(
+				SubGraphPath,
+				Link->GetNodeInputPin()->GetOwningNode()->GetName(),
+				Link->GetNodeInputPin()->GetUniqueName().ToString()
+				);
+			InternalLinksToAdd.FindOrAdd(EntryNodeOutputPinPath).Add(InternalNodeInputPinPath);
+		}
+		else
+		{
+			FString OutputPinPath = ConstructPath(
+					SubGraphPath,
+					Link->GetNodeOutputPin()->GetOwningNode()->GetName(),
+					Link->GetNodeOutputPin()->GetUniqueName().ToString()
+					);
+
+			FString InputPinPath = ConstructPath(
+				SubGraphPath,
+				Link->GetNodeInputPin()->GetOwningNode()->GetName(),
+				Link->GetNodeInputPin()->GetUniqueName().ToString()
+				);
+			
+			InternalLinksToAdd.FindOrAdd(OutputPinPath).Add(InputPinPath);
+		}
+	}	
+	
+	
+	TArray<UOptimusNode*> SubGraphNodes = FunctionGraph->GetAllNodes();
+	
+	SubGraphNodes.RemoveAll([](UOptimusNode* InNode)
+	{
+		return Cast<UOptimusNode_GraphTerminal>(InNode) != nullptr;
+	});
+
+	TArray<UOptimusNode*> NodesToDuplicate;
+	TArray<UOptimusNodePair*> NodePairsToDuplicate;
+	AddPairNodesToArray(SubGraphNodes, NodesToDuplicate, NodePairsToDuplicate);
+
+	TArray<TPair<FString, FString>> NewNodePairs;
+	for (const UOptimusNodePair* NodePair : NodePairsToDuplicate)
+	{
+		// Generate new pair
+		const UOptimusNode *FirstNode = NodePair->GetFirst();
+		const UOptimusNode *SecondNode = NodePair->GetSecond();
+
+		FString NewFirstNodePath = ConstructPath(SubGraphPath, FirstNode->GetName(), {});
+		FString NewSecondNodePath = ConstructPath(SubGraphPath, SecondNode->GetName(), {});
+
+		NewNodePairs.Add({NewFirstNodePath, NewSecondNodePath});
+	}	
+
+	{
+		FOptimusActionScope(*GetActionStack(), TEXT("Convert To SubGraph"));
+		GetActionStack()->RunAction<FOptimusNodeGraphAction_AddGraph>(this, EOptimusNodeGraphType::SubGraph, SubGraphName, INDEX_NONE,
+			[FunctionGraph](UOptimusNodeGraph* InGraph)
+			{
+				UOptimusNodeSubGraph* SubGraph = CastChecked<UOptimusNodeSubGraph>(InGraph);
+				SubGraph->InputBindings = FunctionGraph->InputBindings;
+				SubGraph->OutputBindings = FunctionGraph->OutputBindings;
+				return true;
+			});
+	
+		GetActionStack()->RunAction<FOptimusNodeGraphAction_AddNode>(
+			SubGraphPath, UOptimusNode_GraphTerminal::StaticClass(), UOptimusNode_GraphTerminal::EntryNodeName,
+			[SubGraphPath, EntryNodePosition, PathResolver](UOptimusNode* InNode)
+			{
+				UOptimusNode_GraphTerminal* EntryNode = Cast<UOptimusNode_GraphTerminal>(InNode); 
+				UOptimusNodeSubGraph* SubGraph = Cast<UOptimusNodeSubGraph>(PathResolver->ResolveGraphPath(SubGraphPath));
+
+				SubGraph->EntryNode = EntryNode;
+		
+				EntryNode->TerminalType = EOptimusTerminalType::Entry;
+				EntryNode->OwningGraph = SubGraph;
+				return EntryNode->SetGraphPositionDirect(EntryNodePosition);
+			});
+
+		GetActionStack()->RunAction<FOptimusNodeGraphAction_AddNode>(
+			SubGraphPath, UOptimusNode_GraphTerminal::StaticClass(), UOptimusNode_GraphTerminal::ReturnNodeName,
+			[SubGraphPath, ReturnNodePosition, PathResolver](UOptimusNode* InNode)
+			{
+				UOptimusNode_GraphTerminal* ReturnNode = Cast<UOptimusNode_GraphTerminal>(InNode); 
+				UOptimusNodeSubGraph* SubGraph = Cast<UOptimusNodeSubGraph>(PathResolver->ResolveGraphPath(SubGraphPath));
+
+				SubGraph->ReturnNode = ReturnNode;
+			
+				ReturnNode->TerminalType = EOptimusTerminalType::Return;
+				ReturnNode->OwningGraph = SubGraph;
+				return ReturnNode->SetGraphPositionDirect(ReturnNodePosition);
+			});
+
+		
+		// Copy the nodes to the sub graph
+		for (UOptimusNode* Node: NodesToDuplicate)
+		{
+			GetActionStack()->RunAction<FOptimusNodeGraphAction_DuplicateNode>(
+				SubGraphPath, Node, Node->GetFName(), [](UOptimusNode*) { return true; });
+		}
+
+		for (const TPair<FString, FString>& NewPair : NewNodePairs)
+		{
+			GetActionStack()->RunAction<FOptimusNodeGraphAction_AddNodePair>(NewPair.Key, NewPair.Value);
+		}	
+		
+		// Restore links
+		for (const TPair<FString, TArray<FString>>& InternalLinks : InternalLinksToAdd )
+		{
+			const FString& OutputPinPath = InternalLinks.Key;
+			for (const FString& InputPinPath : InternalLinks.Value)
+			{
+				GetActionStack()->RunAction<FOptimusNodeGraphAction_AddLink>(OutputPinPath,	InputPinPath);
+			}
+		}
+		
+		// Create the sub graph ref node
+		GetActionStack()->RunAction<FOptimusNodeGraphAction_AddNode>(
+			GetGraphPath(), UOptimusNode_SubGraphReference::StaticClass(), SubGraphNodeName,
+			[ReferenceNodePosition, SubGraphPath, PathResolver](UOptimusNode* InNode)
+			{
+				UOptimusNode_SubGraphReference* SubGraphNode = Cast<UOptimusNode_SubGraphReference>(InNode); 
+				UOptimusNodeSubGraph* SubGraph = Cast<UOptimusNodeSubGraph>(PathResolver->ResolveGraphPath(SubGraphPath));
+				SubGraphNode->SubGraph = SubGraph;
+				
+				return SubGraphNode->SetGraphPositionDirect(ReferenceNodePosition);
+			});
+
+		// Remove the function node
+		RemoveNodes({FunctionNode});
+		
+		// Restore links
+		for (const TPair<FString, TArray<FString>>& ExternalLinks : ExternalLinksToAdd )
+		{
+			const FString& OutputPinPath = ExternalLinks.Key;
+			for (const FString& InputPinPath : ExternalLinks.Value)
+			{
+				GetActionStack()->RunAction<FOptimusNodeGraphAction_AddLink>(OutputPinPath,	InputPinPath);
+			}
+		}
+	}
+	
+	return true;
+}
+
 
 bool UOptimusNodeGraph::IsCustomKernel(UOptimusNode* InNode) const
 {
@@ -2728,6 +2954,13 @@ TArray<int32> UOptimusNodeGraph::GetAllLinkIndexesToPin(
 	return LinkIndexes;
 }
 
+FString UOptimusNodeGraph::ConstructSubGraphPath(const FString& InSubGraphName) const
+{
+	FString SubGraphPath = GetGraphPath() + TEXT("/") + InSubGraphName;
+
+	return SubGraphPath;
+}
+
 
 UOptimusActionStack* UOptimusNodeGraph::GetActionStack() const
 {
@@ -2870,7 +3103,23 @@ bool UOptimusNodeGraph::RenameGraphDirect(UOptimusNodeGraph* InGraph, const FStr
 {
 	if (Optimus::RenameObject(InGraph, *InNewName, nullptr))
 	{
+		if (InGraph->GetGraphType() == EOptimusNodeGraphType::SubGraph)
+		{
+			for (UOptimusNode* Node : GetAllNodes())
+			{
+				if (UOptimusNode_SubGraphReference* SubGraphNode = Cast<UOptimusNode_SubGraphReference>(Node))
+				{
+					if (SubGraphNode->SubGraph == InGraph)
+					{
+						Notify(EOptimusGraphNotifyType::NodeDisplayNameChanged, SubGraphNode);
+						break;
+					}
+				}
+			}
+		}
+		
 		GlobalNotify(EOptimusGlobalNotifyType::GraphRenamed, InGraph);
+
 		return true;
 	}
 
diff --git a/Engine/Plugins/Animation/DeformerGraph/Source/OptimusCore/Public/OptimusNodeGraph.h b/Engine/Plugins/Animation/DeformerGraph/Source/OptimusCore/Public/OptimusNodeGraph.h
--- a/Engine/Plugins/Animation/DeformerGraph/Source/OptimusCore/Public/OptimusNodeGraph.h
+++ b/Engine/Plugins/Animation/DeformerGraph/Source/OptimusCore/Public/OptimusNodeGraph.h
@@ -75,10 +75,12 @@ public:
 	static const FName SetupGraphName;
 	static const FName UpdateGraphName;
 	static const TCHAR* LibraryRoot;
+	static const FName DefaultSubGraphName;
+	static const FName DefaultSubGraphRefNodeName;
 	
 	// Function Graphs are addressed in a special way
 	static FString GetFunctionGraphPath(const FString& InFunctionName);
-	
+
 	
 	// Check if the duplication took place at the asset level
 	// if so, we have to recreate all constant/attribute nodes such that their class pointers
@@ -272,6 +274,11 @@ public:
 	 */
 	UFUNCTION(BlueprintCallable, Category = OptimusNodeGraph)
 	bool ConvertToFunction(UOptimusNode* InSubGraphNode);
+
+	/** Take a function node convert it to a subgraph node in-place
+	 */
+	UFUNCTION(BlueprintCallable, Category = OptimusNodeGraph)
+	bool ConvertToSubGraph(UOptimusNode* InFunctionNode);
 	
 	/** Returns true if the node in question is a custom kernel node that can be converted to
 	  * a kernel function with ConvertCustomKernelToFunction.
@@ -545,6 +552,8 @@ private:
 		const UOptimusNodePin* InNodePin
 		) const;
 
+	FString ConstructSubGraphPath(const FString& InSubGraphName) const;
+	
 	UPROPERTY(NonTransactional)
 	TArray<TObjectPtr<UOptimusNode>> Nodes;
 
diff --git a/Engine/Plugins/Animation/DeformerGraph/Source/OptimusEditor/Private/OptimusEditor.cpp b/Engine/Plugins/Animation/DeformerGraph/Source/OptimusEditor/Private/OptimusEditor.cpp
--- a/Engine/Plugins/Animation/DeformerGraph/Source/OptimusEditor/Private/OptimusEditor.cpp
+++ b/Engine/Plugins/Animation/DeformerGraph/Source/OptimusEditor/Private/OptimusEditor.cpp
@@ -732,6 +732,12 @@ bool FOptimusEditor::CanConvertToFunction() const
 
 void FOptimusEditor::ConvertToSubGraph()
 {
+	FOptimusActionScope ActionScope(*GetActionStack(), TEXT("Convert to SubGraphs"));
+	UOptimusNodeGraph* ModelGraph = EditorGraph->GetModelGraph();
+	for (UOptimusNode* ModelNode: GetSelectedModelNodes())
+	{
+		ModelGraph->ConvertToSubGraph(ModelNode);
+	}
 }
 
 bool FOptimusEditor::CanConvertToSubGraph() const
