[Deformer Graph] Added support for expanding function node

diff --git a/Engine/Plugins/Animation/DeformerGraph/Source/OptimusCore/Private/Nodes/OptimusNode_FunctionReference.cpp b/Engine/Plugins/Animation/DeformerGraph/Source/OptimusCore/Private/Nodes/OptimusNode_FunctionReference.cpp
--- a/Engine/Plugins/Animation/DeformerGraph/Source/OptimusCore/Private/Nodes/OptimusNode_FunctionReference.cpp
+++ b/Engine/Plugins/Animation/DeformerGraph/Source/OptimusCore/Private/Nodes/OptimusNode_FunctionReference.cpp
@@ -106,6 +106,11 @@ UOptimusNodeGraph* UOptimusNode_FunctionReference::GetNodeGraphToShow()
 	return FunctionGraph.Get();
 }
 
+UOptimusNodeSubGraph* UOptimusNode_FunctionReference::GetReferencedSubGraph() const
+{
+	return FunctionGraph.Get();
+}
+
 UOptimusComponentSourceBinding* UOptimusNode_FunctionReference::GetDefaultComponentBinding(const FOptimusPinTraversalContext& InTraversalContext) const
 {
 	if (!ensure(DefaultComponentPin.IsValid()))
@@ -135,6 +140,11 @@ UOptimusComponentSourceBinding* UOptimusNode_FunctionReference::GetDefaultCompon
 	return nullptr;		
 }
 
+UOptimusNodePin* UOptimusNode_FunctionReference::GetDefaultComponentBindingPin() const
+{
+	return DefaultComponentPin.Get();
+}
+
 FSoftObjectPath UOptimusNode_FunctionReference::GetSerializedGraphPath() const
 {
 	return FunctionGraph.ToSoftObjectPath();
diff --git a/Engine/Plugins/Animation/DeformerGraph/Source/OptimusCore/Private/Nodes/OptimusNode_FunctionReference.h b/Engine/Plugins/Animation/DeformerGraph/Source/OptimusCore/Private/Nodes/OptimusNode_FunctionReference.h
--- a/Engine/Plugins/Animation/DeformerGraph/Source/OptimusCore/Private/Nodes/OptimusNode_FunctionReference.h
+++ b/Engine/Plugins/Animation/DeformerGraph/Source/OptimusCore/Private/Nodes/OptimusNode_FunctionReference.h
@@ -41,7 +41,9 @@ public:
 	UOptimusNodeGraph* GetNodeGraphToShow() override;
 	
 	// IOptimusNodeSubGraphReferencer
+	UOptimusNodeSubGraph* GetReferencedSubGraph() const override;
 	UOptimusComponentSourceBinding* GetDefaultComponentBinding(const FOptimusPinTraversalContext& InTraversalContext) const override;
+	UOptimusNodePin* GetDefaultComponentBindingPin() const override;
 
 	FSoftObjectPath GetSerializedGraphPath() const;
 
diff --git a/Engine/Plugins/Animation/DeformerGraph/Source/OptimusCore/Private/Nodes/OptimusNode_SubGraphReference.cpp b/Engine/Plugins/Animation/DeformerGraph/Source/OptimusCore/Private/Nodes/OptimusNode_SubGraphReference.cpp
--- a/Engine/Plugins/Animation/DeformerGraph/Source/OptimusCore/Private/Nodes/OptimusNode_SubGraphReference.cpp
+++ b/Engine/Plugins/Animation/DeformerGraph/Source/OptimusCore/Private/Nodes/OptimusNode_SubGraphReference.cpp
@@ -112,6 +112,11 @@ UOptimusNodeGraph* UOptimusNode_SubGraphReference::GetNodeGraphToShow()
 	return SubGraph.Get();
 }
 
+UOptimusNodeSubGraph* UOptimusNode_SubGraphReference::GetReferencedSubGraph() const
+{
+	return SubGraph.Get();
+}
+
 UOptimusComponentSourceBinding* UOptimusNode_SubGraphReference::GetDefaultComponentBinding(const FOptimusPinTraversalContext& InTraversalContext) const
 {
 	if (!ensure(DefaultComponentPin.IsValid()))
@@ -141,6 +146,11 @@ UOptimusComponentSourceBinding* UOptimusNode_SubGraphReference::GetDefaultCompon
 	return nullptr;	
 }
 
+UOptimusNodePin* UOptimusNode_SubGraphReference::GetDefaultComponentBindingPin() const
+{
+	return DefaultComponentPin.Get();
+}
+
 void UOptimusNode_SubGraphReference::SubscribeToSubGraph()
 {
 	if (ensure(!SubGraph->GetOnBindingArrayPasted().IsBoundToObject(this)))
diff --git a/Engine/Plugins/Animation/DeformerGraph/Source/OptimusCore/Private/Nodes/OptimusNode_SubGraphReference.h b/Engine/Plugins/Animation/DeformerGraph/Source/OptimusCore/Private/Nodes/OptimusNode_SubGraphReference.h
--- a/Engine/Plugins/Animation/DeformerGraph/Source/OptimusCore/Private/Nodes/OptimusNode_SubGraphReference.h
+++ b/Engine/Plugins/Animation/DeformerGraph/Source/OptimusCore/Private/Nodes/OptimusNode_SubGraphReference.h
@@ -44,7 +44,9 @@ public:
 	UOptimusNodeGraph* GetNodeGraphToShow() override;
 
 	// IOptimusNodeSubGraphReferencer
+	UOptimusNodeSubGraph* GetReferencedSubGraph() const override;
 	UOptimusComponentSourceBinding* GetDefaultComponentBinding(const FOptimusPinTraversalContext& InTraversalContext) const override;
+	UOptimusNodePin* GetDefaultComponentBindingPin() const override;
 
 protected:
 	friend class UOptimusNodeGraph;
diff --git a/Engine/Plugins/Animation/DeformerGraph/Source/OptimusCore/Private/OptimusNodeGraph.cpp b/Engine/Plugins/Animation/DeformerGraph/Source/OptimusCore/Private/OptimusNodeGraph.cpp
--- a/Engine/Plugins/Animation/DeformerGraph/Source/OptimusCore/Private/OptimusNodeGraph.cpp
+++ b/Engine/Plugins/Animation/DeformerGraph/Source/OptimusCore/Private/OptimusNodeGraph.cpp
@@ -1352,11 +1352,11 @@ UOptimusNode* UOptimusNodeGraph::CollapseNodesToSubGraph(
 
 
 TArray<UOptimusNode*> UOptimusNodeGraph::ExpandCollapsedNodes(
-	UOptimusNode* InFunctionNode
+	UOptimusNode* InGraphReferenceNode
 	)
 {
-	const bool bIsFunction = IsFunctionReference(InFunctionNode);
-	const bool bIsSubGraph = IsSubGraphReference(InFunctionNode);
+	const bool bIsFunction = IsFunctionReference(InGraphReferenceNode);
+	const bool bIsSubGraph = IsSubGraphReference(InGraphReferenceNode);
 	if (!bIsFunction && !bIsSubGraph)
 	{
 		return {};
@@ -1384,7 +1384,7 @@ TArray<UOptimusNode*> UOptimusNodeGraph::ExpandCollapsedNodes(
 	
 	TArray<UOptimusNodeLink*> ExternalLinks;	
 
-	for (const int32 LinkIndex : GetAllLinkIndexesToNode(InFunctionNode))
+	for (const int32 LinkIndex : GetAllLinkIndexesToNode(InGraphReferenceNode))
 	{
 		ExternalLinks.Add(Links[LinkIndex]);
 	}
@@ -1392,23 +1392,26 @@ TArray<UOptimusNode*> UOptimusNodeGraph::ExpandCollapsedNodes(
 	// Use to move all nodes in the outer graph away from the expansion location
 	// such that nodes currently in the subgraph would have enough space to spawn in the outer graph
 	
-	FVector2D ExpandLocation = InFunctionNode->GetGraphPosition();
+	FVector2D ExpandLocation = InGraphReferenceNode->GetGraphPosition();
 	
 	TArray<TPair<UOptimusNode*, FVector2D>> NodePositionInfos;
 	for (UOptimusNode* Node : Nodes)
 	{
-		if (Node == InFunctionNode)
+		if (Node == InGraphReferenceNode)
 		{
 			continue;
 		}
 
-		NodePositionInfos.Add({Node, Node->GetGraphPosition() - InFunctionNode->GetGraphPosition()});
+		NodePositionInfos.Add({Node, Node->GetGraphPosition() - InGraphReferenceNode->GetGraphPosition()});
 	}
 
 	
-	if (UOptimusNode_SubGraphReference* SubGraphReferenceNode = Cast<UOptimusNode_SubGraphReference>(InFunctionNode))
+	IOptimusNodeSubGraphReferencer* AsReferencerNode = Cast<IOptimusNodeSubGraphReferencer>(InGraphReferenceNode);
+	IOptimusNodePinRouter* AsRouterNode = Cast<IOptimusNodePinRouter>(InGraphReferenceNode);
+
+	if (AsReferencerNode && AsRouterNode)
 	{
-		UOptimusNodeSubGraph* SubGraph = SubGraphReferenceNode->SubGraph.Get();
+		UOptimusNodeSubGraph* SubGraph = AsReferencerNode->GetReferencedSubGraph();
 		
 		TArray<UOptimusNode*> NodesToDuplicate = SubGraph->Nodes;
 		TArray<UOptimusNodePair*> NodePairsToDuplicate = SubGraph->NodePairs;
@@ -1456,9 +1459,9 @@ TArray<UOptimusNode*> UOptimusNodeGraph::ExpandCollapsedNodes(
 		TArray<TPair<FString, FString>> LinksToAdd;
 
 		// Record all links linking terminal pins and nodes in the outer graph, saving to LinksToAdd array
-		auto AddLink = [&LinksToAdd, GraphPath, SubGraphReferenceNode, NewNodeNameMap](UOptimusNodePin* InReferenceNodePin)
+		auto AddLink = [&LinksToAdd, GraphPath, AsRouterNode, NewNodeNameMap](UOptimusNodePin* InReferenceNodePin)
 		{
-			const FOptimusRoutedNodePin PinCounterpart = SubGraphReferenceNode->GetPinCounterpart(InReferenceNodePin, {});
+			const FOptimusRoutedNodePin PinCounterpart = AsRouterNode->GetPinCounterpart(InReferenceNodePin, {});
 
 			for (const UOptimusNodePin* ExternalPin : InReferenceNodePin->GetConnectedPins())
 			{
@@ -1481,13 +1484,13 @@ TArray<UOptimusNode*> UOptimusNodeGraph::ExpandCollapsedNodes(
 			}	
 		};
 		
-		for (UOptimusNodePin* Pin : SubGraphReferenceNode->GetPins())
+		for (UOptimusNodePin* Pin : InGraphReferenceNode->GetPins())
 		{
 			AddLink(Pin);
 		}
 		
 		// Unconnected component pins should be linked to outer graph explicitly if the subgraph default component pin is connected externally
-		TArray<UOptimusNodePin*> ExternalComponentPins = SubGraphReferenceNode->DefaultComponentPin->GetConnectedPins();
+		TArray<UOptimusNodePin*> ExternalComponentPins = AsReferencerNode->GetDefaultComponentBindingPin()->GetConnectedPins();
 		if (ExternalComponentPins.Num() > 0)
 		{
 			UOptimusNodePin* ExternalComponentPin = ExternalComponentPins[0];
@@ -1576,7 +1579,8 @@ TArray<UOptimusNode*> UOptimusNodeGraph::ExpandCollapsedNodes(
 
 		RemoveNodesToAction(Action, SubGraph->GetAllNodes());
 
-		RemoveNodesToAction(Action, {SubGraphReferenceNode});
+		// This will remove the subgraph if it is a subgraph node
+		RemoveNodesToAction(Action, {InGraphReferenceNode});
 
 		if (!GetActionStack()->RunAction(Action))
 		{
diff --git a/Engine/Plugins/Animation/DeformerGraph/Source/OptimusCore/Public/IOptimusNodeSubGraphReferencer.h b/Engine/Plugins/Animation/DeformerGraph/Source/OptimusCore/Public/IOptimusNodeSubGraphReferencer.h
--- a/Engine/Plugins/Animation/DeformerGraph/Source/OptimusCore/Public/IOptimusNodeSubGraphReferencer.h
+++ b/Engine/Plugins/Animation/DeformerGraph/Source/OptimusCore/Public/IOptimusNodeSubGraphReferencer.h
@@ -7,7 +7,9 @@
 #include "IOptimusNodeSubGraphReferencer.generated.h"
 
 class UOptimusNodeGraph;
+class UOptimusNodeSubGraph;
 class UOptimusComponentSourceBinding;
+class UOptimusNodePin;
 struct FOptimusPinTraversalContext;
 
 UINTERFACE()
@@ -25,6 +27,8 @@ class OPTIMUSCORE_API IOptimusNodeSubGraphReferencer
 	GENERATED_BODY()
 
 public:
+	virtual UOptimusNodeSubGraph* GetReferencedSubGraph() const = 0;
 	virtual UOptimusComponentSourceBinding* GetDefaultComponentBinding(const FOptimusPinTraversalContext& InTraversalContext) const = 0;
+	virtual UOptimusNodePin* GetDefaultComponentBindingPin() const = 0;
 };
 
diff --git a/Engine/Plugins/Animation/DeformerGraph/Source/OptimusCore/Public/OptimusNodeGraph.h b/Engine/Plugins/Animation/DeformerGraph/Source/OptimusCore/Public/OptimusNodeGraph.h
--- a/Engine/Plugins/Animation/DeformerGraph/Source/OptimusCore/Public/OptimusNodeGraph.h
+++ b/Engine/Plugins/Animation/DeformerGraph/Source/OptimusCore/Public/OptimusNodeGraph.h
@@ -268,7 +268,7 @@ public:
 	 *  sub-graph was expanded, the sub-graph is deleted.
 	 */
 	UFUNCTION(BlueprintCallable, Category = OptimusNodeGraph)
-	TArray<UOptimusNode *> ExpandCollapsedNodes(UOptimusNode* InFunctionNode);
+	TArray<UOptimusNode *> ExpandCollapsedNodes(UOptimusNode* InGraphReferenceNode);
 	
 	/** Take a subgraph node convert it to a function in-place
 	 */
