[Deformer Graph] Only rebuild halfedge at runtime if the index buffer grants CPU access

diff --git a/Engine/Plugins/Animation/DeformerGraph/Source/OptimusCore/Private/DataInterfaces/OptimusDataInterfaceHalfEdge.cpp b/Engine/Plugins/Animation/DeformerGraph/Source/OptimusCore/Private/DataInterfaces/OptimusDataInterfaceHalfEdge.cpp
--- a/Engine/Plugins/Animation/DeformerGraph/Source/OptimusCore/Private/DataInterfaces/OptimusDataInterfaceHalfEdge.cpp
+++ b/Engine/Plugins/Animation/DeformerGraph/Source/OptimusCore/Private/DataInterfaces/OptimusDataInterfaceHalfEdge.cpp
@@ -209,11 +209,25 @@ bool FOptimusHalfEdgeDataProviderProxy::IsValid(FValidationData const& InValidat
 	{
 		return false;
 	}
-	if (SkeletalMeshObject->GetSkeletalMeshRenderData().LODRenderData[SkeletalMeshObject->GetLOD()].RenderSections.Num() != InValidationData.NumInvocations)
+
+	FSkeletalMeshLODRenderData const* LodRenderData = &SkeletalMeshObject->GetSkeletalMeshRenderData().LODRenderData[SkeletalMeshObject->GetLOD()];	
+	
+	if (LodRenderData->RenderSections.Num() != InValidationData.NumInvocations)
 	{
 		return false;
 	}
 
+	// Invalid if there is no cooked buffer and run time generation isn't possible either
+	const FRawStaticIndexBuffer16or32Interface* IndexBuffer = LodRenderData->MultiSizeIndexContainer.GetIndexBuffer();
+	const FPositionVertexBuffer& VertexBuffer = LodRenderData->StaticVertexBuffers.PositionVertexBuffer;
+	if (!LodRenderData->HalfEdgeBuffer.IsInitialized())
+	{
+		if (!IndexBuffer->GetNeedsCPUAccess() || !VertexBuffer.GetAllowCPUAccess())
+		{
+			return false;
+		}
+	}
+
 	return true;
 }
 
