[Deformer Graph] fix SubGraph node duplication not working

diff --git a/Engine/Plugins/Animation/DeformerGraph/Source/OptimusCore/Private/Nodes/OptimusNode_SubGraphReference.cpp b/Engine/Plugins/Animation/DeformerGraph/Source/OptimusCore/Private/Nodes/OptimusNode_SubGraphReference.cpp
--- a/Engine/Plugins/Animation/DeformerGraph/Source/OptimusCore/Private/Nodes/OptimusNode_SubGraphReference.cpp
+++ b/Engine/Plugins/Animation/DeformerGraph/Source/OptimusCore/Private/Nodes/OptimusNode_SubGraphReference.cpp
@@ -151,6 +151,23 @@ UOptimusNodePin* UOptimusNode_SubGraphReference::GetDefaultComponentBindingPin()
 	return DefaultComponentPin.Get();
 }
 
+void UOptimusNode_SubGraphReference::SetSerializedSubGraphName(FName InSubGraphName)
+{
+	SubGraphName = InSubGraphName;
+	FString SubGraphPath = GetOwningGraph()->ConstructSubGraphPath(SubGraphName.ToString());
+	SubGraph = Cast<UOptimusNodeSubGraph>(GetOwningGraph()->GetPathResolver()->ResolveGraphPath(SubGraphPath));
+}
+
+void UOptimusNode_SubGraphReference::RefreshSerializedSubGraphName()
+{
+	SubGraphName = SubGraph->GetFName();
+}
+
+FName UOptimusNode_SubGraphReference::GetSerializedSubGraphName() const
+{
+	return SubGraphName;
+}
+
 void UOptimusNode_SubGraphReference::SubscribeToSubGraph()
 {
 	if (ensure(!SubGraph->GetOnBindingArrayPasted().IsBoundToObject(this)))
diff --git a/Engine/Plugins/Animation/DeformerGraph/Source/OptimusCore/Private/Nodes/OptimusNode_SubGraphReference.h b/Engine/Plugins/Animation/DeformerGraph/Source/OptimusCore/Private/Nodes/OptimusNode_SubGraphReference.h
--- a/Engine/Plugins/Animation/DeformerGraph/Source/OptimusCore/Private/Nodes/OptimusNode_SubGraphReference.h
+++ b/Engine/Plugins/Animation/DeformerGraph/Source/OptimusCore/Private/Nodes/OptimusNode_SubGraphReference.h
@@ -48,10 +48,12 @@ public:
 	UOptimusComponentSourceBinding* GetDefaultComponentBinding(const FOptimusPinTraversalContext& InTraversalContext) const override;
 	UOptimusNodePin* GetDefaultComponentBindingPin() const override;
 
+	// Only used during node creation, cannot be used to reference a different graph once node is constructed
+	void SetSerializedSubGraphName(FName InSubGraphName);
+	void RefreshSerializedSubGraphName();
+	FName GetSerializedSubGraphName() const;
+	
 protected:
-	friend class UOptimusNodeGraph;
-	friend class UOptimusDeformer;
-
 	void SubscribeToSubGraph();
 	void UnsubscribeFromSubGraph() const;
 
@@ -62,14 +64,18 @@ protected:
 	void SyncPinsToBindings(FName InBindingArrayPropertyName);
 
 	TArray<UOptimusNodePin*> GetBindingPinsByDirection(EOptimusNodePinDirection InDirection);
-	/** The graph that owns us. This contains all the necessary pin information to add on
-	 * the terminal node.
-	 */
-	UPROPERTY()
-	TWeakObjectPtr<UOptimusNodeSubGraph> SubGraph;
+
 
+	UPROPERTY()
+	FName SubGraphName;
+	
 	UPROPERTY()
 	TWeakObjectPtr<UOptimusNodePin> DefaultComponentPin;
 
-	
+private:
+	/** The graph that owns us. This contains all the necessary pin information to add on
+	 * the terminal node.
+	 */
+	UPROPERTY()
+	TWeakObjectPtr<UOptimusNodeSubGraph> SubGraph;	
 };
diff --git a/Engine/Plugins/Animation/DeformerGraph/Source/OptimusCore/Private/OptimusDeformer.cpp b/Engine/Plugins/Animation/DeformerGraph/Source/OptimusCore/Private/OptimusDeformer.cpp
--- a/Engine/Plugins/Animation/DeformerGraph/Source/OptimusCore/Private/OptimusDeformer.cpp
+++ b/Engine/Plugins/Animation/DeformerGraph/Source/OptimusCore/Private/OptimusDeformer.cpp
@@ -175,7 +175,7 @@ UOptimusNode* UOptimusDeformer::GetSubGraphReferenceNode(const UOptimusNodeSubGr
 	for (UOptimusNode* Node: AllSubGraphNodes)
 	{
 		const UOptimusNode_SubGraphReference* SubGraphReference = Cast<UOptimusNode_SubGraphReference>(Node);
-		if (SubGraphReference->SubGraph == InSubGraph)
+		if (SubGraphReference->GetReferencedSubGraph() == InSubGraph)
 		{
 			UsedNode = Node;
 			break;
diff --git a/Engine/Plugins/Animation/DeformerGraph/Source/OptimusCore/Private/OptimusHelpers.cpp b/Engine/Plugins/Animation/DeformerGraph/Source/OptimusCore/Private/OptimusHelpers.cpp
--- a/Engine/Plugins/Animation/DeformerGraph/Source/OptimusCore/Private/OptimusHelpers.cpp
+++ b/Engine/Plugins/Animation/DeformerGraph/Source/OptimusCore/Private/OptimusHelpers.cpp
@@ -20,6 +20,30 @@ FName Optimus::GetUniqueNameForScope(UObject* InScopeObj, FName InName)
 	return InName;
 }
 
+Optimus::FUniqueNameGenerator::FUniqueNameGenerator(UObject* InScopeObject)
+{
+	ScopeObject = InScopeObject;
+}
+
+FName Optimus::FUniqueNameGenerator::GetUniqueName(FName InName)
+{
+	FName Result = Optimus::GetUniqueNameForScope(ScopeObject, InName);
+	GenerateUniqueNameFromExistingNames(Result, GeneratedName);
+
+	// Result be usable at this point since the name number strictly increases.
+	// Only take the slow route if there is still a name collision for mysterious reasons
+	if (!ensure(StaticFindObjectFast(UObject::StaticClass(), ScopeObject, Result) == nullptr))
+	{
+		do
+		{
+			Result.SetNumber( Result.GetNumber() + 1);
+		} while (StaticFindObjectFast(UObject::StaticClass(), ScopeObject, Result) != nullptr || GeneratedName.Contains(Result));
+	}
+
+	GeneratedName.Add(Result);
+	return Result;
+}
+
 FName Optimus::GetSanitizedNameForHlsl(FName InName)
 {
 	// Sanitize the name
diff --git a/Engine/Plugins/Animation/DeformerGraph/Source/OptimusCore/Private/OptimusNode.cpp b/Engine/Plugins/Animation/DeformerGraph/Source/OptimusCore/Private/OptimusNode.cpp
--- a/Engine/Plugins/Animation/DeformerGraph/Source/OptimusCore/Private/OptimusNode.cpp
+++ b/Engine/Plugins/Animation/DeformerGraph/Source/OptimusCore/Private/OptimusNode.cpp
@@ -110,10 +110,8 @@ bool UOptimusNode::SetGraphPositionDirect(
 
 	GraphPosition = InPosition;
 
-	if (bSendNotifications)
-	{
-		Notify(EOptimusGraphNotifyType::NodePositionChanged);
-	}
+
+	Notify(EOptimusGraphNotifyType::NodePositionChanged);
 
 	return true;
 }
diff --git a/Engine/Plugins/Animation/DeformerGraph/Source/OptimusCore/Private/OptimusNodeGraph.cpp b/Engine/Plugins/Animation/DeformerGraph/Source/OptimusCore/Private/OptimusNodeGraph.cpp
--- a/Engine/Plugins/Animation/DeformerGraph/Source/OptimusCore/Private/OptimusNodeGraph.cpp
+++ b/Engine/Plugins/Animation/DeformerGraph/Source/OptimusCore/Private/OptimusNodeGraph.cpp
@@ -369,10 +369,11 @@ TArray<UOptimusNode*> UOptimusNodeGraph::AddNodePairInternal(const TSubclassOf<U
 	return AddedNodes;
 }
 
-void UOptimusNodeGraph::AddPairNodesToArray(
+void UOptimusNodeGraph::GatherRelatedObjects(
 	const TArray<UOptimusNode*>& InNodes,
 	TArray<UOptimusNode*>& OutNodes,
-	TArray<UOptimusNodePair*>& OutNodePairs
+	TArray<UOptimusNodePair*>& OutNodePairs,
+	TArray<UOptimusNodeSubGraph*>& OutSubGraphs
 	)
 {
 	OutNodes = InNodes;
@@ -386,9 +387,185 @@ void UOptimusNodeGraph::AddPairNodesToArray(
 				OutNodePairs.AddUnique(NodePair);
 			}
 		}
+
+		if (UOptimusNode_SubGraphReference* SubGraphNode = Cast<UOptimusNode_SubGraphReference>(Node))
+		{
+			OutSubGraphs.AddUnique(SubGraphNode->GetReferencedSubGraph());
+		}
 	}	
 }
 
+bool UOptimusNodeGraph::DuplicateSubGraph(UOptimusNodeSubGraph* InSourceSubGraph, FName InNewGraphName)
+{
+	IOptimusPathResolver* PathResolver = GetPathResolver();
+
+	// Ensure the name is unique here instead of relying on AddGraph action
+	// because subsequent actions needs to use the final graph path
+	// Note: we could change the actions to accept a lambda instead of specific value in the future
+	FName SubGraphName = InNewGraphName;
+	
+	const FString SubGraphPath = ConstructSubGraphPath(SubGraphName.ToString());
+	const FVector2D EntryNodePosition = InSourceSubGraph->EntryNode->GetGraphPosition();
+	const FVector2D ReturnNodePosition = InSourceSubGraph->ReturnNode->GetGraphPosition();
+	const FString CurrentGraphPath = GetGraphPath();
+
+	TMap<FString, TArray<FString>> InternalLinksToAdd;
+	
+	for (UOptimusNodeLink* Link : InSourceSubGraph->GetAllLinks())
+	{
+		if (Link->GetNodeInputPin()->GetOwningNode() == InSourceSubGraph->ReturnNode)
+		{
+			FString ReturnNodeInputPinPath = ConstructPath(
+				SubGraphPath,
+				UOptimusNode_GraphTerminal::ReturnNodeName.ToString(),
+				Link->GetNodeInputPin()->GetUniqueName().ToString()
+				);
+
+			FString InternalNodeOutputPinPath = ConstructPath(
+				SubGraphPath,
+				Link->GetNodeOutputPin()->GetOwningNode()->GetName(),
+				Link->GetNodeOutputPin()->GetUniqueName().ToString()
+				);
+
+			InternalLinksToAdd.FindOrAdd(InternalNodeOutputPinPath).Add(ReturnNodeInputPinPath);
+		}
+		else if (Link->GetNodeOutputPin()->GetOwningNode() == InSourceSubGraph->EntryNode)
+		{
+			FString EntryNodeOutputPinPath = ConstructPath(
+				SubGraphPath,
+				UOptimusNode_GraphTerminal::EntryNodeName.ToString(),
+				Link->GetNodeOutputPin()->GetUniqueName().ToString()
+				);
+
+			FString InternalNodeInputPinPath = ConstructPath(
+				SubGraphPath,
+				Link->GetNodeInputPin()->GetOwningNode()->GetName(),
+				Link->GetNodeInputPin()->GetUniqueName().ToString()
+				);
+			InternalLinksToAdd.FindOrAdd(EntryNodeOutputPinPath).Add(InternalNodeInputPinPath);
+		}
+		else
+		{
+			FString OutputPinPath = ConstructPath(
+					SubGraphPath,
+					Link->GetNodeOutputPin()->GetOwningNode()->GetName(),
+					Link->GetNodeOutputPin()->GetUniqueName().ToString()
+					);
+
+			FString InputPinPath = ConstructPath(
+				SubGraphPath,
+				Link->GetNodeInputPin()->GetOwningNode()->GetName(),
+				Link->GetNodeInputPin()->GetUniqueName().ToString()
+				);
+			
+			InternalLinksToAdd.FindOrAdd(OutputPinPath).Add(InputPinPath);
+		}
+	}	
+	
+	
+	TArray<UOptimusNode*> SubGraphNodes = InSourceSubGraph->GetAllNodes();
+	
+	SubGraphNodes.RemoveAll([](UOptimusNode* InNode)
+	{
+		return Cast<UOptimusNode_GraphTerminal>(InNode) != nullptr;
+	});
+
+	TArray<UOptimusNode*> NodesToDuplicate;
+	TArray<UOptimusNodePair*> NodePairsToDuplicate;
+	TArray<UOptimusNodeSubGraph*> SubGraphsToDuplicate;
+	GatherRelatedObjects(SubGraphNodes, NodesToDuplicate, NodePairsToDuplicate, SubGraphsToDuplicate);
+
+	TArray<TPair<FString, FString>> NewNodePairs;
+	for (const UOptimusNodePair* NodePair : NodePairsToDuplicate)
+	{
+		// Generate new pair
+		const UOptimusNode *FirstNode = NodePair->GetFirst();
+		const UOptimusNode *SecondNode = NodePair->GetSecond();
+
+		FString NewFirstNodePath = ConstructPath(SubGraphPath, FirstNode->GetName(), {});
+		FString NewSecondNodePath = ConstructPath(SubGraphPath, SecondNode->GetName(), {});
+
+		NewNodePairs.Add({NewFirstNodePath, NewSecondNodePath});
+	}
+
+	{
+		FOptimusActionScope(*GetActionStack(), TEXT("Duplicate SubGraph"));
+		GetActionStack()->RunAction<FOptimusNodeGraphAction_AddGraph>(this, EOptimusNodeGraphType::SubGraph, SubGraphName, INDEX_NONE,
+			[InSourceSubGraph](UOptimusNodeGraph* InGraph)
+			{
+				UOptimusNodeSubGraph* SubGraph = CastChecked<UOptimusNodeSubGraph>(InGraph);
+				SubGraph->InputBindings = InSourceSubGraph->InputBindings;
+				SubGraph->OutputBindings = InSourceSubGraph->OutputBindings;
+				return true;
+			});
+	
+		GetActionStack()->RunAction<FOptimusNodeGraphAction_AddNode>(
+			SubGraphPath, UOptimusNode_GraphTerminal::StaticClass(), UOptimusNode_GraphTerminal::EntryNodeName,
+			[SubGraphPath, EntryNodePosition, PathResolver](UOptimusNode* InNode)
+			{
+				UOptimusNode_GraphTerminal* EntryNode = Cast<UOptimusNode_GraphTerminal>(InNode); 
+				UOptimusNodeSubGraph* SubGraph = Cast<UOptimusNodeSubGraph>(PathResolver->ResolveGraphPath(SubGraphPath));
+
+				SubGraph->EntryNode = EntryNode;
+		
+				EntryNode->TerminalType = EOptimusTerminalType::Entry;
+				EntryNode->OwningGraph = SubGraph;
+				return EntryNode->SetGraphPositionDirect(EntryNodePosition);
+			});
+
+		GetActionStack()->RunAction<FOptimusNodeGraphAction_AddNode>(
+			SubGraphPath, UOptimusNode_GraphTerminal::StaticClass(), UOptimusNode_GraphTerminal::ReturnNodeName,
+			[SubGraphPath, ReturnNodePosition, PathResolver](UOptimusNode* InNode)
+			{
+				UOptimusNode_GraphTerminal* ReturnNode = Cast<UOptimusNode_GraphTerminal>(InNode); 
+				UOptimusNodeSubGraph* SubGraph = Cast<UOptimusNodeSubGraph>(PathResolver->ResolveGraphPath(SubGraphPath));
+
+				SubGraph->ReturnNode = ReturnNode;
+			
+				ReturnNode->TerminalType = EOptimusTerminalType::Return;
+				ReturnNode->OwningGraph = SubGraph;
+				return ReturnNode->SetGraphPositionDirect(ReturnNodePosition);
+			});
+
+		for (UOptimusNodeSubGraph* SubGraphToDuplicate : SubGraphsToDuplicate)
+		{
+			DuplicateSubGraph(SubGraphToDuplicate, SubGraphToDuplicate->GetFName());
+		}
+		
+		// Copy the nodes to the sub graph
+		for (UOptimusNode* Node: NodesToDuplicate)
+		{
+			GetActionStack()->RunAction<FOptimusNodeGraphAction_DuplicateNode>(
+				SubGraphPath, Node, Node->GetFName(), [](UOptimusNode* InNode)
+				{
+					if (UOptimusNode_SubGraphReference* SubGraphNode = Cast<UOptimusNode_SubGraphReference>(InNode))
+					{
+						FName NewSubGraphName = SubGraphNode->GetReferencedSubGraph()->GetFName();
+						SubGraphNode->SetSerializedSubGraphName(NewSubGraphName);
+					}
+					return true;
+				});
+		}
+
+		for (const TPair<FString, FString>& NewPair : NewNodePairs)
+		{
+			GetActionStack()->RunAction<FOptimusNodeGraphAction_AddNodePair>(NewPair.Key, NewPair.Value);
+		}	
+		
+		// Restore links
+		for (const TPair<FString, TArray<FString>>& InternalLinks : InternalLinksToAdd )
+		{
+			const FString& OutputPinPath = InternalLinks.Key;
+			for (const FString& InputPinPath : InternalLinks.Value)
+			{
+				GetActionStack()->RunAction<FOptimusNodeGraphAction_AddLink>(OutputPinPath,	InputPinPath);
+			}
+		}
+	}
+	
+	return true;
+}
+
 
 UOptimusNode* UOptimusNodeGraph::AddNode(
 	const TSubclassOf<UOptimusNode> InNodeClass, 
@@ -604,7 +781,8 @@ bool UOptimusNodeGraph::RemoveNodesToAction(
 
 	TArray<UOptimusNode*> NodesToRemove;
 	TArray<UOptimusNodePair*> NodePairsToRemove;
-	AddPairNodesToArray(InNodes, NodesToRemove, NodePairsToRemove);
+	TArray<UOptimusNodeSubGraph*> SubGraphsToRemove;
+	GatherRelatedObjects(InNodes, NodesToRemove, NodePairsToRemove, SubGraphsToRemove);
 	
 	for (UOptimusNode* Node : NodesToRemove)
 	{
@@ -614,18 +792,6 @@ bool UOptimusNodeGraph::RemoveNodesToAction(
 		}
 	}
 
-	TArray<UOptimusNodeSubGraph*> SubGraphsToRemove;
-	for (UOptimusNode* Node : NodesToRemove)
-	{
-		if (UOptimusNode_SubGraphReference* SubGraphReference = Cast<UOptimusNode_SubGraphReference>(Node))
-		{
-			if (ensure(SubGraphReference->SubGraph.IsValid()))
-			{
-				SubGraphsToRemove.Add(SubGraphReference->SubGraph.Get());
-			}
-		}
-	}
-	
 	TSet<int32> AllLinkIndexes;
 
 	// Get all unique links for all the given nodes and remove them *before* we remove the nodes.
@@ -667,7 +833,7 @@ UOptimusNode* UOptimusNodeGraph::DuplicateNode(
 		return nullptr;
 	}
 
-	// This API does not support duplicating loop terminal currently
+	// This API does not support duplicating loop terminals, or subgraph nodes currently
 	// Use DuplicateNodes() below
 	if (!ensure(!Cast<UOptimusNode_LoopTerminal>(InNode)))
 	{
@@ -728,27 +894,11 @@ bool UOptimusNodeGraph::DuplicateNodes(
 
 	TArray<UOptimusNode*> NodesToDuplicate;
 	TArray<UOptimusNodePair*> NodePairsToDuplicate;
-	AddPairNodesToArray(InNodes, NodesToDuplicate, NodePairsToDuplicate);
+	TArray<UOptimusNodeSubGraph*> SubGraphsToDuplicate;
+	GatherRelatedObjects(InNodes, NodesToDuplicate, NodePairsToDuplicate, SubGraphsToDuplicate);
 
 	// Figure out the non-clashing names to use, to avoid collisions during actual execution.
-	TSet<FName> ExistingObjects;
-	for (const UOptimusNode* Node: Nodes)
-	{
-		if (ensure(Node != nullptr))
-		{
-			ExistingObjects.Add(Node->GetFName());
-		}
-	}
-
-	auto MakeUniqueNodeName = [&ExistingObjects](FName InName)
-	{
-		while(ExistingObjects.Contains(InName))
-		{
-			InName.SetNumber(InName.GetNumber() + 1);
-		}
-		ExistingObjects.Add(InName);
-		return InName;
-	};
+	Optimus::FUniqueNameGenerator UniqueNameGenerator(this);
 
 	using FloatType = decltype(FVector2D::X);
 	FVector2D TopLeft{std::numeric_limits<FloatType>::max()};
@@ -756,7 +906,7 @@ bool UOptimusNodeGraph::DuplicateNodes(
 	for (UOptimusNode* Node: NodesToDuplicate)
 	{
 		TopLeft = FVector2D::Min(TopLeft, Node->GraphPosition);
-		NewNodeNameMap.Add(Node, MakeUniqueNodeName(Node->GetFName()));
+		NewNodeNameMap.Add(Node, UniqueNameGenerator.GetUniqueName(Node->GetFName()));
 	}
 	FVector2D NodeOffset = InPosition - TopLeft;
 
@@ -778,6 +928,12 @@ bool UOptimusNodeGraph::DuplicateNodes(
 		}	
 	}
 
+	// Using GraphName instead of the graph object as key here because this map is to be captured by a lambda to be used during undo/redo 
+	TMap<FName , FName> NewSubGraphNameMap;
+	for (UOptimusNodeSubGraph* SubGraphToDuplicate : SubGraphsToDuplicate)
+	{
+		NewSubGraphNameMap.Add(SubGraphToDuplicate->GetFName(), UniqueNameGenerator.GetUniqueName(SubGraphToDuplicate->GetFName()));
+	}	
 	
 	/// Collect the links between these existing nodes. 
 	TArray<TPair<FString, FString>> NodeLinks;
@@ -798,51 +954,69 @@ bool UOptimusNodeGraph::DuplicateNodes(
 		}
 	}
 
-	FOptimusCompoundAction *Action = new FOptimusCompoundAction;
+	FString ActionTitle; 
 	if (NodesToDuplicate.Num() == 1)
 	{
-		Action->SetTitlef(TEXT("%s Node"), *InActionName);
+		ActionTitle = FString::Printf(TEXT("%s Node"), *InActionName);
 	}
 	else
 	{
-		Action->SetTitlef(TEXT("%s %d Nodes"), *InActionName, NodesToDuplicate.Num());
+		ActionTitle = FString::Printf(TEXT("%s %d Nodes"), *InActionName, NodesToDuplicate.Num());
 	}
-
-	// Add all the pre-duplicate requirement actions first. This allows certain nodes to set up their
-	// operating environment correctly (e.g. a missing variable description for get variable nodes, etc.)
-	for (UOptimusNode* Node: NodesToDuplicate)
+	
 	{
-		Node->PreDuplicateRequirementActions(this, Action);
-	}
+		FOptimusActionScope ActionScope(*GetActionStack(), ActionTitle);
+		
+		for (UOptimusNodeSubGraph* SubGraphToDuplicate: SubGraphsToDuplicate)
+		{
+			DuplicateSubGraph(SubGraphToDuplicate, NewSubGraphNameMap[SubGraphToDuplicate->GetFName()]);
+		}
 
-	// Duplicate the nodes and place them correctly
-	for (UOptimusNode* Node: NodesToDuplicate)
-	{
-		FOptimusNodeGraphAction_DuplicateNode *DuplicateNodeAction = new FOptimusNodeGraphAction_DuplicateNode(
-			GetGraphPath(), Node, NewNodeNameMap[Node],
-			[Node, NodeOffset](UOptimusNode *InNode) {
-				return InNode->SetGraphPositionDirect(Node->GraphPosition + NodeOffset); 
-		});
+		{
+			FOptimusCompoundAction *PreDuplicateRequirementAction = new FOptimusCompoundAction;	
+			// Add all the pre-duplicate requirement actions first. This allows certain nodes to set up their
+			// operating environment correctly (e.g. a missing variable description for get variable nodes, etc.)
+			for (UOptimusNode* Node: NodesToDuplicate)
+			{
+				Node->PreDuplicateRequirementActions(this, PreDuplicateRequirementAction);
+			}
+			GetActionStack()->RunAction(PreDuplicateRequirementAction);
+		}
 		
-		Action->AddSubAction(DuplicateNodeAction);
-	}
+		// Duplicate the nodes and place them correctly
+		for (UOptimusNode* Node: NodesToDuplicate)
+		{
+			FOptimusNodeGraphAction_DuplicateNode *DuplicateNodeAction = new FOptimusNodeGraphAction_DuplicateNode(
+				GetGraphPath(), Node, NewNodeNameMap[Node],
+				[Node, NodeOffset, NewSubGraphNameMap](UOptimusNode *InNode) {
+					if (UOptimusNode_SubGraphReference* SubGraphNode = Cast<UOptimusNode_SubGraphReference>(InNode))
+					{
+						FName NewSubGraphName = NewSubGraphNameMap[SubGraphNode->GetSerializedSubGraphName()];
+						SubGraphNode->SetSerializedSubGraphName(NewSubGraphName);
+					}
+					return InNode->SetGraphPositionDirect(Node->GraphPosition + NodeOffset); 
+			});
+			
+			GetActionStack()->RunAction(DuplicateNodeAction);
+		}
 
-	for (const TPair<FString, FString>& NewPair : NewNodePairs)
-	{
-		Action->AddSubAction<FOptimusNodeGraphAction_AddNodePair>(NewPair.Key, NewPair.Value);
-	}
-	
-	// Add any links that the nodes may have had. These operations are allowed to fail, in which case
-	// we just end up with dangling nodes.
-	// In the future we would like to introduce connections that are in an error state (will fail compile)
-	// but show up as disconnectable wires in the graph (e.g. if a type no longer matches).
-	for (const TTuple<FString, FString>& LinkInfo: NodeLinks)
-	{
-		constexpr bool bCanFail = true;
-		Action->AddSubAction<FOptimusNodeGraphAction_AddLink>(LinkInfo.Key, LinkInfo.Value, bCanFail);
+		for (const TPair<FString, FString>& NewPair : NewNodePairs)
+		{
+			GetActionStack()->RunAction<FOptimusNodeGraphAction_AddNodePair>(NewPair.Key, NewPair.Value);
+		}
+		
+		// Add any links that the nodes may have had. These operations are allowed to fail, in which case
+		// we just end up with dangling nodes.
+		// In the future we would like to introduce connections that are in an error state (will fail compile)
+		// but show up as disconnectable wires in the graph (e.g. if a type no longer matches).
+		for (const TTuple<FString, FString>& LinkInfo: NodeLinks)
+		{
+			constexpr bool bCanFail = true;
+			GetActionStack()->RunAction<FOptimusNodeGraphAction_AddLink>(LinkInfo.Key, LinkInfo.Value, bCanFail);
+		}
 	}
 
-	return GetActionStack()->RunAction(Action);
+	return true;
 }
 
 
@@ -1111,7 +1285,8 @@ UOptimusNode* UOptimusNodeGraph::CollapseNodesToSubGraph(
 
 	TArray<UOptimusNode*> NodesToDuplicate;
 	TArray<UOptimusNodePair*> NodePairsToDuplicate;
-	AddPairNodesToArray(NodesToCollapse, NodesToDuplicate, NodePairsToDuplicate);
+	TArray<UOptimusNodeSubGraph*> SubGraphsToDuplicate;
+	GatherRelatedObjects(NodesToCollapse, NodesToDuplicate, NodePairsToDuplicate, SubGraphsToDuplicate);
 
 	const TSet<UOptimusNode*> NodeSet(NodesToDuplicate);
 
@@ -1231,12 +1406,25 @@ UOptimusNode* UOptimusNodeGraph::CollapseNodesToSubGraph(
 			ReturnNode->OwningGraph = SubGraph;
 			return ReturnNode->SetGraphPositionDirect({NodeBox.Max.X + 750.0f, NodeBox.GetCenter().Y});
 		});
+
+	for (UOptimusNodeSubGraph* SubGraphToDuplicate: SubGraphsToDuplicate)
+	{
+		DuplicateSubGraph(SubGraphToDuplicate, SubGraphToDuplicate->GetFName());
+	}
 	
 	// Duplicate the nodes into the graph.
 	for (UOptimusNode* Node: NodesToDuplicate)
 	{
 		Action->AddSubAction<FOptimusNodeGraphAction_DuplicateNode>(
-			SubGraphPath, Node, Node->GetFName(), [](UOptimusNode*) { return true; });
+			SubGraphPath, Node, Node->GetFName(), [](UOptimusNode* InNode)
+			{
+				if (UOptimusNode_SubGraphReference* SubGraphNode = Cast<UOptimusNode_SubGraphReference>(InNode))
+				{
+					FName NewSubGraphName = SubGraphNode->GetReferencedSubGraph()->GetFName();
+					SubGraphNode->SetSerializedSubGraphName(NewSubGraphName);
+				}	
+				return true;
+			});
 	}
 
 	TArray<TPair<FString, FString>> NewNodePairs;
@@ -1262,12 +1450,10 @@ UOptimusNode* UOptimusNodeGraph::CollapseNodesToSubGraph(
 	GraphNodeRefName = Optimus::GetUniqueNameForScope(this, GraphNodeRefName);
 	FOptimusNodeGraphAction_AddNode* AddSubGraphRefNodeAction = new FOptimusNodeGraphAction_AddNode(
 		GetGraphPath(), UOptimusNode_SubGraphReference::StaticClass(), GraphNodeRefName,
-		[NodeBox, SubGraphPath, PathResolver](UOptimusNode* InNode)
+		[NodeBox, SubGraphName](UOptimusNode* InNode)
 		{
 			UOptimusNode_SubGraphReference* SubGraphNode = Cast<UOptimusNode_SubGraphReference>(InNode); 
-			UOptimusNodeSubGraph* SubGraph = Cast<UOptimusNodeSubGraph>(PathResolver->ResolveGraphPath(SubGraphPath));
-			
-			SubGraphNode->SubGraph = SubGraph;
+			SubGraphNode->SetSerializedSubGraphName(SubGraphName);
 			return SubGraphNode->SetGraphPositionDirect(NodeBox.GetCenter());
 		});
 	
@@ -1608,7 +1794,7 @@ bool UOptimusNodeGraph::ConvertToFunction(UOptimusNode* InSubGraphNode)
 	}
 	
 	UOptimusNode_SubGraphReference* SubGraphNode = CastChecked<UOptimusNode_SubGraphReference>(InSubGraphNode);
-	UOptimusNodeSubGraph* SubGraph = SubGraphNode->SubGraph.Get();
+	UOptimusNodeSubGraph* SubGraph = SubGraphNode->GetReferencedSubGraph();
 
 	UOptimusDeformer* Deformer = Cast<UOptimusDeformer>(GetCollectionRoot());
 	if (!ensure(Deformer))
@@ -1723,7 +1909,8 @@ bool UOptimusNodeGraph::ConvertToFunction(UOptimusNode* InSubGraphNode)
 
 	TArray<UOptimusNode*> NodesToDuplicate;
 	TArray<UOptimusNodePair*> NodePairsToDuplicate;
-	AddPairNodesToArray(SubGraphNodes, NodesToDuplicate, NodePairsToDuplicate);
+	TArray<UOptimusNodeSubGraph* > SubGraphsToDuplicate;
+	GatherRelatedObjects(SubGraphNodes, NodesToDuplicate, NodePairsToDuplicate, SubGraphsToDuplicate);
 
 	TArray<TPair<FString, FString>> NewNodePairs;
 	for (const UOptimusNodePair* NodePair : NodePairsToDuplicate)
@@ -1778,12 +1965,24 @@ bool UOptimusNodeGraph::ConvertToFunction(UOptimusNode* InSubGraphNode)
 				return ReturnNode->SetGraphPositionDirect(ReturnNodePosition);
 			});
 
+		for (UOptimusNodeSubGraph* SubGraphToDuplicate : SubGraphsToDuplicate)
+		{
+			DuplicateSubGraph(SubGraphToDuplicate, SubGraphToDuplicate->GetFName());
+		}	
 		
 		// Copy the nodes to the function graph
 		for (UOptimusNode* Node: NodesToDuplicate)
 		{
 			GetActionStack()->RunAction<FOptimusNodeGraphAction_DuplicateNode>(
-				FunctionGraphPath, Node, Node->GetFName(), [](UOptimusNode*) { return true; });
+				FunctionGraphPath, Node, Node->GetFName(), [](UOptimusNode* InNode)
+				{
+					if (UOptimusNode_SubGraphReference* SubGraphNode = Cast<UOptimusNode_SubGraphReference>(InNode))
+					{
+						FName NewSubGraphName = SubGraphNode->GetReferencedSubGraph()->GetFName();
+						SubGraphNode->SetSerializedSubGraphName(NewSubGraphName);
+					}
+					return true;
+				});
 		}
 
 		for (const TPair<FString, FString>& NewPair : NewNodePairs)
@@ -1949,7 +2148,8 @@ bool UOptimusNodeGraph::ConvertToSubGraph(UOptimusNode* InFunctionNode)
 
 	TArray<UOptimusNode*> NodesToDuplicate;
 	TArray<UOptimusNodePair*> NodePairsToDuplicate;
-	AddPairNodesToArray(SubGraphNodes, NodesToDuplicate, NodePairsToDuplicate);
+	TArray<UOptimusNodeSubGraph*> SubGraphsToDuplicate;
+	GatherRelatedObjects(SubGraphNodes, NodesToDuplicate, NodePairsToDuplicate, SubGraphsToDuplicate);
 
 	TArray<TPair<FString, FString>> NewNodePairs;
 	for (const UOptimusNodePair* NodePair : NodePairsToDuplicate)
@@ -2003,12 +2203,24 @@ bool UOptimusNodeGraph::ConvertToSubGraph(UOptimusNode* InFunctionNode)
 				return ReturnNode->SetGraphPositionDirect(ReturnNodePosition);
 			});
 
+		for (UOptimusNodeSubGraph* SubGraphToDuplicate : SubGraphsToDuplicate)
+		{
+			DuplicateSubGraph(SubGraphToDuplicate, SubGraphToDuplicate->GetFName());
+		}
 		
 		// Copy the nodes to the sub graph
 		for (UOptimusNode* Node: NodesToDuplicate)
 		{
 			GetActionStack()->RunAction<FOptimusNodeGraphAction_DuplicateNode>(
-				SubGraphPath, Node, Node->GetFName(), [](UOptimusNode*) { return true; });
+				SubGraphPath, Node, Node->GetFName(), [](UOptimusNode* InNode)
+				{
+					if (UOptimusNode_SubGraphReference* SubGraphNode = Cast<UOptimusNode_SubGraphReference>(InNode))
+					{
+						FName NewSubGraphName = SubGraphNode->GetReferencedSubGraph()->GetFName();
+						SubGraphNode->SetSerializedSubGraphName(NewSubGraphName);
+					}
+					return true;
+				});
 		}
 
 		for (const TPair<FString, FString>& NewPair : NewNodePairs)
@@ -2029,11 +2241,10 @@ bool UOptimusNodeGraph::ConvertToSubGraph(UOptimusNode* InFunctionNode)
 		// Create the sub graph ref node
 		GetActionStack()->RunAction<FOptimusNodeGraphAction_AddNode>(
 			GetGraphPath(), UOptimusNode_SubGraphReference::StaticClass(), SubGraphNodeName,
-			[ReferenceNodePosition, SubGraphPath, PathResolver](UOptimusNode* InNode)
+			[ReferenceNodePosition, SubGraphName](UOptimusNode* InNode)
 			{
-				UOptimusNode_SubGraphReference* SubGraphNode = Cast<UOptimusNode_SubGraphReference>(InNode); 
-				UOptimusNodeSubGraph* SubGraph = Cast<UOptimusNodeSubGraph>(PathResolver->ResolveGraphPath(SubGraphPath));
-				SubGraphNode->SubGraph = SubGraph;
+				UOptimusNode_SubGraphReference* SubGraphNode = Cast<UOptimusNode_SubGraphReference>(InNode);
+				SubGraphNode->SetSerializedSubGraphName(SubGraphName);
 				
 				return SubGraphNode->SetGraphPositionDirect(ReferenceNodePosition);
 			});
@@ -3113,8 +3324,9 @@ bool UOptimusNodeGraph::RenameGraphDirect(UOptimusNodeGraph* InGraph, const FStr
 			{
 				if (UOptimusNode_SubGraphReference* SubGraphNode = Cast<UOptimusNode_SubGraphReference>(Node))
 				{
-					if (SubGraphNode->SubGraph == InGraph)
+					if (SubGraphNode->GetReferencedSubGraph() == InGraph)
 					{
+						SubGraphNode->RefreshSerializedSubGraphName();
 						Notify(EOptimusGraphNotifyType::NodeDisplayNameChanged, SubGraphNode);
 						break;
 					}
diff --git a/Engine/Plugins/Animation/DeformerGraph/Source/OptimusCore/Public/OptimusHelpers.h b/Engine/Plugins/Animation/DeformerGraph/Source/OptimusCore/Public/OptimusHelpers.h
--- a/Engine/Plugins/Animation/DeformerGraph/Source/OptimusCore/Public/OptimusHelpers.h
+++ b/Engine/Plugins/Animation/DeformerGraph/Source/OptimusCore/Public/OptimusHelpers.h
@@ -45,6 +45,21 @@ namespace Optimus
 	    If the name is already unique, it will be returned unchanged. */
 	FName GetUniqueNameForScope(UObject *InScopeObj, FName InName);
 
+	/** Given an object scope, generate names that is not only unique within the scope,
+	 * but also unique among all names generated before.
+	 * Note: It uses GetUniqueNameForScope so using the same scope object and
+	 * the same input name does not guarantee the same output name
+	 */
+	struct FUniqueNameGenerator
+	{
+		explicit FUniqueNameGenerator(UObject* InScopeObject);
+		
+		UObject* ScopeObject = nullptr;
+		TArray<FName> GeneratedName;
+
+		FName GetUniqueName(FName InName);
+	};
+	
 	/** A small helper class to enable binary reads on an archive, since the 
 		FObjectReader::Serialize(TArray<uint8>& InBytes) constructor is protected */
 	class FBinaryObjectReader : public FObjectReader
diff --git a/Engine/Plugins/Animation/DeformerGraph/Source/OptimusCore/Public/OptimusNodeGraph.h b/Engine/Plugins/Animation/DeformerGraph/Source/OptimusCore/Public/OptimusNodeGraph.h
--- a/Engine/Plugins/Animation/DeformerGraph/Source/OptimusCore/Public/OptimusNodeGraph.h
+++ b/Engine/Plugins/Animation/DeformerGraph/Source/OptimusCore/Public/OptimusNodeGraph.h
@@ -442,6 +442,7 @@ protected:
 	friend class UOptimusNode;
 	friend class UOptimusNodePin;
 	friend class UOptimusNode_ConstantValue;
+	friend class UOptimusNode_SubGraphReference;
 	friend class FOptimusEditorClipboard;
 	friend struct FOptimusNodeGraphAction_AddNode;
 	friend struct FOptimusNodeGraphAction_DuplicateNode;
@@ -524,7 +525,14 @@ private:
 		TFunction<void(UOptimusNode*)> InSecondNodeConfigFunc
 	);
 
-	static void AddPairNodesToArray(const TArray<UOptimusNode*>& InNodes, TArray<UOptimusNode*>& OutNodes, TArray<UOptimusNodePair*>& OutNodePairs);
+	static void GatherRelatedObjects(
+		const TArray<UOptimusNode*>& InNodes,
+		TArray<UOptimusNode*>& OutNodes,
+		TArray<UOptimusNodePair*>& OutNodePairs,
+		TArray<UOptimusNodeSubGraph*>& OutSubGraphs
+		);
+
+	bool DuplicateSubGraph(UOptimusNodeSubGraph* InSourceSubGraph, FName InNewGraphName);
 
 	void RemoveNodePairByIndex(int32 NodePairIndex);
 	
diff --git a/Engine/Plugins/Animation/DeformerGraph/Source/OptimusEditor/Private/OptimusEditorClipboard.cpp b/Engine/Plugins/Animation/DeformerGraph/Source/OptimusEditor/Private/OptimusEditorClipboard.cpp
--- a/Engine/Plugins/Animation/DeformerGraph/Source/OptimusEditor/Private/OptimusEditorClipboard.cpp
+++ b/Engine/Plugins/Animation/DeformerGraph/Source/OptimusEditor/Private/OptimusEditorClipboard.cpp
@@ -4,6 +4,7 @@
 
 #include "OptimusNode.h"
 #include "OptimusNodeGraph.h"
+#include "OptimusNodeSubGraph.h"
 #include "OptimusNodePair.h"
 #include "OptimusNodeLink.h"
 #include "OptimusNodePin.h"
@@ -25,10 +26,11 @@ void FOptimusEditorClipboard::SetClipboardFromNodes(const TArray<UOptimusNode*>&
 
 	TArray<UOptimusNode*> NodesToConsider;
 	TArray<UOptimusNodePair*> NodePairsToConsider;
+	TArray<UOptimusNodeSubGraph*> SubGraphsToConsider;
 	
 	UOptimusNodeGraph* Graph = InNodes[0]->GetOwningGraph();
 
-	Graph->AddPairNodesToArray(InNodes, NodesToConsider, NodePairsToConsider);
+	Graph->GatherRelatedObjects(InNodes, NodesToConsider, NodePairsToConsider, SubGraphsToConsider);
 
 	// Export the clipboard to text.
 	const FExportObjectInnerContext Context;
