[Deformer Graph] fix compilation error

diff --git a/Engine/Plugins/Animation/DeformerGraph/Source/OptimusCore/Public/OptimusNodeGraph.h b/Engine/Plugins/Animation/DeformerGraph/Source/OptimusCore/Public/OptimusNodeGraph.h
--- a/Engine/Plugins/Animation/DeformerGraph/Source/OptimusCore/Public/OptimusNodeGraph.h
+++ b/Engine/Plugins/Animation/DeformerGraph/Source/OptimusCore/Public/OptimusNodeGraph.h
@@ -26,6 +26,7 @@ class UOptimusActionStack;
 class UOptimusNode;
 class UOptimusNodePair;
 class UOptimusNodeGraph;
+class UOptimusNodeSubGraph;
 class UOptimusNodeLink;
 class UOptimusNodePin;
 enum class EOptimusNodePinDirection : uint8;
