[Deformer Graph] Fix adding function node of a unload graph failing

diff --git a/Engine/Plugins/Animation/DeformerGraph/Source/OptimusCore/Private/Nodes/OptimusNode_FunctionReference.h b/Engine/Plugins/Animation/DeformerGraph/Source/OptimusCore/Private/Nodes/OptimusNode_FunctionReference.h
--- a/Engine/Plugins/Animation/DeformerGraph/Source/OptimusCore/Private/Nodes/OptimusNode_FunctionReference.h
+++ b/Engine/Plugins/Animation/DeformerGraph/Source/OptimusCore/Private/Nodes/OptimusNode_FunctionReference.h
@@ -50,9 +50,6 @@ public:
 	void SetSerializedGraphPath(const FSoftObjectPath& InNewGraphPath);
 	
 protected:
-	friend class UOptimusNodeGraph;
-	friend class FOptimusCoreModule;
-
 	
 	/** The graph that owns us. This contains all the necessary pin information to add on
 	 * the terminal node.
diff --git a/Engine/Plugins/Animation/DeformerGraph/Source/OptimusCore/Private/OptimusNodeGraph.cpp b/Engine/Plugins/Animation/DeformerGraph/Source/OptimusCore/Private/OptimusNodeGraph.cpp
--- a/Engine/Plugins/Animation/DeformerGraph/Source/OptimusCore/Private/OptimusNodeGraph.cpp
+++ b/Engine/Plugins/Animation/DeformerGraph/Source/OptimusCore/Private/OptimusNodeGraph.cpp
@@ -657,7 +657,7 @@ UOptimusNode* UOptimusNodeGraph::AddFunctionReferenceNode(TSoftObjectPtr<UOptimu
 	[InFunctionGraph](UOptimusNode *InNode)
 	{
 		UOptimusNode_FunctionReference* FunctionNode = CastChecked<UOptimusNode_FunctionReference>(InNode);
-		FunctionNode->FunctionGraph = InFunctionGraph;
+		FunctionNode->SetSerializedGraphPath(InFunctionGraph.ToSoftObjectPath());
 	});
 }
 
@@ -2007,7 +2007,7 @@ bool UOptimusNodeGraph::ConvertToFunction(UOptimusNode* InSubGraphNode)
 			{
 				UOptimusNode_FunctionReference* FunctionNode = Cast<UOptimusNode_FunctionReference>(InNode); 
 				UOptimusFunctionNodeGraph* FunctionGraph = Cast<UOptimusFunctionNodeGraph>(PathResolver->ResolveGraphPath(FunctionGraphPath));
-				FunctionNode->FunctionGraph = FunctionGraph;
+				FunctionNode->SetSerializedGraphPath(FunctionGraph);
 				
 				return FunctionNode->SetGraphPositionDirect(ReferenceNodePosition);
 			});
@@ -2039,7 +2039,12 @@ bool UOptimusNodeGraph::ConvertToSubGraph(UOptimusNode* InFunctionNode)
 	}
 	
 	UOptimusNode_FunctionReference* FunctionNode = CastChecked<UOptimusNode_FunctionReference>(InFunctionNode);
-	UOptimusFunctionNodeGraph* FunctionGraph = FunctionNode->FunctionGraph.Get();
+	UOptimusFunctionNodeGraph* FunctionGraph = Cast<UOptimusFunctionNodeGraph>(FunctionNode->GetReferencedSubGraph());
+
+	if (!FunctionGraph)
+	{
+		return false;
+	}
 
 	IOptimusPathResolver* PathResolver = GetPathResolver();
 
