[Deformer Graph] fixed a bug where private functions are not listed in the graph context menu

diff --git a/Engine/Plugins/Animation/DeformerGraph/Source/OptimusCore/Private/OptimusDeformer.cpp b/Engine/Plugins/Animation/DeformerGraph/Source/OptimusCore/Private/OptimusDeformer.cpp
--- a/Engine/Plugins/Animation/DeformerGraph/Source/OptimusCore/Private/OptimusDeformer.cpp
+++ b/Engine/Plugins/Animation/DeformerGraph/Source/OptimusCore/Private/OptimusDeformer.cpp
@@ -184,6 +184,27 @@ UOptimusNode* UOptimusDeformer::GetSubGraphReferenceNode(const UOptimusNodeSubGr
 	return UsedNode;
 }
 
+TArray<UOptimusFunctionNodeGraph*> UOptimusDeformer::GetFunctionGraphs(FName InAccessSpecifier) const
+{
+	TArray<UOptimusFunctionNodeGraph*> FunctionGraphs;
+	for (UOptimusNodeGraph* Graph : Graphs)
+	{
+		if (UOptimusFunctionNodeGraph* FunctionNodeGraph = Cast<UOptimusFunctionNodeGraph>(Graph))
+		{
+			if (InAccessSpecifier.IsNone())
+			{
+				FunctionGraphs.Add(FunctionNodeGraph);
+			}
+			else if (FunctionNodeGraph->AccessSpecifier == InAccessSpecifier)
+			{
+				FunctionGraphs.Add(FunctionNodeGraph);
+			}
+		}
+	}
+
+	return FunctionGraphs;
+}
+
 
 UOptimusVariableDescription* UOptimusDeformer::AddVariable(
 	FOptimusDataTypeRef InDataTypeRef, 
@@ -3456,15 +3477,9 @@ void UOptimusDeformer::GetAssetRegistryTags(FAssetRegistryTagsContext Context) c
 	// Expose all the public functions
 	{
 		FOptimusFunctionNodeGraphHeaderArray PublicFunctionHeadersArray;
-		for (UOptimusNodeGraph* Graph : Graphs)
+		for (const UOptimusFunctionNodeGraph* FunctionNodeGraph: GetFunctionGraphs(UOptimusFunctionNodeGraph::AccessSpecifierPublicName))
 		{
-			if (UOptimusFunctionNodeGraph* FunctionNodeGraph = Cast<UOptimusFunctionNodeGraph>(Graph))
-			{
-				if (FunctionNodeGraph->AccessSpecifier == UOptimusFunctionNodeGraph::AccessSpecifierPublicName)
-				{
-					PublicFunctionHeadersArray.Headers.Add(FunctionNodeGraph->GetHeader());
-				}
-			}
+			PublicFunctionHeadersArray.Headers.Add(FunctionNodeGraph->GetHeader());
 		}
 
 	
@@ -3659,7 +3674,7 @@ UOptimusNodeGraph* UOptimusDeformer::CreateGraphDirect(
 		{
 			return nullptr;
 		}
-		// Not fully implemented yet.
+		
 		GraphClass = UOptimusFunctionNodeGraph::StaticClass();
 
 		// If there's already an object with this name, then attempt to make the name unique.
diff --git a/Engine/Plugins/Animation/DeformerGraph/Source/OptimusCore/Public/OptimusDeformer.h b/Engine/Plugins/Animation/DeformerGraph/Source/OptimusCore/Public/OptimusDeformer.h
--- a/Engine/Plugins/Animation/DeformerGraph/Source/OptimusCore/Public/OptimusDeformer.h
+++ b/Engine/Plugins/Animation/DeformerGraph/Source/OptimusCore/Public/OptimusDeformer.h
@@ -168,6 +168,9 @@ public:
 	UOptimusNode* GetSubGraphReferenceNode(
 		const UOptimusNodeSubGraph* InSubGraph
 		) const;
+
+	/// Returns all function graphs with the given access specifier. If InAccessSpecifier is None, it performs no filtering
+	TArray<UOptimusFunctionNodeGraph*> GetFunctionGraphs(FName InAccessSpecifier = NAME_None) const;
 	
 	// Variables
 	UOptimusVariableDescription* AddVariable(
diff --git a/Engine/Plugins/Animation/DeformerGraph/Source/OptimusEditor/Private/OptimusEditorGraphSchema.cpp b/Engine/Plugins/Animation/DeformerGraph/Source/OptimusEditor/Private/OptimusEditorGraphSchema.cpp
--- a/Engine/Plugins/Animation/DeformerGraph/Source/OptimusEditor/Private/OptimusEditorGraphSchema.cpp
+++ b/Engine/Plugins/Animation/DeformerGraph/Source/OptimusEditor/Private/OptimusEditorGraphSchema.cpp
@@ -166,43 +166,75 @@ void UOptimusEditorGraphSchema::GetGraphActions(
 		IoActionBuilder.AddAction(Action);
 	}
 
-	FAssetRegistryModule& AssetRegistryModule = FModuleManager::LoadModuleChecked<FAssetRegistryModule>("AssetRegistry");
-
-	TArray<FAssetData> AssetDatas;
+	// Private functions
+	{
+		if (const UOptimusEditorGraph* Graph = Cast<UOptimusEditorGraph>(InGraph))
+		{
+			if (UOptimusNodeGraph* ModelGraph = Graph->GetModelGraph())
+			{
+				if (UOptimusDeformer* Deformer = Cast<UOptimusDeformer>(ModelGraph->GetCollectionRoot()))
+				{
+					for (UOptimusFunctionNodeGraph* FunctionNodeGraph : Deformer->GetFunctionGraphs(UOptimusFunctionNodeGraph::AccessSpecifierPrivateName))
+					{
+						FOptimusFunctionNodeGraphHeader Header = FunctionNodeGraph->GetHeader();
+				
+						TSharedPtr< FOptimusGraphSchemaAction_NewFunctionReferenceNode> Action(
+									new FOptimusGraphSchemaAction_NewFunctionReferenceNode(
+										FText::FromName(Header.Category),
+										FText::FromName(Header.FunctionName),
+										/* Tooltip */{}, 0, /* Keywords */{}
+								));
+
+						Action->GraphPath = Header.GraphPath;
+
+						IoActionBuilder.AddAction(Action);
+					}
+				}	
+			}
+		}
+	}
+	
+	// Public functions
+	{
+		FAssetRegistryModule& AssetRegistryModule = FModuleManager::LoadModuleChecked<FAssetRegistryModule>("AssetRegistry");
 
-	// This triggers the gathering of asset tags across assets
-	AssetRegistryModule.Get().GetAssetsByClass(UOptimusDeformer::StaticClass()->GetClassPathName(), AssetDatas);
+		TArray<FAssetData> AssetDatas;
 
-	for (FAssetData& AssetData : AssetDatas)
-	{
-		FString PublicFunctionString = AssetData.GetTagValueRef<FString>(UOptimusDeformer::PublicFunctionsAssetTagName);
+		// This triggers the gathering of asset tags across assets
+		AssetRegistryModule.Get().GetAssetsByClass(UOptimusDeformer::StaticClass()->GetClassPathName(), AssetDatas);
 
-		if (!PublicFunctionString.IsEmpty())
+		for (FAssetData& AssetData : AssetDatas)
 		{
-			FOptimusFunctionNodeGraphHeaderArray PublicFunctionHeaderArray;
-
-			FOptimusFunctionNodeGraphHeaderArray::StaticStruct()->ImportText(*PublicFunctionString, &PublicFunctionHeaderArray, nullptr, PPF_None, nullptr, {});
+			FString PublicFunctionString = AssetData.GetTagValueRef<FString>(UOptimusDeformer::PublicFunctionsAssetTagName);
 
-			for (const FOptimusFunctionNodeGraphHeader& Header : PublicFunctionHeaderArray.Headers)
+			if (!PublicFunctionString.IsEmpty())
 			{
-				const FText FunctionName = FText::FromName(Header.FunctionName);
+				FOptimusFunctionNodeGraphHeaderArray PublicFunctionHeaderArray;
 
-				const FText Category = FText::FromName(Header.Category);
+				FOptimusFunctionNodeGraphHeaderArray::StaticStruct()->ImportText(*PublicFunctionString, &PublicFunctionHeaderArray, nullptr, PPF_None, nullptr, {});
 
-				TSharedPtr< FOptimusGraphSchemaAction_NewFunctionReferenceNode> Action(
-					new FOptimusGraphSchemaAction_NewFunctionReferenceNode(
-						Category,
-						FunctionName,
-						/* Tooltip */{}, 0, /* Keywords */{}
-				));
+				for (const FOptimusFunctionNodeGraphHeader& Header : PublicFunctionHeaderArray.Headers)
+				{
+					const FText FunctionName = FText::FromName(Header.FunctionName);
+
+					const FText Category = FText::FromName(Header.Category);
 
-				Action->GraphPath = Header.GraphPath;
+					TSharedPtr< FOptimusGraphSchemaAction_NewFunctionReferenceNode> Action(
+						new FOptimusGraphSchemaAction_NewFunctionReferenceNode(
+							Category,
+							FunctionName,
+							/* Tooltip */{}, 0, /* Keywords */{}
+					));
 
-				IoActionBuilder.AddAction(Action);	
+					Action->GraphPath = Header.GraphPath;
+
+					IoActionBuilder.AddAction(Action);	
+				}
 			}
-		}
+		}	
 	}
-
+	
+	// Loop nodes
 	{
 		TSharedPtr< FOptimusGraphSchemaAction_NewLoopTerminalNodes> Action(
 			new FOptimusGraphSchemaAction_NewLoopTerminalNodes(
diff --git a/Engine/Plugins/Animation/DeformerGraph/Source/OptimusEditor/Private/Widgets/SOptimusNodePalette.cpp b/Engine/Plugins/Animation/DeformerGraph/Source/OptimusEditor/Private/Widgets/SOptimusNodePalette.cpp
--- a/Engine/Plugins/Animation/DeformerGraph/Source/OptimusEditor/Private/Widgets/SOptimusNodePalette.cpp
+++ b/Engine/Plugins/Animation/DeformerGraph/Source/OptimusEditor/Private/Widgets/SOptimusNodePalette.cpp
@@ -86,9 +86,14 @@ FText SOptimusNodePaletteItem::GetItemTooltip() const
 SOptimusNodePalette::~SOptimusNodePalette()
 {
 	const TSharedPtr<FOptimusEditor> Editor = OwningEditor.Pin();
-	if (Editor && Editor->GetDeformer())
+	if (Editor)
 	{
-		Editor->GetDeformer()->GetNotifyDelegate().RemoveAll(this);
+		Editor->OnRefresh().RemoveAll(this);
+		
+		if (Editor->GetDeformer())
+		{
+			Editor->GetDeformer()->GetNotifyDelegate().RemoveAll(this);
+		}
 	}
 }
 
@@ -112,12 +117,22 @@ void SOptimusNodePalette::Construct(const FArguments& InArgs, TWeakPtr<FOptimusE
 	];
 
 	const TSharedPtr<FOptimusEditor> Editor = OwningEditor.Pin();
-	if (ensure(Editor) && Editor->GetDeformer())
+	if (ensure(Editor))
 	{
-		Editor->GetDeformer()->GetNotifyDelegate().AddSP(this, &SOptimusNodePalette::GraphCollectionNotify);
+		Editor->OnRefresh().AddRaw(this, &SOptimusNodePalette::Refresh);
+		
+		if (Editor->GetDeformer())
+		{
+			Editor->GetDeformer()->GetNotifyDelegate().AddSP(this, &SOptimusNodePalette::GraphCollectionNotify);
+		}
 	}
 }
 
+void SOptimusNodePalette::Refresh()
+{
+	GraphActionMenu->RefreshAllActions(true);
+}
+
 
 TSharedRef<SWidget> SOptimusNodePalette::OnCreateWidgetForAction(FCreateWidgetForActionData* const InCreateData)
 {
@@ -131,7 +146,7 @@ void SOptimusNodePalette::CollectAllActions(FGraphActionListBuilderBase& OutAllA
 	{
 		const UOptimusEditorGraphSchema* Schema = Cast<UOptimusEditorGraphSchema>(Editor->GetGraph()->GetSchema());
 
-		Schema->GetGraphActions(OutAllActions, nullptr, nullptr);
+		Schema->GetGraphActions(OutAllActions, nullptr, Editor->GetGraph());
 	}
 }
 
diff --git a/Engine/Plugins/Animation/DeformerGraph/Source/OptimusEditor/Private/Widgets/SOptimusNodePalette.h b/Engine/Plugins/Animation/DeformerGraph/Source/OptimusEditor/Private/Widgets/SOptimusNodePalette.h
--- a/Engine/Plugins/Animation/DeformerGraph/Source/OptimusEditor/Private/Widgets/SOptimusNodePalette.h
+++ b/Engine/Plugins/Animation/DeformerGraph/Source/OptimusEditor/Private/Widgets/SOptimusNodePalette.h
@@ -38,7 +38,10 @@ public:
 	SLATE_END_ARGS()
 
 	void Construct(const FArguments& InArgs, TWeakPtr<FOptimusEditor> InEditor);
-
+	
+	// Forcibly refresh
+	void Refresh();
+	
 protected:
 	// SGraphPalette overrides
 	TSharedRef<SWidget> OnCreateWidgetForAction(FCreateWidgetForActionData* const InCreateData) override;
