[Deformer Graph] Fixed shader compilation error when constant nodes using the same name but in different subgraph lead to function name collision in the generated shader code.

This fix makes sure every data interface function associated with a constant node also uses a unique index in its name. And uses a soft obj ptr instead of the value name to establish mapping from constant node to values in the graph data interface to avoid ambiguity

diff --git a/Engine/Plugins/Animation/DeformerGraph/Source/OptimusCore/Private/DataInterfaces/OptimusDataInterfaceGraph.cpp b/Engine/Plugins/Animation/DeformerGraph/Source/OptimusCore/Private/DataInterfaces/OptimusDataInterfaceGraph.cpp
--- a/Engine/Plugins/Animation/DeformerGraph/Source/OptimusCore/Private/DataInterfaces/OptimusDataInterfaceGraph.cpp
+++ b/Engine/Plugins/Animation/DeformerGraph/Source/OptimusCore/Private/DataInterfaces/OptimusDataInterfaceGraph.cpp
@@ -98,11 +98,11 @@ UComputeDataProvider* UOptimusGraphDataInterface::CreateDataProvider(TObjectPtr<
 }
 
 
-void UOptimusGraphDataProvider::SetConstant(FString const& InVariableName, TArray<uint8> const& InValue)
+void UOptimusGraphDataProvider::SetConstant(TSoftObjectPtr<UObject> InSourceObject, TArray<uint8> const& InValue)
 {
 	for (int32 VariableIndex = 0; VariableIndex < Variables.Num(); ++VariableIndex)
 	{
-		if (Variables[VariableIndex].Name == InVariableName)
+		if (Variables[VariableIndex].SourceObject == InSourceObject)
 		{
 			if (ensure(Variables[VariableIndex].Value.Num() == InValue.Num()))
 			{
diff --git a/Engine/Plugins/Animation/DeformerGraph/Source/OptimusCore/Private/DataInterfaces/OptimusDataInterfaceGraph.h b/Engine/Plugins/Animation/DeformerGraph/Source/OptimusCore/Private/DataInterfaces/OptimusDataInterfaceGraph.h
--- a/Engine/Plugins/Animation/DeformerGraph/Source/OptimusCore/Private/DataInterfaces/OptimusDataInterfaceGraph.h
+++ b/Engine/Plugins/Animation/DeformerGraph/Source/OptimusCore/Private/DataInterfaces/OptimusDataInterfaceGraph.h
@@ -28,6 +28,9 @@ struct FOptimusGraphVariableDescription
 
 	UPROPERTY()
 	int32 Offset = 0;
+
+	UPROPERTY()
+	TSoftObjectPtr<UObject> SourceObject;
 };
 
 /** Compute Framework Data Interface used for marshaling compute graph parameters and variables. */
@@ -73,8 +76,8 @@ public:
 
 	int32 ParameterBufferSize = 0;
 
-	void SetConstant(FString const& InVariableName, TArray<uint8> const& InValue);
-
+	void SetConstant(TSoftObjectPtr<UObject> InSourceObject, TArray<uint8> const& InValue);
+	
 	//~ Begin UComputeDataProvider Interface
 	FComputeDataProviderRenderProxy* GetRenderProxy() override;
 	//~ End UComputeDataProvider Interface
diff --git a/Engine/Plugins/Animation/DeformerGraph/Source/OptimusCore/Private/Nodes/OptimusNode_ComputeKernelBase.cpp b/Engine/Plugins/Animation/DeformerGraph/Source/OptimusCore/Private/Nodes/OptimusNode_ComputeKernelBase.cpp
--- a/Engine/Plugins/Animation/DeformerGraph/Source/OptimusCore/Private/Nodes/OptimusNode_ComputeKernelBase.cpp
+++ b/Engine/Plugins/Animation/DeformerGraph/Source/OptimusCore/Private/Nodes/OptimusNode_ComputeKernelBase.cpp
@@ -415,7 +415,7 @@ TOptional<FText> UOptimusNode_ComputeKernelBase::ProcessInputPinForComputeKernel
 				// Value nodes bind the single graph data interface.
 				DataInterfaceFuncIndex = InValueNodes.Find(OutputNode);
 				check(DataInterfaceFuncIndex != INDEX_NONE);
-				DataFunctionName = ValueProvider->GetValueName();	
+				DataFunctionName = Optimus::MakeUniqueValueName(ValueProvider->GetValueDebugName(), DataInterfaceFuncIndex);
 			}
 		}
 		else if (Cast<const UOptimusNode_LoopTerminal>(OutputNode))
diff --git a/Engine/Plugins/Animation/DeformerGraph/Source/OptimusCore/Private/Nodes/OptimusNode_ConstantValue.cpp b/Engine/Plugins/Animation/DeformerGraph/Source/OptimusCore/Private/Nodes/OptimusNode_ConstantValue.cpp
--- a/Engine/Plugins/Animation/DeformerGraph/Source/OptimusCore/Private/Nodes/OptimusNode_ConstantValue.cpp
+++ b/Engine/Plugins/Animation/DeformerGraph/Source/OptimusCore/Private/Nodes/OptimusNode_ConstantValue.cpp
@@ -101,7 +101,7 @@ void UOptimusNode_ConstantValue::PostEditChangeProperty(FPropertyChangedEvent& P
 #endif // WITH_EDITOR
 
 
-FString UOptimusNode_ConstantValue::GetValueName() const
+FString UOptimusNode_ConstantValue::GetValueDebugName() const
 {
 	return GetName();
 }
diff --git a/Engine/Plugins/Animation/DeformerGraph/Source/OptimusCore/Private/Nodes/OptimusNode_ConstantValue.h b/Engine/Plugins/Animation/DeformerGraph/Source/OptimusCore/Private/Nodes/OptimusNode_ConstantValue.h
--- a/Engine/Plugins/Animation/DeformerGraph/Source/OptimusCore/Private/Nodes/OptimusNode_ConstantValue.h
+++ b/Engine/Plugins/Animation/DeformerGraph/Source/OptimusCore/Private/Nodes/OptimusNode_ConstantValue.h
@@ -65,7 +65,7 @@ public:
 #endif
 
 	// IOptimusValueProvider overrides 
-	FString GetValueName() const override;
+	FString GetValueDebugName() const override;
 	FOptimusDataTypeRef GetValueType() const override;
 	FShaderValueType::FValue GetShaderValue() const override;
 
diff --git a/Engine/Plugins/Animation/DeformerGraph/Source/OptimusCore/Private/Nodes/OptimusNode_GetVariable.cpp b/Engine/Plugins/Animation/DeformerGraph/Source/OptimusCore/Private/Nodes/OptimusNode_GetVariable.cpp
--- a/Engine/Plugins/Animation/DeformerGraph/Source/OptimusCore/Private/Nodes/OptimusNode_GetVariable.cpp
+++ b/Engine/Plugins/Animation/DeformerGraph/Source/OptimusCore/Private/Nodes/OptimusNode_GetVariable.cpp
@@ -195,7 +195,7 @@ TOptional<FText> UOptimusNode_GetVariable::ValidateForCompile(const FOptimusPinT
 	return {};
 }
 
-FString UOptimusNode_GetVariable::GetValueName() const
+FString UOptimusNode_GetVariable::GetValueDebugName() const
 {
 	if (const UOptimusVariableDescription* Var = VariableDesc.Get())
 	{
diff --git a/Engine/Plugins/Animation/DeformerGraph/Source/OptimusCore/Private/Nodes/OptimusNode_GetVariable.h b/Engine/Plugins/Animation/DeformerGraph/Source/OptimusCore/Private/Nodes/OptimusNode_GetVariable.h
--- a/Engine/Plugins/Animation/DeformerGraph/Source/OptimusCore/Private/Nodes/OptimusNode_GetVariable.h
+++ b/Engine/Plugins/Animation/DeformerGraph/Source/OptimusCore/Private/Nodes/OptimusNode_GetVariable.h
@@ -54,7 +54,7 @@ public:
 	TOptional<FText> ValidateForCompile(const FOptimusPinTraversalContext& InContext) const override;
 	
 	// IOptimusValueProvider overrides 
-	FString GetValueName() const override;
+	FString GetValueDebugName() const override;
 	FOptimusDataTypeRef GetValueType() const override;
 	FShaderValueType::FValue GetShaderValue() const override;
 
diff --git a/Engine/Plugins/Animation/DeformerGraph/Source/OptimusCore/Private/OptimusDeformer.cpp b/Engine/Plugins/Animation/DeformerGraph/Source/OptimusCore/Private/OptimusDeformer.cpp
--- a/Engine/Plugins/Animation/DeformerGraph/Source/OptimusCore/Private/OptimusDeformer.cpp
+++ b/Engine/Plugins/Animation/DeformerGraph/Source/OptimusCore/Private/OptimusDeformer.cpp
@@ -2147,13 +2147,15 @@ TArray<FOptimusComputeGraphInfo> UOptimusDeformer::CompileNodeGraphToComputeGrap
 
 	TArray<FOptimusGraphVariableDescription> ValueNodeDescriptions;
 	ValueNodeDescriptions.Reserve(ValueNodes.Num());
-	for (UOptimusNode const* ValueNode : ValueNodes)
+	for (int32 ValueNodeIndex = 0 ; ValueNodeIndex < ValueNodes.Num(); ValueNodeIndex++)
 	{
+		UOptimusNode const* ValueNode = ValueNodes[ValueNodeIndex];
 		if (IOptimusValueProvider const* ValueProvider = Cast<const IOptimusValueProvider>(ValueNode))
 		{
 			FOptimusGraphVariableDescription& ValueNodeDescription = ValueNodeDescriptions.AddDefaulted_GetRef();
-			ValueNodeDescription.Name = ValueProvider->GetValueName();
+			ValueNodeDescription.Name = Optimus::MakeUniqueValueName(ValueProvider->GetValueDebugName(), ValueNodeIndex);
 			ValueNodeDescription.ValueType = ValueProvider->GetValueType()->ShaderValueType;
+			ValueNodeDescription.SourceObject = ValueNode;
 
 			if (UOptimusNode_ConstantValue const* ConstantNode = Cast<const UOptimusNode_ConstantValue>(ValueNode))
 			{
@@ -3027,7 +3029,7 @@ void UOptimusDeformer::Notify(EOptimusGlobalNotifyType InNotifyType, UObject* In
 	case EOptimusGlobalNotifyType::ConstantValueChanged:
 		if (UOptimusNode_ConstantValue* ConstantValue = Cast<UOptimusNode_ConstantValue>(InObject))
 		{
-			ConstantValueUpdateDelegate.Broadcast(ConstantValue->GetValueName(), ConstantValue->GetShaderValue().ShaderValue);
+			ConstantValueUpdateDelegate.Broadcast(ConstantValue, ConstantValue->GetShaderValue().ShaderValue);
 		}
 		
 		break;
diff --git a/Engine/Plugins/Animation/DeformerGraph/Source/OptimusCore/Private/OptimusDeformerInstance.cpp b/Engine/Plugins/Animation/DeformerGraph/Source/OptimusCore/Private/OptimusDeformerInstance.cpp
--- a/Engine/Plugins/Animation/DeformerGraph/Source/OptimusCore/Private/OptimusDeformerInstance.cpp
+++ b/Engine/Plugins/Animation/DeformerGraph/Source/OptimusCore/Private/OptimusDeformerInstance.cpp
@@ -649,7 +649,7 @@ bool UOptimusDeformerInstance::EnqueueTriggerGraph(FName InTriggerGraphName)
 }
 
 
-void UOptimusDeformerInstance::SetConstantValueDirect(FString const& InVariableName, TArray<uint8> const& InValue)
+void UOptimusDeformerInstance::SetConstantValueDirect(TSoftObjectPtr<UObject> InSourceObject, TArray<uint8> const& InValue)
 {
 	// Poke constants into the UGraphDataProvider objects.
 	// This is an editor only operation when constant nodes are edited in the graph and we want to see the result without a full compile step.
@@ -661,7 +661,7 @@ void UOptimusDeformerInstance::SetConstantValueDirect(FString const& InVariableN
 		{
 			if (UOptimusGraphDataProvider* GraphDataProvider = Cast<UOptimusGraphDataProvider>(DataProvider))
 			{
-				GraphDataProvider->SetConstant(InVariableName, InValue);
+				GraphDataProvider->SetConstant(InSourceObject, InValue);
 				break;
 			}
 		}
diff --git a/Engine/Plugins/Animation/DeformerGraph/Source/OptimusCore/Private/OptimusHelpers.cpp b/Engine/Plugins/Animation/DeformerGraph/Source/OptimusCore/Private/OptimusHelpers.cpp
--- a/Engine/Plugins/Animation/DeformerGraph/Source/OptimusCore/Private/OptimusHelpers.cpp
+++ b/Engine/Plugins/Animation/DeformerGraph/Source/OptimusCore/Private/OptimusHelpers.cpp
@@ -303,3 +303,8 @@ FName Optimus::GenerateUniqueNameFromExistingNames(FName InBaseName, const TArra
 
 	return NewName;
 }
+
+FString Optimus::MakeUniqueValueName(const FString& InValueName, int32 InUniqueIndex)
+{
+	return InValueName + TEXT("_") + FString::FromInt(InUniqueIndex);
+}
diff --git a/Engine/Plugins/Animation/DeformerGraph/Source/OptimusCore/Public/IOptimusValueProvider.h b/Engine/Plugins/Animation/DeformerGraph/Source/OptimusCore/Public/IOptimusValueProvider.h
--- a/Engine/Plugins/Animation/DeformerGraph/Source/OptimusCore/Public/IOptimusValueProvider.h
+++ b/Engine/Plugins/Animation/DeformerGraph/Source/OptimusCore/Public/IOptimusValueProvider.h
@@ -24,7 +24,7 @@ class IOptimusValueProvider
 
 public:
 	// Returns the value name.
-	virtual FString GetValueName() const = 0;
+	virtual FString GetValueDebugName() const = 0;
 	// Returns the value type.
 	virtual FOptimusDataTypeRef GetValueType() const = 0;
 	// Returns the stored value as a shader-compatible value.
diff --git a/Engine/Plugins/Animation/DeformerGraph/Source/OptimusCore/Public/OptimusDeformer.h b/Engine/Plugins/Animation/DeformerGraph/Source/OptimusCore/Public/OptimusDeformer.h
--- a/Engine/Plugins/Animation/DeformerGraph/Source/OptimusCore/Public/OptimusDeformer.h
+++ b/Engine/Plugins/Animation/DeformerGraph/Source/OptimusCore/Public/OptimusDeformer.h
@@ -37,7 +37,7 @@ struct FOptimusCompoundAction;
 DECLARE_MULTICAST_DELEGATE_OneParam(FOptimusCompileBegin, UOptimusDeformer *);
 DECLARE_MULTICAST_DELEGATE_OneParam(FOptimusCompileEnd, UOptimusDeformer *);
 DECLARE_MULTICAST_DELEGATE_OneParam(FOptimusGraphCompileMessageDelegate, FOptimusCompilerDiagnostic const&);
-DECLARE_MULTICAST_DELEGATE_TwoParams(FOptimusConstantValueUpdate, FString const&, TArray<uint8> const&);
+DECLARE_MULTICAST_DELEGATE_TwoParams(FOptimusConstantValueUpdate, TSoftObjectPtr<UObject>, TArray<uint8> const&);
 DECLARE_MULTICAST_DELEGATE_OneParam(FOptimusSetAllInstancesCanbeActive, bool);
 
 UENUM()
diff --git a/Engine/Plugins/Animation/DeformerGraph/Source/OptimusCore/Public/OptimusDeformerInstance.h b/Engine/Plugins/Animation/DeformerGraph/Source/OptimusCore/Public/OptimusDeformerInstance.h
--- a/Engine/Plugins/Animation/DeformerGraph/Source/OptimusCore/Public/OptimusDeformerInstance.h
+++ b/Engine/Plugins/Animation/DeformerGraph/Source/OptimusCore/Public/OptimusDeformerInstance.h
@@ -246,7 +246,7 @@ public:
 	
 	
 	/** Directly set a graph constant value. */
-	void SetConstantValueDirect(FString const& InVariableName, TArray<uint8> const& InValue);
+	void SetConstantValueDirect(TSoftObjectPtr<UObject> InSourceObject, TArray<uint8> const& InValue);
 
 	FOptimusPersistentBufferPoolPtr GetBufferPool() const { return BufferPool; }
 
diff --git a/Engine/Plugins/Animation/DeformerGraph/Source/OptimusCore/Public/OptimusHelpers.h b/Engine/Plugins/Animation/DeformerGraph/Source/OptimusCore/Public/OptimusHelpers.h
--- a/Engine/Plugins/Animation/DeformerGraph/Source/OptimusCore/Public/OptimusHelpers.h
+++ b/Engine/Plugins/Animation/DeformerGraph/Source/OptimusCore/Public/OptimusHelpers.h
@@ -126,4 +126,6 @@ namespace Optimus
 	bool FindMovedItemInNameArray(const TArray<FName>& Old, const TArray<FName>& New, FName& OutSubjectName, FName& OutNextName);
 
 	FName GenerateUniqueNameFromExistingNames(FName InBaseName, const TArray<FName>& InExistingNames);
+
+	FString MakeUniqueValueName(const FString& InValueName, int32 InUniqueIndex);
 }
