[Deformer Graph] fix copy paste subgraph node leads to a new subgraph node with missing graph

diff --git a/Engine/Plugins/Animation/DeformerGraph/Source/OptimusCore/Private/Nodes/OptimusNode_FunctionReference.cpp b/Engine/Plugins/Animation/DeformerGraph/Source/OptimusCore/Private/Nodes/OptimusNode_FunctionReference.cpp
--- a/Engine/Plugins/Animation/DeformerGraph/Source/OptimusCore/Private/Nodes/OptimusNode_FunctionReference.cpp
+++ b/Engine/Plugins/Animation/DeformerGraph/Source/OptimusCore/Private/Nodes/OptimusNode_FunctionReference.cpp
@@ -33,7 +33,7 @@ FText UOptimusNode_FunctionReference::GetDisplayName() const
 
 void UOptimusNode_FunctionReference::ConstructNode()
 {
-	if (ensure(FunctionGraph.IsValid()))
+	if (FunctionGraph.IsValid())
 	{
 		const FOptimusDataTypeRegistry& TypeRegistry = FOptimusDataTypeRegistry::Get();
 		FOptimusDataTypeRef ComponentSourceType = TypeRegistry.FindType(*UOptimusComponentSourceBinding::StaticClass());
@@ -53,15 +53,6 @@ void UOptimusNode_FunctionReference::ConstructNode()
 }
 
 
-void UOptimusNode_FunctionReference::PostLoad()
-{
-	Super::PostLoad();
-
-	// Load the graph into memory
-	FunctionGraph.LoadSynchronous();
-}
-
-
 FOptimusRoutedNodePin UOptimusNode_FunctionReference::GetPinCounterpart(
 	UOptimusNodePin* InNodePin,
 	const FOptimusPinTraversalContext& InTraversalContext
@@ -72,7 +63,7 @@ FOptimusRoutedNodePin UOptimusNode_FunctionReference::GetPinCounterpart(
 		return {};
 	}
 
-	if (!ensure(FunctionGraph.IsValid()))
+	if (!FunctionGraph.IsValid())
 	{
 		return {};
 	}
@@ -150,11 +141,24 @@ FSoftObjectPath UOptimusNode_FunctionReference::GetSerializedGraphPath() const
 	return FunctionGraph.ToSoftObjectPath();
 }
 
-void UOptimusNode_FunctionReference::SetSerializedGraphPath(const FSoftObjectPath& InNewGraphPath)
+void UOptimusNode_FunctionReference::InitializeSerializedGraphPath(const FSoftObjectPath& InInitialGraphPath)
 {
-	FunctionGraph = InNewGraphPath;
-	// Making sure the pointer is alive such that we avoid loading on demand everywhere
-	FunctionGraph.LoadSynchronous();
+	FunctionGraph = InInitialGraphPath;
+}
 
+void UOptimusNode_FunctionReference::RefreshSerializedGraphPath(const FSoftObjectPath& InRenamedGraphPath)
+{
+	// Unlike subgraph reference where we can directly query the graph for its new name/path
+	// the function graph may not be loaded at the time of rename and thus
+	// the caller needs to provide the path to the renamed graph 
+	FunctionGraph = InRenamedGraphPath;
+	FunctionGraph.LoadSynchronous();
 	SetDisplayName(GetDisplayName());
 }
+
+void UOptimusNode_FunctionReference::InitializeTransientData()
+{
+	Super::InitializeTransientData();
+	// Making sure the pointer is alive such that we avoid loading on demand everywhere
+	FunctionGraph.LoadSynchronous();
+}
diff --git a/Engine/Plugins/Animation/DeformerGraph/Source/OptimusCore/Private/Nodes/OptimusNode_FunctionReference.h b/Engine/Plugins/Animation/DeformerGraph/Source/OptimusCore/Private/Nodes/OptimusNode_FunctionReference.h
--- a/Engine/Plugins/Animation/DeformerGraph/Source/OptimusCore/Private/Nodes/OptimusNode_FunctionReference.h
+++ b/Engine/Plugins/Animation/DeformerGraph/Source/OptimusCore/Private/Nodes/OptimusNode_FunctionReference.h
@@ -28,9 +28,6 @@ public:
 	FText GetDisplayName() const override;
 	void ConstructNode() override;
 
-	// UObject
-	void PostLoad() override;
-	
 	// IOptimusNodePinRouter implementation
 	FOptimusRoutedNodePin GetPinCounterpart(
 		UOptimusNodePin* InNodePin,
@@ -46,10 +43,11 @@ public:
 	UOptimusNodePin* GetDefaultComponentBindingPin() const override;
 
 	FSoftObjectPath GetSerializedGraphPath() const;
-
-	void SetSerializedGraphPath(const FSoftObjectPath& InNewGraphPath);
+	void InitializeSerializedGraphPath(const FSoftObjectPath& InInitialGraphPath);
+	void RefreshSerializedGraphPath(const FSoftObjectPath& InRenamedGraphPath);
 	
 protected:
+	void InitializeTransientData() override;
 	
 	/** The graph that owns us. This contains all the necessary pin information to add on
 	 * the terminal node.
diff --git a/Engine/Plugins/Animation/DeformerGraph/Source/OptimusCore/Private/Nodes/OptimusNode_SubGraphReference.cpp b/Engine/Plugins/Animation/DeformerGraph/Source/OptimusCore/Private/Nodes/OptimusNode_SubGraphReference.cpp
--- a/Engine/Plugins/Animation/DeformerGraph/Source/OptimusCore/Private/Nodes/OptimusNode_SubGraphReference.cpp
+++ b/Engine/Plugins/Animation/DeformerGraph/Source/OptimusCore/Private/Nodes/OptimusNode_SubGraphReference.cpp
@@ -48,16 +48,14 @@ void UOptimusNode_SubGraphReference::ConstructNode()
 		{
 			AddPinDirect(Binding, EOptimusNodePinDirection::Output);
 		}
-
-		SubscribeToSubGraph();
 	}
 }
 
-void UOptimusNode_SubGraphReference::PostLoad()
+void UOptimusNode_SubGraphReference::InitializeTransientData()
 {
-	Super::PostLoad();
-
-	SubscribeToSubGraph();
+	Super::InitializeTransientData();
+	
+	ResolveSubGraphPointerAndSubscribe();
 }
 
 void UOptimusNode_SubGraphReference::BeginDestroy()
@@ -151,11 +149,10 @@ UOptimusNodePin* UOptimusNode_SubGraphReference::GetDefaultComponentBindingPin()
 	return DefaultComponentPin.Get();
 }
 
-void UOptimusNode_SubGraphReference::SetSerializedSubGraphName(FName InSubGraphName)
+void UOptimusNode_SubGraphReference::InitializeSerializedSubGraphName(FName InInitialSubGraphName)
 {
-	SubGraphName = InSubGraphName;
-	FString SubGraphPath = GetOwningGraph()->ConstructSubGraphPath(SubGraphName.ToString());
-	SubGraph = Cast<UOptimusNodeSubGraph>(GetOwningGraph()->GetPathResolver()->ResolveGraphPath(SubGraphPath));
+	SubGraphName = InInitialSubGraphName;
+	// SubGraph object pointer is resolved during InitializeTransientData()
 }
 
 void UOptimusNode_SubGraphReference::RefreshSerializedSubGraphName()
@@ -168,9 +165,15 @@ FName UOptimusNode_SubGraphReference::GetSerializedSubGraphName() const
 	return SubGraphName;
 }
 
+void UOptimusNode_SubGraphReference::ResolveSubGraphPointerAndSubscribe()
+{
+	SubGraph = Cast<UOptimusNodeSubGraph>(GetOwningGraph()->FindGraphByName(GetSerializedSubGraphName()));
+	SubscribeToSubGraph();
+}
+
 void UOptimusNode_SubGraphReference::SubscribeToSubGraph()
 {
-	if (ensure(!SubGraph->GetOnBindingArrayPasted().IsBoundToObject(this)))
+	if (ensure(SubGraph.IsValid() && !SubGraph->GetOnBindingArrayPasted().IsBoundToObject(this)))
 	{
 		SubGraph->GetOnBindingArrayPasted().AddUObject(this, &UOptimusNode_SubGraphReference::RecreateBindingPins);
 		SubGraph->GetOnBindingValueChanged().AddUObject(this, &UOptimusNode_SubGraphReference::SyncPinsToBindings);
diff --git a/Engine/Plugins/Animation/DeformerGraph/Source/OptimusCore/Private/Nodes/OptimusNode_SubGraphReference.h b/Engine/Plugins/Animation/DeformerGraph/Source/OptimusCore/Private/Nodes/OptimusNode_SubGraphReference.h
--- a/Engine/Plugins/Animation/DeformerGraph/Source/OptimusCore/Private/Nodes/OptimusNode_SubGraphReference.h
+++ b/Engine/Plugins/Animation/DeformerGraph/Source/OptimusCore/Private/Nodes/OptimusNode_SubGraphReference.h
@@ -31,7 +31,6 @@ public:
 	void ConstructNode() override;
 
 	// UObject overrides
-	void PostLoad() override;
 	void BeginDestroy() override;
 
 	// IOptimusNodePinRouter implementation
@@ -49,11 +48,15 @@ public:
 	UOptimusNodePin* GetDefaultComponentBindingPin() const override;
 
 	// Only used during node creation, cannot be used to reference a different graph once node is constructed
-	void SetSerializedSubGraphName(FName InSubGraphName);
+	void InitializeSerializedSubGraphName(FName InInitialSubGraphName);
 	void RefreshSerializedSubGraphName();
 	FName GetSerializedSubGraphName() const;
 	
 protected:
+	// UOptimusNode overrides
+	void InitializeTransientData() override;
+	
+	void ResolveSubGraphPointerAndSubscribe();
 	void SubscribeToSubGraph();
 	void UnsubscribeFromSubGraph() const;
 
@@ -73,9 +76,10 @@ protected:
 	TWeakObjectPtr<UOptimusNodePin> DefaultComponentPin;
 
 private:
+	
+	
 	/** The graph that owns us. This contains all the necessary pin information to add on
-	 * the terminal node.
+	 * the terminal node. Initialized when the node is loaded/created
 	 */
-	UPROPERTY()
 	TWeakObjectPtr<UOptimusNodeSubGraph> SubGraph;	
 };
diff --git a/Engine/Plugins/Animation/DeformerGraph/Source/OptimusCore/Private/OptimusCoreModule.cpp b/Engine/Plugins/Animation/DeformerGraph/Source/OptimusCore/Private/OptimusCoreModule.cpp
--- a/Engine/Plugins/Animation/DeformerGraph/Source/OptimusCore/Private/OptimusCoreModule.cpp
+++ b/Engine/Plugins/Animation/DeformerGraph/Source/OptimusCore/Private/OptimusCoreModule.cpp
@@ -79,7 +79,7 @@ void FOptimusCoreModule::UpdateFunctionReferences(const FSoftObjectPath& InOldGr
 			if (FunctionNode)
 			{
 				FunctionNode->Modify();
-				FunctionNode->SetSerializedGraphPath(InNewGraphPath);
+				FunctionNode->RefreshSerializedGraphPath(InNewGraphPath);
 				FunctionNode->MarkPackageDirty();
 			}
 		}
diff --git a/Engine/Plugins/Animation/DeformerGraph/Source/OptimusCore/Private/OptimusDeformer.cpp b/Engine/Plugins/Animation/DeformerGraph/Source/OptimusCore/Private/OptimusDeformer.cpp
--- a/Engine/Plugins/Animation/DeformerGraph/Source/OptimusCore/Private/OptimusDeformer.cpp
+++ b/Engine/Plugins/Animation/DeformerGraph/Source/OptimusCore/Private/OptimusDeformer.cpp
@@ -3632,6 +3632,18 @@ UOptimusNodePin* UOptimusDeformer::ResolvePinPath(const FString& InPinPath)
 }
 
 
+UOptimusNodeGraph* UOptimusDeformer::FindGraphByName(FName InGraphName) const
+{
+	for (UOptimusNodeGraph* Graph : GetGraphs())
+	{
+		if (Graph->GetFName() == InGraphName)
+		{
+			return Graph;
+		}
+	}
+
+	return nullptr;
+}
 
 UOptimusNodeGraph* UOptimusDeformer::CreateGraphDirect(
 	EOptimusNodeGraphType InType, 
diff --git a/Engine/Plugins/Animation/DeformerGraph/Source/OptimusCore/Private/OptimusNode.cpp b/Engine/Plugins/Animation/DeformerGraph/Source/OptimusCore/Private/OptimusNode.cpp
--- a/Engine/Plugins/Animation/DeformerGraph/Source/OptimusCore/Private/OptimusNode.cpp
+++ b/Engine/Plugins/Animation/DeformerGraph/Source/OptimusCore/Private/OptimusNode.cpp
@@ -400,6 +400,8 @@ FName UOptimusNode::GetAvailablePinNameStable(const UObject* InNodeOrPin, FName
 
 void UOptimusNode::PostCreateNode()
 {
+	InitializeTransientData();
+	
 	CachedPinLookup.Empty();
 	Pins.Empty();
 
@@ -429,6 +431,8 @@ void UOptimusNode::PostLoad()
 	{
 		Pin->ConditionalPostLoad();
 	}
+	
+	InitializeTransientData();
 }
 
 #if WITH_EDITOR
@@ -467,6 +471,14 @@ void UOptimusNode::ConstructNode()
 	CreatePinsFromStructLayout(GetClass(), nullptr);
 }
 
+void UOptimusNode::InitializeTransientData()
+{
+	// Called from three places, which are all possible way a node object can be created
+	// 1. PostCreateNode
+	// 2. PostLoad
+	// 3. FOptimusEditorClipboard::ProcessPostCreateObject
+}
+
 
 void UOptimusNode::SaveState(FArchive& Ar) const
 {
diff --git a/Engine/Plugins/Animation/DeformerGraph/Source/OptimusCore/Private/OptimusNodeGraph.cpp b/Engine/Plugins/Animation/DeformerGraph/Source/OptimusCore/Private/OptimusNodeGraph.cpp
--- a/Engine/Plugins/Animation/DeformerGraph/Source/OptimusCore/Private/OptimusNodeGraph.cpp
+++ b/Engine/Plugins/Animation/DeformerGraph/Source/OptimusCore/Private/OptimusNodeGraph.cpp
@@ -491,11 +491,11 @@ bool UOptimusNodeGraph::DuplicateSubGraph(UOptimusNodeSubGraph* InSourceSubGraph
 	{
 		FOptimusActionScope(*GetActionStack(), TEXT("Duplicate SubGraph"));
 		GetActionStack()->RunAction<FOptimusNodeGraphAction_AddGraph>(this, EOptimusNodeGraphType::SubGraph, SubGraphName, INDEX_NONE,
-			[InSourceSubGraph](UOptimusNodeGraph* InGraph)
+			[InputBindings = InSourceSubGraph->InputBindings, OutputBindings = InSourceSubGraph->OutputBindings](UOptimusNodeGraph* InGraph)
 			{
 				UOptimusNodeSubGraph* SubGraph = CastChecked<UOptimusNodeSubGraph>(InGraph);
-				SubGraph->InputBindings = InSourceSubGraph->InputBindings;
-				SubGraph->OutputBindings = InSourceSubGraph->OutputBindings;
+				SubGraph->InputBindings = InputBindings;
+				SubGraph->OutputBindings = OutputBindings;
 				return true;
 			});
 	
@@ -536,13 +536,8 @@ bool UOptimusNodeGraph::DuplicateSubGraph(UOptimusNodeSubGraph* InSourceSubGraph
 		for (UOptimusNode* Node: NodesToDuplicate)
 		{
 			GetActionStack()->RunAction<FOptimusNodeGraphAction_DuplicateNode>(
-				SubGraphPath, Node, Node->GetFName(), [](UOptimusNode* InNode)
+				SubGraphPath, Node, Node->GetFName(), [](UOptimusNode*)
 				{
-					if (UOptimusNode_SubGraphReference* SubGraphNode = Cast<UOptimusNode_SubGraphReference>(InNode))
-					{
-						FName NewSubGraphName = SubGraphNode->GetReferencedSubGraph()->GetFName();
-						SubGraphNode->SetSerializedSubGraphName(NewSubGraphName);
-					}
 					return true;
 				});
 		}
@@ -657,7 +652,7 @@ UOptimusNode* UOptimusNodeGraph::AddFunctionReferenceNode(TSoftObjectPtr<UOptimu
 	[InFunctionGraph](UOptimusNode *InNode)
 	{
 		UOptimusNode_FunctionReference* FunctionNode = CastChecked<UOptimusNode_FunctionReference>(InNode);
-		FunctionNode->SetSerializedGraphPath(InFunctionGraph.ToSoftObjectPath());
+		FunctionNode->InitializeSerializedGraphPath(InFunctionGraph.ToSoftObjectPath());
 	});
 }
 
@@ -988,14 +983,15 @@ bool UOptimusNodeGraph::DuplicateNodes(
 		{
 			FOptimusNodeGraphAction_DuplicateNode *DuplicateNodeAction = new FOptimusNodeGraphAction_DuplicateNode(
 				GetGraphPath(), Node, NewNodeNameMap[Node],
-				[Node, NodeOffset, NewSubGraphNameMap](UOptimusNode *InNode) {
+				[SourceNodeGraphPosition = Node->GraphPosition, NodeOffset, NewSubGraphNameMap](UOptimusNode *InNode)
+				{
 					if (UOptimusNode_SubGraphReference* SubGraphNode = Cast<UOptimusNode_SubGraphReference>(InNode))
 					{
 						FName NewSubGraphName = NewSubGraphNameMap[SubGraphNode->GetSerializedSubGraphName()];
-						SubGraphNode->SetSerializedSubGraphName(NewSubGraphName);
+						SubGraphNode->InitializeSerializedSubGraphName(NewSubGraphName);
 					}
-					return InNode->SetGraphPositionDirect(Node->GraphPosition + NodeOffset); 
-			});
+					return InNode->SetGraphPositionDirect(SourceNodeGraphPosition + NodeOffset); 
+				});
 			
 			GetActionStack()->RunAction(DuplicateNodeAction);
 		}
@@ -1416,13 +1412,8 @@ UOptimusNode* UOptimusNodeGraph::CollapseNodesToSubGraph(
 	for (UOptimusNode* Node: NodesToDuplicate)
 	{
 		Action->AddSubAction<FOptimusNodeGraphAction_DuplicateNode>(
-			SubGraphPath, Node, Node->GetFName(), [](UOptimusNode* InNode)
+			SubGraphPath, Node, Node->GetFName(), [](UOptimusNode*)
 			{
-				if (UOptimusNode_SubGraphReference* SubGraphNode = Cast<UOptimusNode_SubGraphReference>(InNode))
-				{
-					FName NewSubGraphName = SubGraphNode->GetReferencedSubGraph()->GetFName();
-					SubGraphNode->SetSerializedSubGraphName(NewSubGraphName);
-				}	
 				return true;
 			});
 	}
@@ -1453,7 +1444,7 @@ UOptimusNode* UOptimusNodeGraph::CollapseNodesToSubGraph(
 		[NodeBox, SubGraphName](UOptimusNode* InNode)
 		{
 			UOptimusNode_SubGraphReference* SubGraphNode = Cast<UOptimusNode_SubGraphReference>(InNode); 
-			SubGraphNode->SetSerializedSubGraphName(SubGraphName);
+			SubGraphNode->InitializeSerializedSubGraphName(SubGraphName);
 			return SubGraphNode->SetGraphPositionDirect(NodeBox.GetCenter());
 		});
 	
@@ -1745,8 +1736,8 @@ TArray<UOptimusNode*> UOptimusNodeGraph::ExpandCollapsedNodes(
 		{
 			FOptimusNodeGraphAction_DuplicateNode *DuplicateNodeAction = new FOptimusNodeGraphAction_DuplicateNode(
 				GetGraphPath(), Node, NewNodeNameMap[Node],
-				[Node, Center, ExpandLocation](UOptimusNode *InNode) {
-					InNode->SetGraphPositionDirect(Node->GraphPosition - Center + ExpandLocation);
+				[SourceNodeGraphPosition = Node->GraphPosition, Center, ExpandLocation](UOptimusNode *InNode) {
+					InNode->SetGraphPositionDirect(SourceNodeGraphPosition - Center + ExpandLocation);
 					return true;
 			});
 		
@@ -1974,13 +1965,8 @@ bool UOptimusNodeGraph::ConvertToFunction(UOptimusNode* InSubGraphNode)
 		for (UOptimusNode* Node: NodesToDuplicate)
 		{
 			GetActionStack()->RunAction<FOptimusNodeGraphAction_DuplicateNode>(
-				FunctionGraphPath, Node, Node->GetFName(), [](UOptimusNode* InNode)
+				FunctionGraphPath, Node, Node->GetFName(), [](UOptimusNode*)
 				{
-					if (UOptimusNode_SubGraphReference* SubGraphNode = Cast<UOptimusNode_SubGraphReference>(InNode))
-					{
-						FName NewSubGraphName = SubGraphNode->GetReferencedSubGraph()->GetFName();
-						SubGraphNode->SetSerializedSubGraphName(NewSubGraphName);
-					}
 					return true;
 				});
 		}
@@ -2007,7 +1993,7 @@ bool UOptimusNodeGraph::ConvertToFunction(UOptimusNode* InSubGraphNode)
 			{
 				UOptimusNode_FunctionReference* FunctionNode = Cast<UOptimusNode_FunctionReference>(InNode); 
 				UOptimusFunctionNodeGraph* FunctionGraph = Cast<UOptimusFunctionNodeGraph>(PathResolver->ResolveGraphPath(FunctionGraphPath));
-				FunctionNode->SetSerializedGraphPath(FunctionGraph);
+				FunctionNode->InitializeSerializedGraphPath(FunctionGraph);
 				
 				return FunctionNode->SetGraphPositionDirect(ReferenceNodePosition);
 			});
@@ -2217,13 +2203,8 @@ bool UOptimusNodeGraph::ConvertToSubGraph(UOptimusNode* InFunctionNode)
 		for (UOptimusNode* Node: NodesToDuplicate)
 		{
 			GetActionStack()->RunAction<FOptimusNodeGraphAction_DuplicateNode>(
-				SubGraphPath, Node, Node->GetFName(), [](UOptimusNode* InNode)
+				SubGraphPath, Node, Node->GetFName(), [](UOptimusNode*)
 				{
-					if (UOptimusNode_SubGraphReference* SubGraphNode = Cast<UOptimusNode_SubGraphReference>(InNode))
-					{
-						FName NewSubGraphName = SubGraphNode->GetReferencedSubGraph()->GetFName();
-						SubGraphNode->SetSerializedSubGraphName(NewSubGraphName);
-					}
 					return true;
 				});
 		}
@@ -2249,7 +2230,7 @@ bool UOptimusNodeGraph::ConvertToSubGraph(UOptimusNode* InFunctionNode)
 			[ReferenceNodePosition, SubGraphName](UOptimusNode* InNode)
 			{
 				UOptimusNode_SubGraphReference* SubGraphNode = Cast<UOptimusNode_SubGraphReference>(InNode);
-				SubGraphNode->SetSerializedSubGraphName(SubGraphName);
+				SubGraphNode->InitializeSerializedSubGraphName(SubGraphName);
 				
 				return SubGraphNode->SetGraphPositionDirect(ReferenceNodePosition);
 			});
@@ -3221,6 +3202,19 @@ FString UOptimusNodeGraph::GetCollectionPath() const
 }
 
 
+UOptimusNodeGraph* UOptimusNodeGraph::FindGraphByName(FName InGraphName) const
+{
+	for (UOptimusNodeGraph* Graph : GetGraphs())
+	{
+		if (Graph->GetFName() == InGraphName)
+		{
+			return Graph;
+		}
+	}
+
+	return nullptr;
+}
+
 UOptimusNodeGraph* UOptimusNodeGraph::CreateGraphDirect(
 	EOptimusNodeGraphType InType,
 	FName InName,
diff --git a/Engine/Plugins/Animation/DeformerGraph/Source/OptimusCore/Public/IOptimusNodeGraphCollectionOwner.h b/Engine/Plugins/Animation/DeformerGraph/Source/OptimusCore/Public/IOptimusNodeGraphCollectionOwner.h
--- a/Engine/Plugins/Animation/DeformerGraph/Source/OptimusCore/Public/IOptimusNodeGraphCollectionOwner.h
+++ b/Engine/Plugins/Animation/DeformerGraph/Source/OptimusCore/Public/IOptimusNodeGraphCollectionOwner.h
@@ -47,6 +47,8 @@ public:
 	/// Returns all immediately owned node graphs.
 	virtual const TArray<UOptimusNodeGraph*> &GetGraphs() const = 0;
 
+	virtual UOptimusNodeGraph* FindGraphByName(FName InGraphName) const = 0;
+
 	/// Create a new graph of a given type, with an optional name. The name may be changed to 
 	/// fit into the namespace. Only setup and trigger graphs can currently be created directly,
 	/// and only a single setup graph. The setup graph is always the first, and the trigger graphs
diff --git a/Engine/Plugins/Animation/DeformerGraph/Source/OptimusCore/Public/OptimusDeformer.h b/Engine/Plugins/Animation/DeformerGraph/Source/OptimusCore/Public/OptimusDeformer.h
--- a/Engine/Plugins/Animation/DeformerGraph/Source/OptimusCore/Public/OptimusDeformer.h
+++ b/Engine/Plugins/Animation/DeformerGraph/Source/OptimusCore/Public/OptimusDeformer.h
@@ -339,6 +339,8 @@ public:
 
 	const TArray<UOptimusNodeGraph*> &GetGraphs() const override { return Graphs; }
 
+	UOptimusNodeGraph* FindGraphByName(FName InGraphName) const override;
+
 	UOptimusNodeGraph* CreateGraph(
 		EOptimusNodeGraphType InType,
 		FName InName)
diff --git a/Engine/Plugins/Animation/DeformerGraph/Source/OptimusCore/Public/OptimusNode.h b/Engine/Plugins/Animation/DeformerGraph/Source/OptimusCore/Public/OptimusNode.h
--- a/Engine/Plugins/Animation/DeformerGraph/Source/OptimusCore/Public/OptimusNode.h
+++ b/Engine/Plugins/Animation/DeformerGraph/Source/OptimusCore/Public/OptimusNode.h
@@ -188,6 +188,7 @@ public:
 #endif
 	
 protected:
+	friend class FOptimusEditorClipboard;
 	friend class UOptimusNodeGraph;
 	friend class UOptimusNodeSubGraph;
 	friend class UOptimusNodePin;
@@ -217,6 +218,9 @@ protected:
 	// Called when the node is being constructed
 	virtual void ConstructNode();
 
+	// Allows the node to initialize any transient data that can be derived from serialized properties
+	virtual void InitializeTransientData();
+
 	virtual bool ValidateConnection(
 		const UOptimusNodePin& InThisNodesPin,
 		const UOptimusNodePin& InOtherNodesPin,
diff --git a/Engine/Plugins/Animation/DeformerGraph/Source/OptimusCore/Public/OptimusNodeGraph.h b/Engine/Plugins/Animation/DeformerGraph/Source/OptimusCore/Public/OptimusNodeGraph.h
--- a/Engine/Plugins/Animation/DeformerGraph/Source/OptimusCore/Public/OptimusNodeGraph.h
+++ b/Engine/Plugins/Animation/DeformerGraph/Source/OptimusCore/Public/OptimusNodeGraph.h
@@ -401,6 +401,8 @@ public:
 	UFUNCTION(BlueprintCallable, Category = OptimusNodeGraph)
 	const TArray<UOptimusNodeGraph*> &GetGraphs() const override { return SubGraphs; }
 
+	UOptimusNodeGraph* FindGraphByName(FName InGraphName) const override;
+	
 	UOptimusNodeGraph* CreateGraphDirect(
 		EOptimusNodeGraphType InType,
 		FName InName,
diff --git a/Engine/Plugins/Animation/DeformerGraph/Source/OptimusEditor/Private/OptimusEditor.cpp b/Engine/Plugins/Animation/DeformerGraph/Source/OptimusEditor/Private/OptimusEditor.cpp
--- a/Engine/Plugins/Animation/DeformerGraph/Source/OptimusEditor/Private/OptimusEditor.cpp
+++ b/Engine/Plugins/Animation/DeformerGraph/Source/OptimusEditor/Private/OptimusEditor.cpp
@@ -869,6 +869,8 @@ void FOptimusEditor::OnNodeDoubleClicked(class UEdGraphNode* Node)
 					
 					if (FoundAssetEditor.IsValid())
 					{
+						FoundAssetEditor->BringToolkitToFront();
+						
 						TSharedPtr<FOptimusEditor> OptimusEditorForExternalGraph = StaticCastSharedPtr<FOptimusEditor>(FoundAssetEditor);
 
 						OptimusEditorForExternalGraph->SetEditGraph(GraphToShow);
@@ -1287,7 +1289,13 @@ void FOptimusEditor::OnDeformerModified(
 	switch (InNotifyType)
 	{
 	case EOptimusGlobalNotifyType::GraphAdded:
-		SetEditGraph(Cast<UOptimusNodeGraph>(InModifiedObject));
+		if (UOptimusNodeGraph* Graph = Cast<UOptimusNodeGraph>(InModifiedObject))
+		{
+			if (Optimus::IsExecutionGraphType(Graph->GetGraphType()))
+			{
+				SetEditGraph(Graph);	
+			}
+		}
 		RefreshEvent.Broadcast();
 		break;
 
diff --git a/Engine/Plugins/Animation/DeformerGraph/Source/OptimusEditor/Private/OptimusEditorClipboard.cpp b/Engine/Plugins/Animation/DeformerGraph/Source/OptimusEditor/Private/OptimusEditorClipboard.cpp
--- a/Engine/Plugins/Animation/DeformerGraph/Source/OptimusEditor/Private/OptimusEditorClipboard.cpp
+++ b/Engine/Plugins/Animation/DeformerGraph/Source/OptimusEditor/Private/OptimusEditorClipboard.cpp
@@ -11,11 +11,48 @@
 
 #include "Editor.h"
 #include "IOptimusGeneratedClassDefiner.h"
+#include "OptimusHelpers.h"
 #include "Exporters/Exporter.h"
 #include "UnrealExporter.h"
 #include "HAL/PlatformApplicationMisc.h"
 #include "Misc/EngineVersion.h"
+#include "OptimusCore/Private/Nodes/OptimusNode_SubGraphReference.h"
 
+static TArray<const IOptimusGeneratedClassDefiner*> CollectClassDefiners(const TArray<UOptimusNode*>& InNodes , const TArray<UOptimusNodeSubGraph*> InSubGraphs)
+{
+	TArray<const IOptimusGeneratedClassDefiner*> ClassDefiners;
+	
+	TArray<UOptimusNode*> NodesToScanForClassDefiner = InNodes;
+
+	TQueue<UOptimusNodeGraph*> SubGraphQueue;
+	for (UOptimusNodeSubGraph* SubGraph : InSubGraphs)
+	{
+		SubGraphQueue.Enqueue(SubGraph);
+	}
+
+	UOptimusNodeGraph* WorkingSubGraph;
+	while (SubGraphQueue.Dequeue(WorkingSubGraph))
+	{
+		NodesToScanForClassDefiner.Append(WorkingSubGraph->GetAllNodes());
+		for (UOptimusNodeGraph* const SubGraph : WorkingSubGraph->GetGraphs())
+		{
+			SubGraphQueue.Enqueue(SubGraph);
+		}
+	}
+	
+	for (UOptimusNode* Node: NodesToScanForClassDefiner)
+	{
+		// Export any generated classes first. If the node implements the IOptimusGeneratedClassDefiner interface,
+		// it means that it uses a generated class as a base. We need to ensure that we have all the relevant information
+		// to construct the generated class on the receiving side.
+		if (const IOptimusGeneratedClassDefiner* ClassDefiner = Cast<IOptimusGeneratedClassDefiner>(Node))
+		{
+			ClassDefiners.Add(ClassDefiner);
+		}
+	}
+
+	return ClassDefiners;
+}
 
 void FOptimusEditorClipboard::SetClipboardFromNodes(const TArray<UOptimusNode*>& InNodes)
 {
@@ -32,6 +69,8 @@ void FOptimusEditorClipboard::SetClipboardFromNodes(const TArray<UOptimusNode*>&
 
 	Graph->GatherRelatedObjects(InNodes, NodesToConsider, NodePairsToConsider, SubGraphsToConsider);
 
+	TArray<const IOptimusGeneratedClassDefiner*> ClassDefiners = CollectClassDefiners(NodesToConsider, SubGraphsToConsider);
+	
 	// Export the clipboard to text.
 	const FExportObjectInnerContext Context;
 	FStringOutputDevice Archive;
@@ -39,7 +78,37 @@ void FOptimusEditorClipboard::SetClipboardFromNodes(const TArray<UOptimusNode*>&
 	// Clear all stray export marks to ensure all the objects get properly exported. 
 	UnMarkAllObjects(OBJECTMARK_TagExp);
 
-	// Export the nodes first.
+	// Export any generated classes first. If the node implements the IOptimusGeneratedClassDefiner interface,
+	// it means that it uses a generated class as a base. We need to ensure that we have all the relevant information
+	// to construct the generated class on the receiving side.
+	for (const IOptimusGeneratedClassDefiner* ClassDefiner : ClassDefiners)
+	{
+		const UOptimusNode* Node = Cast<UOptimusNode>(ClassDefiner);
+		if (ensure(Node))
+		{
+			Archive.Logf(TEXT("GeneratedClass Class=%s GeneratorPath=%s %s\n"),
+						*Node->GetClass()->GetPathName(), *ClassDefiner->GetAssetPathForClassDefiner().ToString(),
+						*ClassDefiner->GetClassCreationString());	
+		}
+	}
+	
+	
+	// Export the sub graphs.
+	for (UOptimusNodeSubGraph* SubGraph : SubGraphsToConsider)
+	{
+		UExporter::ExportToOutputDevice(
+			&Context,
+			SubGraph,
+			nullptr,
+			Archive,
+			TEXT("copy"),
+			0,
+			PPF_ExportsNotFullyQualified | PPF_Copy | PPF_Delimited,
+			false,
+			Graph);	
+	}
+
+	// Export the nodes.
 	TSet<UOptimusNode*> NodesToCopy;
 	NodesToCopy.Reserve(NodesToConsider.Num());
 	for (UOptimusNode* Node: NodesToConsider)
@@ -48,16 +117,6 @@ void FOptimusEditorClipboard::SetClipboardFromNodes(const TArray<UOptimusNode*>&
 		{
 			continue;
 		}
-
-		// Export any generated classes first. If the node implements the IOptimusGeneratedClassDefiner interface,
-		// it means that it uses a generated class as a base. We need to ensure that we have all the relevant information
-		// to construct the generated class on the receiving side.
-		if (const IOptimusGeneratedClassDefiner* ClassDefiner = Cast<IOptimusGeneratedClassDefiner>(Node))
-		{
-			Archive.Logf(TEXT("GeneratedClass Class=%s GeneratorPath=%s %s\n"),
-				*Node->GetClass()->GetPathName(), *ClassDefiner->GetAssetPathForClassDefiner().ToString(),
-				*ClassDefiner->GetClassCreationString());
-		}
 		
 		UExporter::ExportToOutputDevice(
 			&Context,
@@ -129,7 +188,9 @@ UOptimusNodeGraph* FOptimusEditorClipboard::GetGraphFromClipboardContent(
 	// since the original does not handle sub-objects at all.
 	if (CanCreateObjectsFromText(*ClipboardText))
 	{
-		UOptimusNodeGraph* Graph = NewObject<UOptimusNodeGraph>(GetTransientPackage(), FName("ClipboardGraph"), RF_Transient);
+		// Avoid multiple paste pasting into the same transient graph causing name collision
+		const FName UniqueClipboardGraphName = MakeUniqueObjectName(GetTransientPackage(), UOptimusNodeGraph::StaticClass(), FName("ClipboardGraph"));
+		UOptimusNodeGraph* Graph = NewObject<UOptimusNodeGraph>(GetTransientPackage(), UniqueClipboardGraphName, RF_Transient);
 		
 		ProcessObjectBuffer(InTargetPackage, Graph, *ClipboardText);
 		return Graph;
@@ -326,8 +387,9 @@ bool FOptimusEditorClipboard::ProcessObjectBuffer(
 		}
 	}
 
-	// Process the created objects in reverse order of creation.
-	for (int32 Index = CreatedObjects.Num() - 1; Index >= 0; Index--)
+	// Process the created objects in the order of creation because 
+	// objects created later depend on objects created earlier.
+	for (int32 Index = 0; Index < CreatedObjects.Num(); Index++)
 	{
 		if (!ProcessPostCreateObject(InRootOuter, CreatedObjects[Index]))
 		{
@@ -344,7 +406,8 @@ bool FOptimusEditorClipboard::CanCreateClass(const UClass* InClass)
 	return InClass->IsChildOf(UOptimusNode::StaticClass()) ||
 		   InClass == UOptimusNodePin::StaticClass() ||
 		   InClass == UOptimusNodeLink::StaticClass() ||
-		   InClass == UOptimusNodePair::StaticClass();
+		   InClass == UOptimusNodePair::StaticClass() ||
+		   InClass == UOptimusNodeSubGraph::StaticClass();
 }
 
 
@@ -357,6 +420,8 @@ bool FOptimusEditorClipboard::ProcessPostCreateObject(UObject* InRootOuter, UObj
 		{
 			UOptimusNode* Node = Cast<UOptimusNode>(InNewObject);
 
+			Node->InitializeTransientData();
+			
 			return Graph->AddNodeDirect(Node);
 		}
 		else if (InNewObject->IsA<UOptimusNodeLink>())
@@ -364,12 +429,19 @@ bool FOptimusEditorClipboard::ProcessPostCreateObject(UObject* InRootOuter, UObj
 			const UOptimusNodeLink* Link = Cast<UOptimusNodeLink>(InNewObject);
 
 			return Graph->AddLinkDirect(Link->GetNodeOutputPin(), Link->GetNodeInputPin());
-		}else if (InNewObject->IsA<UOptimusNodePair>())
+		}
+		else if (InNewObject->IsA<UOptimusNodePair>())
 		{
 			const UOptimusNodePair* NodePair = Cast<UOptimusNodePair>(InNewObject);
 
 			return Graph->AddNodePairDirect(NodePair->GetFirst(), NodePair->GetSecond());
 		}
+		else if(InNewObject->IsA<UOptimusNodeSubGraph>())
+		{
+			UOptimusNodeSubGraph* SubGraph = Cast<UOptimusNodeSubGraph>(InNewObject);
+
+			return Graph->AddGraphDirect(SubGraph, INDEX_NONE);
+		}
 	}
 
 	return true;
