[Deformer Graph] avoid crash when the debug skm switches to CPU skinning when enabling "Character -> Advanced -> Normals" in the viewport buttons.

diff --git a/Engine/Plugins/Animation/DeformerGraph/Source/OptimusCore/Private/DataInterfaces/OptimusDataInterfaceSkinnedMeshWrite.cpp b/Engine/Plugins/Animation/DeformerGraph/Source/OptimusCore/Private/DataInterfaces/OptimusDataInterfaceSkinnedMeshWrite.cpp
--- a/Engine/Plugins/Animation/DeformerGraph/Source/OptimusCore/Private/DataInterfaces/OptimusDataInterfaceSkinnedMeshWrite.cpp
+++ b/Engine/Plugins/Animation/DeformerGraph/Source/OptimusCore/Private/DataInterfaces/OptimusDataInterfaceSkinnedMeshWrite.cpp
@@ -142,6 +142,10 @@ bool FOptimusSkinnedMeshWriteDataProviderProxy::IsValid(FValidationData const& I
 	{
 		return false;
 	}
+	if (SkeletalMeshObject->IsCPUSkinned())
+	{
+		return false;
+	}
 	if (SkeletalMeshObject->GetSkeletalMeshRenderData().LODRenderData[SkeletalMeshObject->GetLOD()].RenderSections.Num() != InValidationData.NumInvocations)
 	{
 		return false;
