[Deformer Graph] Fix a bug with Skeleton DI shader code not working for unlimited bone weights

diff --git a/Engine/Plugins/Animation/DeformerGraph/Shaders/Private/DataInterfaceSkeleton.ush b/Engine/Plugins/Animation/DeformerGraph/Shaders/Private/DataInterfaceSkeleton.ush
--- a/Engine/Plugins/Animation/DeformerGraph/Shaders/Private/DataInterfaceSkeleton.ush
+++ b/Engine/Plugins/Animation/DeformerGraph/Shaders/Private/DataInterfaceSkeleton.ush
@@ -42,7 +42,8 @@ uint ReadNumBones_{DataInterfaceName}(uint VertexIndex)
 #if !ENABLE_DEFORMER_BONES
 	return 1;
 #elif GPUSKIN_UNLIMITED_BONE_INFLUENCE
-	uint BlendOffsetCount = {DataInterfaceName}_InputWeightLookupStream[VertexIndex];
+	const uint BufferIndex 		= {DataInterfaceName}_InputStreamStart + VertexIndex;
+	uint BlendOffsetCount = {DataInterfaceName}_InputWeightLookupStream[BufferIndex];
 	return BlendOffsetCount & 0xff;
 #else // !GPUSKIN_UNLIMITED_BONE_INFLUENCE
 	return {DataInterfaceName}_NumBoneInfluences;
@@ -56,13 +57,13 @@ float3x4 ReadBoneMatrix_{DataInterfaceName}(uint VertexIndex, uint BoneIndex)
 	return float3x4(1, 0, 0, 0, 0, 1, 0, 0, 0, 0, 1, 0);
 
 #elif GPUSKIN_UNLIMITED_BONE_INFLUENCE
-	
-	uint BlendOffsetCount = {DataInterfaceName}_InputWeightLookupStream[VertexIndex];
+	const uint BufferIndex 		= {DataInterfaceName}_InputStreamStart + VertexIndex;
+	uint BlendOffsetCount = {DataInterfaceName}_InputWeightLookupStream[BufferIndex];
 	int StreamOffset = BlendOffsetCount >> 8;
-
-	int BoneIndexOffset = StreamOffset + ({DataInterfaceName}_InputWeightIndexSize * BoneIndex);
+	int BoneIndexSize = {DataInterfaceName}_InputWeightIndexSize & 0xff;
+	int BoneIndexOffset = StreamOffset + (BoneIndexSize * BoneIndex);
 	int BoneGlobalIndex = {DataInterfaceName}_InputWeightStream[BoneIndexOffset];
-	if ({DataInterfaceName}_InputWeightIndexSize > 1)
+	if (BoneIndexSize > 1)
 	{
 		BoneGlobalIndex = {DataInterfaceName}_InputWeightStream[BoneIndexOffset + 1] << 8 | BoneGlobalIndex;
 	}
@@ -85,8 +86,8 @@ float ReadBoneWeight_{DataInterfaceName}(uint VertexIndex, uint BoneIndex)
 	return 1;
 
 #elif GPUSKIN_UNLIMITED_BONE_INFLUENCE
-	
-	uint BlendOffsetCount = {DataInterfaceName}_InputWeightLookupStream[VertexIndex];
+	const uint BufferIndex 		= {DataInterfaceName}_InputStreamStart + VertexIndex;	
+	uint BlendOffsetCount = {DataInterfaceName}_InputWeightLookupStream[BufferIndex];
 	int NumBoneInfluencesLocal = BlendOffsetCount & 0xff;
 	int StreamOffset = BlendOffsetCount >> 8;
 	int BoneIndexSize = {DataInterfaceName}_InputWeightIndexSize & 0xff;
