Fix incorrect buffer declaration for bound structured buffers.
Fixes corruption on some platforms when rendering deformers that sample half edge connectivity data.
#rnx

diff --git a/Engine/Plugins/Animation/DeformerGraph/Shaders/Private/DataInterfaceHalfEdge.ush b/Engine/Plugins/Animation/DeformerGraph/Shaders/Private/DataInterfaceHalfEdge.ush
--- a/Engine/Plugins/Animation/DeformerGraph/Shaders/Private/DataInterfaceHalfEdge.ush
+++ b/Engine/Plugins/Animation/DeformerGraph/Shaders/Private/DataInterfaceHalfEdge.ush
@@ -8,8 +8,8 @@ StructuredBuffer<int> {DataInterfaceName}_VertexToEdgeBuffer;
 StructuredBuffer<int> {DataInterfaceName}_EdgeToTwinEdgeBuffer;
 
 uint {DataInterfaceName}_bUseBufferFromRenderData;
-Buffer<int> {DataInterfaceName}_RenderDataVertexToEdgeBuffer;
-Buffer<int> {DataInterfaceName}_RenderDataEdgeToTwinEdgeBuffer;
+StructuredBuffer<int> {DataInterfaceName}_RenderDataVertexToEdgeBuffer;
+StructuredBuffer<int> {DataInterfaceName}_RenderDataEdgeToTwinEdgeBuffer;
 
 uint ReadNumVertices_{DataInterfaceName}()
 {
