[Deformer] Fix Graph deletion not working with undo system
+ minor fix for expand and collapse node actions to deal with subgraph nodes

diff --git a/Engine/Plugins/Animation/DeformerGraph/Source/OptimusCore/Private/Actions/OptimusNodeGraphActions.cpp b/Engine/Plugins/Animation/DeformerGraph/Source/OptimusCore/Private/Actions/OptimusNodeGraphActions.cpp
--- a/Engine/Plugins/Animation/DeformerGraph/Source/OptimusCore/Private/Actions/OptimusNodeGraphActions.cpp
+++ b/Engine/Plugins/Animation/DeformerGraph/Source/OptimusCore/Private/Actions/OptimusNodeGraphActions.cpp
@@ -26,23 +26,21 @@
 // ---- Add graph
 
 FOptimusNodeGraphAction_AddGraph::FOptimusNodeGraphAction_AddGraph(
-	const IOptimusNodeGraphCollectionOwner* InGraphOwner, 
+	const FString& InGraphOwnerPath, 
 	EOptimusNodeGraphType InGraphType, 
 	FName InGraphName, 
 	int32 InGraphIndex,
 	TFunction<bool(UOptimusNodeGraph*)> InConfigureGraphFunc
 	)
 {
-	if (ensure(InGraphOwner))
-	{
-		GraphOwnerPath = InGraphOwner->GetCollectionPath();
-		GraphType = InGraphType;
-		GraphName = InGraphName;
-		GraphIndex = InGraphIndex;
-		ConfigureGraphFunc = InConfigureGraphFunc;
+	GraphOwnerPath = InGraphOwnerPath;
+	GraphType = InGraphType;
+	GraphName = InGraphName;
+	GraphIndex = InGraphIndex;
+	ConfigureGraphFunc = InConfigureGraphFunc;
+
+	SetTitlef(TEXT("Add graph"));
 
-		SetTitlef(TEXT("Add graph"));
-	}
 }
 
 
@@ -89,7 +87,7 @@ bool FOptimusNodeGraphAction_AddGraph::Do(
 		GraphName = Graph->GetFName();
 	}
 
-	GraphPath = Graph->GetGraphPath();
+	GraphPath = Graph->GetCollectionPath();
 	return true;
 }
 
@@ -116,7 +114,7 @@ FOptimusNodeGraphAction_RemoveGraph::FOptimusNodeGraphAction_RemoveGraph(
 {
 	if (ensure(InGraph))
 	{
-		GraphPath = InGraph->GetGraphPath();
+		GraphPath = InGraph->GetCollectionPath();
 		GraphOwnerPath = InGraph->GetCollectionOwner()->GetCollectionPath();
 		GraphType = InGraph->GetGraphType();
 		GraphName = InGraph->GetFName();
@@ -192,7 +190,7 @@ FOptimusNodeGraphAction_RenameGraph::FOptimusNodeGraphAction_RenameGraph(
 {
 	if (ensure(InGraph) && InGraph->GetFName() != InNewName)
 	{
-		GraphPath = InGraph->GetGraphPath();
+		GraphPath = InGraph->GetCollectionPath();
 		GraphOwnerPath = InGraph->GetCollectionOwner()->GetCollectionPath();
 
 		// Ensure the name is unique within our namespace.
@@ -227,7 +225,7 @@ bool FOptimusNodeGraphAction_RenameGraph::Do(
 
 	if (GraphOwner->RenameGraphDirect(Graph, NewGraphName.ToString()))
 	{
-		GraphPath = Graph->GetGraphPath();
+		GraphPath = Graph->GetCollectionPath();
 		return true;
 	}
 
@@ -253,7 +251,7 @@ bool FOptimusNodeGraphAction_RenameGraph::Undo(
 
 	if (GraphOwner->RenameGraphDirect(Graph, OldGraphName.ToString()))
 	{
-		GraphPath = Graph->GetGraphPath();
+		GraphPath = Graph->GetCollectionPath();
 		return true;
 	}
 
@@ -449,7 +447,7 @@ FOptimusNodeGraphAction_RemoveNode::FOptimusNodeGraphAction_RemoveNode(
 	{
 		NodePath = InNode->GetNodePath();
 
-		GraphPath = InNode->GetOwningGraph()->GetGraphPath();
+		GraphPath = InNode->GetOwningGraph()->GetCollectionPath();
 		NodeName = InNode->GetFName();
 		NodeClassPath = InNode->GetClass()->GetPathName();
 
@@ -847,7 +845,7 @@ FOptimusNodeGraphAction_PackageKernelFunction::FOptimusNodeGraphAction_PackageKe
 {
 	if (ensure(InKernelNode))
 	{
-		GraphPath = InKernelNode->GetOwningGraph()->GetGraphPath();
+		GraphPath = InKernelNode->GetOwningGraph()->GetCollectionPath();
 
 		NodeName = InNodeName;
 		NodePosition = InKernelNode->GetGraphPosition();
@@ -957,7 +955,7 @@ FOptimusNodeGraphAction_UnpackageKernelFunction::FOptimusNodeGraphAction_Unpacka
 {
 	if (ensure(InKernelFunction))
 	{
-		GraphPath = InKernelFunction->GetOwningGraph()->GetGraphPath();
+		GraphPath = InKernelFunction->GetOwningGraph()->GetCollectionPath();
 		ClassPath = InKernelFunction->GetClass()->GetPathName();
 		NodeName = InNodeName;
 		NodePosition = InKernelFunction->GetGraphPosition();
diff --git a/Engine/Plugins/Animation/DeformerGraph/Source/OptimusCore/Private/Actions/OptimusNodeGraphActions.h b/Engine/Plugins/Animation/DeformerGraph/Source/OptimusCore/Private/Actions/OptimusNodeGraphActions.h
--- a/Engine/Plugins/Animation/DeformerGraph/Source/OptimusCore/Private/Actions/OptimusNodeGraphActions.h
+++ b/Engine/Plugins/Animation/DeformerGraph/Source/OptimusCore/Private/Actions/OptimusNodeGraphActions.h
@@ -34,7 +34,7 @@ public:
 	FOptimusNodeGraphAction_AddGraph() = default;
 
 	FOptimusNodeGraphAction_AddGraph(
-	    const IOptimusNodeGraphCollectionOwner* InGraphOwner,
+		const FString& InGraphOwnerPath,
 		EOptimusNodeGraphType InGraphType,
 		FName InGraphName,
 		int32 InGraphIndex,
diff --git a/Engine/Plugins/Animation/DeformerGraph/Source/OptimusCore/Private/Nodes/OptimusNode_FunctionReference.cpp b/Engine/Plugins/Animation/DeformerGraph/Source/OptimusCore/Private/Nodes/OptimusNode_FunctionReference.cpp
--- a/Engine/Plugins/Animation/DeformerGraph/Source/OptimusCore/Private/Nodes/OptimusNode_FunctionReference.cpp
+++ b/Engine/Plugins/Animation/DeformerGraph/Source/OptimusCore/Private/Nodes/OptimusNode_FunctionReference.cpp
@@ -71,11 +71,11 @@ FOptimusRoutedNodePin UOptimusNode_FunctionReference::GetPinCounterpart(
 	UOptimusNode_GraphTerminal* CounterpartNode = nullptr;
 	if (InNodePin->GetDirection() == EOptimusNodePinDirection::Input)
 	{
-		CounterpartNode = FunctionGraph->EntryNode.Get();
+		CounterpartNode = FunctionGraph->GetTerminalNode(EOptimusTerminalType::Entry);
 	}
 	else if (InNodePin->GetDirection() == EOptimusNodePinDirection::Output)
 	{
-		CounterpartNode = FunctionGraph->ReturnNode.Get();
+		CounterpartNode = FunctionGraph->GetTerminalNode(EOptimusTerminalType::Return);
 	}
 
 	if (!ensure(CounterpartNode))
diff --git a/Engine/Plugins/Animation/DeformerGraph/Source/OptimusCore/Private/Nodes/OptimusNode_GraphTerminal.cpp b/Engine/Plugins/Animation/DeformerGraph/Source/OptimusCore/Private/Nodes/OptimusNode_GraphTerminal.cpp
--- a/Engine/Plugins/Animation/DeformerGraph/Source/OptimusCore/Private/Nodes/OptimusNode_GraphTerminal.cpp
+++ b/Engine/Plugins/Animation/DeformerGraph/Source/OptimusCore/Private/Nodes/OptimusNode_GraphTerminal.cpp
@@ -76,16 +76,9 @@ void UOptimusNode_GraphTerminal::ConstructNode()
 	case EOptimusTerminalType::Unknown:
 		checkNoEntry();
 	}
-
-	SubscribeToOwningGraph();
 }
 
-void UOptimusNode_GraphTerminal::PostLoad()
-{
-	Super::PostLoad();
 
-	SubscribeToOwningGraph();
-}
 
 void UOptimusNode_GraphTerminal::BeginDestroy()
 {
@@ -149,9 +142,15 @@ UOptimusComponentSourceBinding* UOptimusNode_GraphTerminal::GetDefaultComponentB
 	return ReferenceNode->GetDefaultComponentBinding(PinCounterpart.TraversalContext);
 }
 
+void UOptimusNode_GraphTerminal::InitializeTransientData()
+{
+	OwningGraph = Cast<UOptimusNodeSubGraph>(GetOwningGraph());
+	SubscribeToOwningGraph();
+}
+
 void UOptimusNode_GraphTerminal::SubscribeToOwningGraph()
 {
-	if (ensure(!OwningGraph->GetOnBindingArrayPasted().IsBoundToObject(this)))
+	if (ensure(OwningGraph.IsValid()) && ensure(!OwningGraph->GetOnBindingArrayPasted().IsBoundToObject(this)))
 	{
 		OwningGraph->GetOnBindingArrayPasted().AddUObject(this, &UOptimusNode_GraphTerminal::RecreateBindingPins);
 		OwningGraph->GetOnBindingValueChanged().AddUObject(this, &UOptimusNode_GraphTerminal::SyncPinsToBindings);
diff --git a/Engine/Plugins/Animation/DeformerGraph/Source/OptimusCore/Private/Nodes/OptimusNode_GraphTerminal.h b/Engine/Plugins/Animation/DeformerGraph/Source/OptimusCore/Private/Nodes/OptimusNode_GraphTerminal.h
--- a/Engine/Plugins/Animation/DeformerGraph/Source/OptimusCore/Private/Nodes/OptimusNode_GraphTerminal.h
+++ b/Engine/Plugins/Animation/DeformerGraph/Source/OptimusCore/Private/Nodes/OptimusNode_GraphTerminal.h
@@ -4,6 +4,7 @@
 #include "IOptimusAlternativeSelectedObjectProvider.h"
 #include "IOptimusNodePinRouter.h"
 #include "IOptimusNonCollapsibleNode.h"
+#include "IOptimusNonCopyableNode.h"
 #include "OptimusNode.h"
 
 #include "OptimusNode_GraphTerminal.generated.h"
@@ -26,7 +27,8 @@ class UOptimusNode_GraphTerminal :
 	public UOptimusNode,
 	public IOptimusNodePinRouter,
 	public IOptimusAlternativeSelectedObjectProvider,
-	public IOptimusNonCollapsibleNode
+	public IOptimusNonCollapsibleNode,
+	public IOptimusNonCopyableNode
 {
 	GENERATED_BODY()
 	
@@ -42,7 +44,6 @@ public:
 	void ConstructNode() override;
 
 	// UObject overrides
-	void PostLoad() override;
 	void BeginDestroy() override;
 	
 	// IOptimusNodePinRouter implementation
@@ -60,6 +61,7 @@ protected:
 	friend class UOptimusNodeGraph;
 	friend class UOptimusNodeSubGraph;
 
+	void InitializeTransientData() override;
 	void SubscribeToOwningGraph();
 
 	void UnsubscribeFromOwningGraph() const;
@@ -75,12 +77,12 @@ protected:
 	UPROPERTY()
 	EOptimusTerminalType TerminalType;
 
-	/** The graph that owns us. This contains all the necessary pin information to add on
-	 * the terminal node.
-	 */
-	UPROPERTY()
-	TWeakObjectPtr<UOptimusNodeSubGraph> OwningGraph;
 
 	UPROPERTY()
 	TWeakObjectPtr<UOptimusNodePin> DefaultComponentPin;
+	
+	/** The graph that owns us. This contains all the necessary pin information to add on
+	 * the terminal node. Initialized during InitializeTransientData()
+	 */
+	TWeakObjectPtr<UOptimusNodeSubGraph> OwningGraph;
 };
diff --git a/Engine/Plugins/Animation/DeformerGraph/Source/OptimusCore/Private/Nodes/OptimusNode_LoopTerminal.h b/Engine/Plugins/Animation/DeformerGraph/Source/OptimusCore/Private/Nodes/OptimusNode_LoopTerminal.h
--- a/Engine/Plugins/Animation/DeformerGraph/Source/OptimusCore/Private/Nodes/OptimusNode_LoopTerminal.h
+++ b/Engine/Plugins/Animation/DeformerGraph/Source/OptimusCore/Private/Nodes/OptimusNode_LoopTerminal.h
@@ -66,7 +66,6 @@ public:
 #endif
 	
 	// UOptimusNode overrides
-	bool CanUserDeleteNode() const override { return true; }
 	FName GetNodeCategory() const override { return NAME_None; }
 	FText GetDisplayName() const override;
 	void ConstructNode() override;
diff --git a/Engine/Plugins/Animation/DeformerGraph/Source/OptimusCore/Private/Nodes/OptimusNode_SubGraphReference.cpp b/Engine/Plugins/Animation/DeformerGraph/Source/OptimusCore/Private/Nodes/OptimusNode_SubGraphReference.cpp
--- a/Engine/Plugins/Animation/DeformerGraph/Source/OptimusCore/Private/Nodes/OptimusNode_SubGraphReference.cpp
+++ b/Engine/Plugins/Animation/DeformerGraph/Source/OptimusCore/Private/Nodes/OptimusNode_SubGraphReference.cpp
@@ -84,11 +84,11 @@ FOptimusRoutedNodePin UOptimusNode_SubGraphReference::GetPinCounterpart(
 	UOptimusNode_GraphTerminal* CounterpartNode = nullptr;
 	if (InNodePin->GetDirection() == EOptimusNodePinDirection::Input)
 	{
-		CounterpartNode = SubGraph->EntryNode.Get();
+		CounterpartNode = SubGraph->GetTerminalNode(EOptimusTerminalType::Entry);
 	}
 	else if (InNodePin->GetDirection() == EOptimusNodePinDirection::Output)
 	{
-		CounterpartNode = SubGraph->ReturnNode.Get();
+		CounterpartNode = SubGraph->GetTerminalNode(EOptimusTerminalType::Return);
 	}
 
 	if (!ensure(CounterpartNode))
diff --git a/Engine/Plugins/Animation/DeformerGraph/Source/OptimusCore/Private/OptimusActionStack.cpp b/Engine/Plugins/Animation/DeformerGraph/Source/OptimusCore/Private/OptimusActionStack.cpp
--- a/Engine/Plugins/Animation/DeformerGraph/Source/OptimusCore/Private/OptimusActionStack.cpp
+++ b/Engine/Plugins/Animation/DeformerGraph/Source/OptimusCore/Private/OptimusActionStack.cpp
@@ -183,3 +183,8 @@ FOptimusActionScope::~FOptimusActionScope()
 {
 	ActionStack.CloseActionScope();
 }
+
+void FOptimusActionScope::SetTitle(const FString& InTitle)
+{
+	ActionStack.ActionScopes.Last()->SetTitle(InTitle);
+}
diff --git a/Engine/Plugins/Animation/DeformerGraph/Source/OptimusCore/Private/OptimusDeformer.cpp b/Engine/Plugins/Animation/DeformerGraph/Source/OptimusCore/Private/OptimusDeformer.cpp
--- a/Engine/Plugins/Animation/DeformerGraph/Source/OptimusCore/Private/OptimusDeformer.cpp
+++ b/Engine/Plugins/Animation/DeformerGraph/Source/OptimusCore/Private/OptimusDeformer.cpp
@@ -100,7 +100,7 @@ UOptimusActionStack* UOptimusDeformer::GetActionStack()
 UOptimusNodeGraph* UOptimusDeformer::AddSetupGraph()
 {
 	FOptimusNodeGraphAction_AddGraph* AddGraphAction = 
-		new FOptimusNodeGraphAction_AddGraph(this, EOptimusNodeGraphType::Setup, UOptimusNodeGraph::SetupGraphName, 0);
+		new FOptimusNodeGraphAction_AddGraph(GetCollectionPath(), EOptimusNodeGraphType::Setup, UOptimusNodeGraph::SetupGraphName, 0);
 
 	if (GetActionStack()->RunAction(AddGraphAction))
 	{
@@ -121,7 +121,7 @@ UOptimusNodeGraph* UOptimusDeformer::AddTriggerGraph(const FString &InName)
 	}
 
 	FOptimusNodeGraphAction_AddGraph* AddGraphAction =
-	    new FOptimusNodeGraphAction_AddGraph(this, EOptimusNodeGraphType::ExternalTrigger, *InName, INDEX_NONE);
+	    new FOptimusNodeGraphAction_AddGraph(GetCollectionPath(), EOptimusNodeGraphType::ExternalTrigger, *InName, INDEX_NONE);
 
 	if (GetActionStack()->RunAction(AddGraphAction))
 	{
@@ -150,21 +150,24 @@ UOptimusNodeGraph* UOptimusDeformer::GetUpdateGraph() const
 
 bool UOptimusDeformer::RemoveGraph(UOptimusNodeGraph* InGraph)
 {
-	FOptimusCompoundAction *Action = new FOptimusCompoundAction(TEXT("Remove Graph"));
-
 	if (UOptimusNodeSubGraph* SubGraph = Cast<UOptimusNodeSubGraph>(InGraph))
 	{
 		if (UOptimusNode* Node = GetSubGraphReferenceNode(SubGraph))
 		{
+			FOptimusActionScope ActionScope(*GetActionStack(), TEXT("Remove SubGraph"));
 			// Remove node also triggers the removal of the graph
-			Node->GetOwningGraph()->RemoveNodesToAction(Action, {Node} );
-			return GetActionStack()->RunAction(Action);
+			Node->GetOwningGraph()->RemoveNode(Node);
+			return true;
 		}
+		
+		return false;
 	}
+
+	FOptimusActionScope ActionScope(*GetActionStack(), TEXT("Remove Graph"));
+	InGraph->RemoveNodes(InGraph->GetAllNodes());
+	GetActionStack()->RunAction<FOptimusNodeGraphAction_RemoveGraph>(InGraph);
 	
-	Action->AddSubAction<FOptimusNodeGraphAction_RemoveGraph>(InGraph);
-	
-    return GetActionStack()->RunAction(Action);
+	return true;
 }
 
 UOptimusNode* UOptimusDeformer::GetSubGraphReferenceNode(const UOptimusNodeSubGraph* InSubGraph) const
@@ -260,25 +263,25 @@ bool UOptimusDeformer::RemoveVariable(
 		UE_LOG(LogOptimusCore, Error, TEXT("Variable not owned by this deformer."));
 		return false;
 	}
-
-	FOptimusCompoundAction* Action = new FOptimusCompoundAction(TEXT("Remove Variable"));
 	
-	TMap<const UOptimusNodeGraph*, TArray<UOptimusNode*>> NodesByGraph;
+	TMap<UOptimusNodeGraph*, TArray<UOptimusNode*>> NodesByGraph;
 	for (UOptimusNode* Node: GetNodesUsingVariable(InVariableDesc))
 	{
 		UOptimusNode_GetVariable* VariableNode = Cast<UOptimusNode_GetVariable>(Node);
 		NodesByGraph.FindOrAdd(VariableNode->GetOwningGraph()).Add(VariableNode);
 	}
 
-	for (const TTuple<const UOptimusNodeGraph*, TArray<UOptimusNode*>>& GraphNodes: NodesByGraph)
+	FOptimusActionScope ActionScope(*GetActionStack(), TEXT("Remove Variable"));
+	
+	for (const TTuple<UOptimusNodeGraph*, TArray<UOptimusNode*>>& GraphNodes: NodesByGraph)
 	{
-		const UOptimusNodeGraph* Graph = GraphNodes.Key;
-		Graph->RemoveNodesToAction(Action, GraphNodes.Value);
+		UOptimusNodeGraph* Graph = GraphNodes.Key;
+		Graph->RemoveNodes(GraphNodes.Value);
 	}
 
-	Action->AddSubAction<FOptimusVariableAction_RemoveVariable>(InVariableDesc);
+	GetActionStack()->RunAction<FOptimusVariableAction_RemoveVariable>(InVariableDesc);
 
-	return GetActionStack()->RunAction(Action);
+	return true;
 }
 
 
@@ -633,24 +636,24 @@ bool UOptimusDeformer::RemoveResource(UOptimusResourceDescription* InResourceDes
 		return false;
 	}
 
-	FOptimusCompoundAction* Action = new FOptimusCompoundAction(TEXT("Remove Resource"));
-	
-	TMap<const UOptimusNodeGraph*, TArray<UOptimusNode*>> NodesByGraph;
+	TMap<UOptimusNodeGraph*, TArray<UOptimusNode*>> NodesByGraph;
 	for (UOptimusNode* Node: GetNodesUsingResource(InResourceDesc))
 	{
 		UOptimusNode_ResourceAccessorBase* ResourceNode = Cast<UOptimusNode_ResourceAccessorBase>(Node);
 		NodesByGraph.FindOrAdd(ResourceNode->GetOwningGraph()).Add(ResourceNode);
 	}
 
-	for (const TTuple<const UOptimusNodeGraph*, TArray<UOptimusNode*>>& GraphNodes: NodesByGraph)
+	FOptimusActionScope ActionScope(*GetActionStack(), TEXT("Remove Resource"));
+	
+	for (const TTuple<UOptimusNodeGraph*, TArray<UOptimusNode*>>& GraphNodes: NodesByGraph)
 	{
-		const UOptimusNodeGraph* Graph = GraphNodes.Key;
-		Graph->RemoveNodesToAction(Action, GraphNodes.Value);
+		UOptimusNodeGraph* Graph = GraphNodes.Key;
+		Graph->RemoveNodes(GraphNodes.Value);
 	}
 
-	Action->AddSubAction<FOptimusResourceAction_RemoveResource>(InResourceDesc);
+	GetActionStack()->RunAction<FOptimusResourceAction_RemoveResource>(InResourceDesc);
 
-	return GetActionStack()->RunAction(Action);
+	return true;
 }
 
 
@@ -1132,24 +1135,24 @@ bool UOptimusDeformer::RemoveComponentBinding(
 		return false;
 	}
 
-	FOptimusCompoundAction* Action = new FOptimusCompoundAction(TEXT("Remove Binding"));
-	
-	TMap<const UOptimusNodeGraph*, TArray<UOptimusNode*>> NodesByGraph;
+	TMap<UOptimusNodeGraph*, TArray<UOptimusNode*>> NodesByGraph;
 	
 	for (UOptimusNode* Node: GetNodesUsingComponentBinding(InBinding))
 	{
 		NodesByGraph.FindOrAdd(Node->GetOwningGraph()).Add(Node);
 	}
 
-	for (const TTuple<const UOptimusNodeGraph*, TArray<UOptimusNode*>>& GraphNodes: NodesByGraph)
+	FOptimusActionScope ActionScope(*GetActionStack(), TEXT("Remove Binding"));
+
+	for (const TTuple<UOptimusNodeGraph*, TArray<UOptimusNode*>>& GraphNodes: NodesByGraph)
 	{
-		const UOptimusNodeGraph* Graph = GraphNodes.Key;
-		Graph->RemoveNodesToAction(Action, GraphNodes.Value);
+		UOptimusNodeGraph* Graph = GraphNodes.Key;
+		Graph->RemoveNodes(GraphNodes.Value);
 	}
 
-	Action->AddSubAction<FOptimusComponentBindingAction_RemoveBinding>(InBinding);
+	GetActionStack()->RunAction<FOptimusComponentBindingAction_RemoveBinding>(InBinding);
 
-	return GetActionStack()->RunAction(Action);
+	return true;
 }
 
 
@@ -1684,7 +1687,10 @@ TArray<FOptimusComputeGraphInfo> UOptimusDeformer::CompileNodeGraphToComputeGrap
 
 	if (TerminalNodes.IsEmpty())
 	{
-		AddDiagnostic(EOptimusDiagnosticLevel::Error, LOCTEXT("NoOutputDataInterfaceFound", "No connected output data interface nodes found. Compilation aborted."));
+		FText WarnMessage = FText::Format(LOCTEXT("NoOutputDataInterfaceFound", "No connected output data interface nodes found. Compilation for Graph: {0} aborted."), 
+			FText::FromString(InNodeGraph->GetCollectionPath()));
+		
+		AddDiagnostic(EOptimusDiagnosticLevel::Warning, WarnMessage);
 		return {};
 	}
 
@@ -2887,32 +2893,14 @@ UOptimusNodeGraph* UOptimusDeformer::ResolveGraphPath(
 	int32 SubGraphStartIndex = 0;
 	
 	UOptimusNodeGraph* Graph = nullptr;
-	if (Path[0] == UOptimusNodeGraph::LibraryRoot)
-	{
-		// FIXME: Search the library graphs.
-		for (UOptimusNodeGraph* RootGraph : Graphs)
-		{
-			if (RootGraph->GetGraphType() == EOptimusNodeGraphType::Function)
-			{
-				if (Path[1].Equals(RootGraph->GetName(), ESearchCase::IgnoreCase))
-				{
-					SubGraphStartIndex = 2;
-					Graph = RootGraph;
-					break;
-				}
-			}
-		}
-	}
-	else
+	
+	for (UOptimusNodeGraph* RootGraph : Graphs)
 	{
-		for (UOptimusNodeGraph* RootGraph : Graphs)
+		if (Path[0].Equals(RootGraph->GetName(), ESearchCase::IgnoreCase))
 		{
-			if (Path[0].Equals(RootGraph->GetName(), ESearchCase::IgnoreCase))
-			{
-				SubGraphStartIndex = 1;
-				Graph = RootGraph;
-				break;
-			}
+			SubGraphStartIndex = 1;
+			Graph = RootGraph;
+			break;
 		}
 	}
 
diff --git a/Engine/Plugins/Animation/DeformerGraph/Source/OptimusCore/Private/OptimusNode.cpp b/Engine/Plugins/Animation/DeformerGraph/Source/OptimusCore/Private/OptimusNode.cpp
--- a/Engine/Plugins/Animation/DeformerGraph/Source/OptimusCore/Private/OptimusNode.cpp
+++ b/Engine/Plugins/Animation/DeformerGraph/Source/OptimusCore/Private/OptimusNode.cpp
@@ -123,7 +123,7 @@ FString UOptimusNode::GetNodePath() const
 	FString GraphPath(TEXT("<Unknown>"));
 	if (Graph)
 	{
-		GraphPath = Graph->GetGraphPath();
+		GraphPath = Graph->GetCollectionPath();
 	}
 
 	return UOptimusNodeGraph::ConstructPath(GraphPath, GetName(), {});
diff --git a/Engine/Plugins/Animation/DeformerGraph/Source/OptimusCore/Private/OptimusNodeGraph.cpp b/Engine/Plugins/Animation/DeformerGraph/Source/OptimusCore/Private/OptimusNodeGraph.cpp
--- a/Engine/Plugins/Animation/DeformerGraph/Source/OptimusCore/Private/OptimusNodeGraph.cpp
+++ b/Engine/Plugins/Animation/DeformerGraph/Source/OptimusCore/Private/OptimusNodeGraph.cpp
@@ -55,13 +55,13 @@
 
 const FName UOptimusNodeGraph::SetupGraphName("SetupGraph");
 const FName UOptimusNodeGraph::UpdateGraphName("UpdateGraph");
-const TCHAR* UOptimusNodeGraph::LibraryRoot = TEXT("@Library");
 const FName UOptimusNodeGraph::DefaultSubGraphName("SubGraph");
 const FName UOptimusNodeGraph::DefaultSubGraphRefNodeName("SubGraphNode");
 
-FString UOptimusNodeGraph::GetFunctionGraphPath(const FString& InFunctionName)
+FString UOptimusNodeGraph::GetFunctionGraphCollectionPath(const FString& InFunctionName)
 {
-	return FString::Printf(TEXT("%s/%s"), LibraryRoot, *InFunctionName);
+	// Function Graph is always at the top level
+	return InFunctionName;
 }
 
 void UOptimusNodeGraph::PostDuplicate(EDuplicateMode::Type DuplicateMode)
@@ -192,53 +192,6 @@ UOptimusNodeGraph* UOptimusNodeGraph::GetParentGraph() const
 }
 
 
-FString UOptimusNodeGraph::GetGraphPath() const
-{
-	// The path is local to the root of the node graph collection owner hierarchy, excluding the
-	// root collection itself (e.g. the deformer asset).
-	switch(GraphType)
-	{
-	case EOptimusNodeGraphType::Setup:
-	case EOptimusNodeGraphType::Update:
-	case EOptimusNodeGraphType::ExternalTrigger:
-		return GetName();
-
-	case EOptimusNodeGraphType::Function:
-		return GetFunctionGraphPath(GetName());
-		
-	case EOptimusNodeGraphType::SubGraph:
-		{
-			TArray<FString, TInlineAllocator<8>> Ancestry;
-
-			const UOptimusNodeGraph* CurrentGraph = this;
-			while(CurrentGraph)
-			{
-				if (CurrentGraph->GetGraphType() == EOptimusNodeGraphType::SubGraph)
-				{
-					Ancestry.Add(CurrentGraph->GetName());
-				}
-				else
-				{
-					Ancestry.Add(CurrentGraph->GetGraphPath());
-					break;
-				}
-				CurrentGraph = Cast<const UOptimusNodeGraph>(CurrentGraph->GetOuter());
-			}
-			Algo::Reverse(Ancestry);
-
-			TStringBuilder<256> Path;
-			return Path.Join(Ancestry, TEXT('/')).ToString();
-		}
-		
-	case EOptimusNodeGraphType::Transient:
-		return TEXT("Transient");
-	}
-
-	checkNoEntry();
-	return TEXT("<unknown graph type>");
-}
-
-
 bool UOptimusNodeGraph::IsValidUserGraphName(
 	const FString& InGraphName,
 	FText* OutFailureReason
@@ -307,7 +260,7 @@ UOptimusNode* UOptimusNodeGraph::AddNodeInternal(
 	// FIXME: Need better naming.
 	FName NodeName = Optimus::GetUniqueNameForScope(this, InNodeClass->GetFName());
 	FOptimusNodeGraphAction_AddNode *AddNodeAction = new FOptimusNodeGraphAction_AddNode(
-		GetGraphPath(), InNodeClass, NodeName,
+		GetCollectionPath(), InNodeClass, NodeName,
 		[InNodeConfigFunc, InPosition](UOptimusNode *InNode) {
 			if (InNodeConfigFunc)
 			{
@@ -331,7 +284,7 @@ TArray<UOptimusNode*> UOptimusNodeGraph::AddNodePairInternal(const TSubclassOf<U
 	FOptimusCompoundAction *Action = new FOptimusCompoundAction(TEXT("Add Loop Nodes"));
 	
 	FOptimusNodeGraphAction_AddNode *AddFirstNodeAction = new FOptimusNodeGraphAction_AddNode(
-		GetGraphPath(), InNodeClass, NodeName,
+		GetCollectionPath(), InNodeClass, NodeName,
 		[InFirstNodeConfigFunc, InPosition](UOptimusNode *InNode) {
 			if (InFirstNodeConfigFunc)
 			{
@@ -345,7 +298,7 @@ TArray<UOptimusNode*> UOptimusNodeGraph::AddNodePairInternal(const TSubclassOf<U
 	NodeName = Optimus::GetUniqueNameForScope(this, *(InNodeClass->GetName() + TEXT("_Exit")));
 	
 	FOptimusNodeGraphAction_AddNode *AddSecondNodeAction = new FOptimusNodeGraphAction_AddNode(
-			GetGraphPath(), InNodeClass, NodeName,
+			GetCollectionPath(), InNodeClass, NodeName,
 			[InSecondNodeConfigFunc, InPosition](UOptimusNode *InNode) {
 				if (InSecondNodeConfigFunc)
 				{
@@ -395,25 +348,30 @@ void UOptimusNodeGraph::GatherRelatedObjects(
 	}	
 }
 
-bool UOptimusNodeGraph::DuplicateSubGraph(UOptimusNodeSubGraph* InSourceSubGraph, FName InNewGraphName)
+bool UOptimusNodeGraph::DuplicateSubGraph(
+	UOptimusActionStack* InActionStack,
+	const FString& InGraphOwnerPath,
+	UOptimusNodeSubGraph* InSourceSubGraph,
+	FName InNewGraphName
+	)
 {
-	IOptimusPathResolver* PathResolver = GetPathResolver();
-
+	
 	// Ensure the name is unique here instead of relying on AddGraph action
 	// because subsequent actions needs to use the final graph path
 	// Note: we could change the actions to accept a lambda instead of specific value in the future
 	FName SubGraphName = InNewGraphName;
 	
-	const FString SubGraphPath = ConstructSubGraphPath(SubGraphName.ToString());
-	const FVector2D EntryNodePosition = InSourceSubGraph->EntryNode->GetGraphPosition();
-	const FVector2D ReturnNodePosition = InSourceSubGraph->ReturnNode->GetGraphPosition();
-	const FString CurrentGraphPath = GetGraphPath();
+	const FString SubGraphPath = ConstructSubGraphPath(InGraphOwnerPath, SubGraphName.ToString());
+	UOptimusNode_GraphTerminal* SourceSubGraphEntryNode = InSourceSubGraph->GetTerminalNode(EOptimusTerminalType::Entry);
+	UOptimusNode_GraphTerminal* SourceSubGraphReturnNode = InSourceSubGraph->GetTerminalNode(EOptimusTerminalType::Return);
+	const FVector2D EntryNodePosition = SourceSubGraphEntryNode->GetGraphPosition();
+	const FVector2D ReturnNodePosition = SourceSubGraphReturnNode->GetGraphPosition();
 
 	TMap<FString, TArray<FString>> InternalLinksToAdd;
 	
 	for (UOptimusNodeLink* Link : InSourceSubGraph->GetAllLinks())
 	{
-		if (Link->GetNodeInputPin()->GetOwningNode() == InSourceSubGraph->ReturnNode)
+		if (Link->GetNodeInputPin()->GetOwningNode() == SourceSubGraphReturnNode)
 		{
 			FString ReturnNodeInputPinPath = ConstructPath(
 				SubGraphPath,
@@ -429,7 +387,7 @@ bool UOptimusNodeGraph::DuplicateSubGraph(UOptimusNodeSubGraph* InSourceSubGraph
 
 			InternalLinksToAdd.FindOrAdd(InternalNodeOutputPinPath).Add(ReturnNodeInputPinPath);
 		}
-		else if (Link->GetNodeOutputPin()->GetOwningNode() == InSourceSubGraph->EntryNode)
+		else if (Link->GetNodeOutputPin()->GetOwningNode() == SourceSubGraphEntryNode)
 		{
 			FString EntryNodeOutputPinPath = ConstructPath(
 				SubGraphPath,
@@ -489,8 +447,8 @@ bool UOptimusNodeGraph::DuplicateSubGraph(UOptimusNodeSubGraph* InSourceSubGraph
 	}
 
 	{
-		FOptimusActionScope(*GetActionStack(), TEXT("Duplicate SubGraph"));
-		GetActionStack()->RunAction<FOptimusNodeGraphAction_AddGraph>(this, EOptimusNodeGraphType::SubGraph, SubGraphName, INDEX_NONE,
+		FOptimusActionScope(*InActionStack, TEXT("Duplicate SubGraph"));
+		InActionStack->RunAction<FOptimusNodeGraphAction_AddGraph>(InGraphOwnerPath, EOptimusNodeGraphType::SubGraph, SubGraphName, INDEX_NONE,
 			[InputBindings = InSourceSubGraph->InputBindings, OutputBindings = InSourceSubGraph->OutputBindings](UOptimusNodeGraph* InGraph)
 			{
 				UOptimusNodeSubGraph* SubGraph = CastChecked<UOptimusNodeSubGraph>(InGraph);
@@ -499,43 +457,33 @@ bool UOptimusNodeGraph::DuplicateSubGraph(UOptimusNodeSubGraph* InSourceSubGraph
 				return true;
 			});
 	
-		GetActionStack()->RunAction<FOptimusNodeGraphAction_AddNode>(
+		InActionStack->RunAction<FOptimusNodeGraphAction_AddNode>(
 			SubGraphPath, UOptimusNode_GraphTerminal::StaticClass(), UOptimusNode_GraphTerminal::EntryNodeName,
-			[SubGraphPath, EntryNodePosition, PathResolver](UOptimusNode* InNode)
+			[EntryNodePosition](UOptimusNode* InNode)
 			{
 				UOptimusNode_GraphTerminal* EntryNode = Cast<UOptimusNode_GraphTerminal>(InNode); 
-				UOptimusNodeSubGraph* SubGraph = Cast<UOptimusNodeSubGraph>(PathResolver->ResolveGraphPath(SubGraphPath));
-
-				SubGraph->EntryNode = EntryNode;
-		
 				EntryNode->TerminalType = EOptimusTerminalType::Entry;
-				EntryNode->OwningGraph = SubGraph;
 				return EntryNode->SetGraphPositionDirect(EntryNodePosition);
 			});
 
-		GetActionStack()->RunAction<FOptimusNodeGraphAction_AddNode>(
+		InActionStack->RunAction<FOptimusNodeGraphAction_AddNode>(
 			SubGraphPath, UOptimusNode_GraphTerminal::StaticClass(), UOptimusNode_GraphTerminal::ReturnNodeName,
-			[SubGraphPath, ReturnNodePosition, PathResolver](UOptimusNode* InNode)
+			[ReturnNodePosition](UOptimusNode* InNode)
 			{
 				UOptimusNode_GraphTerminal* ReturnNode = Cast<UOptimusNode_GraphTerminal>(InNode); 
-				UOptimusNodeSubGraph* SubGraph = Cast<UOptimusNodeSubGraph>(PathResolver->ResolveGraphPath(SubGraphPath));
-
-				SubGraph->ReturnNode = ReturnNode;
-			
 				ReturnNode->TerminalType = EOptimusTerminalType::Return;
-				ReturnNode->OwningGraph = SubGraph;
 				return ReturnNode->SetGraphPositionDirect(ReturnNodePosition);
 			});
 
 		for (UOptimusNodeSubGraph* SubGraphToDuplicate : SubGraphsToDuplicate)
 		{
-			DuplicateSubGraph(SubGraphToDuplicate, SubGraphToDuplicate->GetFName());
+			DuplicateSubGraph(InActionStack, SubGraphPath, SubGraphToDuplicate, SubGraphToDuplicate->GetFName());
 		}
 		
 		// Copy the nodes to the sub graph
 		for (UOptimusNode* Node: NodesToDuplicate)
 		{
-			GetActionStack()->RunAction<FOptimusNodeGraphAction_DuplicateNode>(
+			InActionStack->RunAction<FOptimusNodeGraphAction_DuplicateNode>(
 				SubGraphPath, Node, Node->GetFName(), [](UOptimusNode*)
 				{
 					return true;
@@ -544,7 +492,7 @@ bool UOptimusNodeGraph::DuplicateSubGraph(UOptimusNodeSubGraph* InSourceSubGraph
 
 		for (const TPair<FString, FString>& NewPair : NewNodePairs)
 		{
-			GetActionStack()->RunAction<FOptimusNodeGraphAction_AddNodePair>(NewPair.Key, NewPair.Value);
+			InActionStack->RunAction<FOptimusNodeGraphAction_AddNodePair>(NewPair.Key, NewPair.Value);
 		}	
 		
 		// Restore links
@@ -553,7 +501,7 @@ bool UOptimusNodeGraph::DuplicateSubGraph(UOptimusNodeSubGraph* InSourceSubGraph
 			const FString& OutputPinPath = InternalLinks.Key;
 			for (const FString& InputPinPath : InternalLinks.Value)
 			{
-				GetActionStack()->RunAction<FOptimusNodeGraphAction_AddLink>(OutputPinPath,	InputPinPath);
+				InActionStack->RunAction<FOptimusNodeGraphAction_AddLink>(OutputPinPath, InputPinPath);
 			}
 		}
 	}
@@ -608,7 +556,7 @@ TArray<UOptimusNode*> UOptimusNodeGraph::AddLoopTerminalNodes(const FVector2D& I
 
 	FName NodeName = Optimus::GetUniqueNameForScope(this, *(UOptimusNode_LoopTerminal::StaticClass()->GetName() + TEXT("_Entry")));
 	FOptimusNodeGraphAction_AddNode *AddEntryNodeAction = new FOptimusNodeGraphAction_AddNode(
-		GetGraphPath(), UOptimusNode_LoopTerminal::StaticClass(), NodeName,
+		GetCollectionPath(), UOptimusNode_LoopTerminal::StaticClass(), NodeName,
 		[InPosition](UOptimusNode *InNode) {
 			UOptimusNode_LoopTerminal* Node = CastChecked<UOptimusNode_LoopTerminal>(InNode);
 			Node->TerminalType = EOptimusTerminalType::Entry;
@@ -621,7 +569,7 @@ TArray<UOptimusNode*> UOptimusNodeGraph::AddLoopTerminalNodes(const FVector2D& I
 	NodeName = Optimus::GetUniqueNameForScope(this, *(UOptimusNode_LoopTerminal::StaticClass()->GetName() + TEXT("_Return"))); 
 
 	FOptimusNodeGraphAction_AddNode *AddReturnNodeAction = new FOptimusNodeGraphAction_AddNode(
-		GetGraphPath(), UOptimusNode_LoopTerminal::StaticClass(), NodeName,
+		GetCollectionPath(), UOptimusNode_LoopTerminal::StaticClass(), NodeName,
 		[InPosition](UOptimusNode *InNode) {
 			UOptimusNode_LoopTerminal* Return = CastChecked<UOptimusNode_LoopTerminal>(InNode);
 			Return->TerminalType = EOptimusTerminalType::Return;
@@ -734,44 +682,16 @@ bool UOptimusNodeGraph::RemoveNode(UOptimusNode* InNode)
 
 bool UOptimusNodeGraph::RemoveNodes(const TArray<UOptimusNode*> &InNodes)
 {
-	return RemoveNodes(InNodes, TEXT("Remove"));
-}
-
-
-bool UOptimusNodeGraph::RemoveNodes(
-		const TArray<UOptimusNode*>& InNodes,
-		const FString& InActionName
-		)
-{
-	FOptimusCompoundAction* Action = new FOptimusCompoundAction;
-	if (InNodes.Num() == 1)
-	{
-		Action->SetTitlef(TEXT("%s Node"), *InActionName);
-	}
-	else
-	{
-		Action->SetTitlef(TEXT("%s %d Nodes"), *InActionName, InNodes.Num());
-	}
-
-	if (RemoveNodesToAction(Action, InNodes))
-	{
-		return GetActionStack()->RunAction(Action);
-	}
-
-	delete Action;
-	return false;
+	return RemoveNodesAndCount(InNodes) > 0;
 }
 
 
-bool UOptimusNodeGraph::RemoveNodesToAction(
-	FOptimusCompoundAction* InOutAction,
-	const TArray<UOptimusNode*>& InNodes
-	) const
+int32 UOptimusNodeGraph::RemoveNodesAndCount(const TArray<UOptimusNode*> &InNodes)
 {
 	// Validate the input set.
 	if (InNodes.Num() == 0)
 	{
-		return false;
+		return 0;
 	}
 
 	TArray<UOptimusNode*> NodesToRemove;
@@ -783,7 +703,7 @@ bool UOptimusNodeGraph::RemoveNodesToAction(
 	{
 		if (Node == nullptr || Node->GetOwningGraph() != this)
 		{
-			return false;
+			return 0;
 		}
 	}
 
@@ -795,27 +715,45 @@ bool UOptimusNodeGraph::RemoveNodesToAction(
 		AllLinkIndexes.Append(GetAllLinkIndexesToNode(Node));
 	}
 
+	const FString ActionTitle = NodesToRemove.Num() == 1 ? TEXT("Remove Node") : FString::Printf(TEXT("Remove %d Nodes"), NodesToRemove.Num()) ;
+	
+	FOptimusActionScope ActionScope(*GetActionStack(), ActionTitle);	
+	
 	for (const int32 LinkIndex : AllLinkIndexes)
 	{
-		InOutAction->AddSubAction<FOptimusNodeGraphAction_RemoveLink>(Links[LinkIndex]);
+		GetActionStack()->RunAction<FOptimusNodeGraphAction_RemoveLink>(Links[LinkIndex]);
 	}
 
 	for (const UOptimusNodePair* NodePair: NodePairsToRemove)
 	{
-		InOutAction->AddSubAction<FOptimusNodeGraphAction_RemoveNodePair>(NodePair);
+		GetActionStack()->RunAction<FOptimusNodeGraphAction_RemoveNodePair>(NodePair);
 	}	
 	
 	for (UOptimusNode* Node : NodesToRemove)
 	{
-		InOutAction->AddSubAction<FOptimusNodeGraphAction_RemoveNode>(Node);
+		GetActionStack()->RunAction<FOptimusNodeGraphAction_RemoveNode>(Node);
 	}
 
 	for (UOptimusNodeSubGraph* SubGraph : SubGraphsToRemove)
 	{
-		InOutAction->AddSubAction<FOptimusNodeGraphAction_RemoveGraph>(SubGraph);
+		RemoveGraph(SubGraph);
 	}
 
-	return true;
+	return NodesToRemove.Num();
+}
+
+void UOptimusNodeGraph::RemoveGraph(UOptimusNodeGraph* InNodeGraph)
+{
+	if (!InNodeGraph)
+	{
+		return;
+	}
+
+	FOptimusActionScope ActionScope(*GetActionStack(), TEXT("Remove Graph"));
+
+	InNodeGraph->RemoveNodes(InNodeGraph->GetAllNodes());
+
+	GetActionStack()->RunAction<FOptimusNodeGraphAction_RemoveGraph>(InNodeGraph);
 }
 
 UOptimusNode* UOptimusNodeGraph::DuplicateNode(
@@ -838,7 +776,7 @@ UOptimusNode* UOptimusNodeGraph::DuplicateNode(
 	const FName NodeName = Optimus::GetUniqueNameForScope(this, InNode->GetFName());
 	
 	FOptimusNodeGraphAction_DuplicateNode *DuplicateNodeAction = new FOptimusNodeGraphAction_DuplicateNode(
-		GetGraphPath(), InNode, NodeName,
+		GetCollectionPath(), InNode, NodeName,
 		[InPosition](UOptimusNode *InNode) {
 			return InNode->SetGraphPositionDirect(InPosition); 
 		});
@@ -905,7 +843,7 @@ bool UOptimusNodeGraph::DuplicateNodes(
 	}
 	FVector2D NodeOffset = InPosition - TopLeft;
 
-	const FString GraphPath = GetGraphPath();
+	const FString GraphPath = GetCollectionPath();
 	
 	TArray<TPair<FString, FString>> NewNodePairs;
 	for (const UOptimusNodePair* NodePair : NodePairsToDuplicate)
@@ -964,7 +902,7 @@ bool UOptimusNodeGraph::DuplicateNodes(
 		
 		for (UOptimusNodeSubGraph* SubGraphToDuplicate: SubGraphsToDuplicate)
 		{
-			DuplicateSubGraph(SubGraphToDuplicate, NewSubGraphNameMap[SubGraphToDuplicate->GetFName()]);
+			DuplicateSubGraph(GetActionStack(), GetCollectionPath(), SubGraphToDuplicate, NewSubGraphNameMap[SubGraphToDuplicate->GetFName()]);
 		}
 
 		{
@@ -982,7 +920,7 @@ bool UOptimusNodeGraph::DuplicateNodes(
 		for (UOptimusNode* Node: NodesToDuplicate)
 		{
 			FOptimusNodeGraphAction_DuplicateNode *DuplicateNodeAction = new FOptimusNodeGraphAction_DuplicateNode(
-				GetGraphPath(), Node, NewNodeNameMap[Node],
+				GetCollectionPath(), Node, NewNodeNameMap[Node],
 				[SourceNodeGraphPosition = Node->GraphPosition, NodeOffset, NewSubGraphNameMap](UOptimusNode *InNode)
 				{
 					if (UOptimusNode_SubGraphReference* SubGraphNode = Cast<UOptimusNode_SubGraphReference>(InNode))
@@ -1349,182 +1287,161 @@ UOptimusNode* UOptimusNodeGraph::CollapseNodesToSubGraph(
 		}
 	}
 
-	FOptimusCompoundAction *Action = new FOptimusCompoundAction(TEXT("Collapse Nodes to Sub-graph"));
-	IOptimusPathResolver* PathResolver = GetPathResolver();
-
-	FName SubGraphName = DefaultSubGraphName;
-	SubGraphName = Optimus::GetUniqueNameForScope(this, SubGraphName);
+	TWeakPtr<FOptimusNodeGraphAction_AddNode> WeakAddSubGraphRefNodeAction;
 	
-	FOptimusNodeGraphAction_AddGraph* CreateGraph = new FOptimusNodeGraphAction_AddGraph(
-		this, EOptimusNodeGraphType::SubGraph, SubGraphName, INDEX_NONE,
-		[InputBindings, OutputBindings](UOptimusNodeGraph* InGraph) -> bool
-		{
-			UOptimusNodeSubGraph* SubGraph = Cast<UOptimusNodeSubGraph>(InGraph);
-			SubGraph->InputBindings = InputBindings;
-			SubGraph->OutputBindings = OutputBindings;
-			return true;
-		});
-	Action->AddSubAction(CreateGraph);
-
-	FString SubGraphPath = ConstructSubGraphPath(SubGraphName.ToString());
-
-	// Create the entry and return nodes.
-	FBox2D NodeBox(ForceInit);
-	for (const UOptimusNode* Node: NodesToDuplicate)
-	{
-		NodeBox += Node->GetGraphPosition();
-	}
 	
-	Action->AddSubAction<FOptimusNodeGraphAction_AddNode>(
-		SubGraphPath, UOptimusNode_GraphTerminal::StaticClass(), UOptimusNode_GraphTerminal::EntryNodeName,
-		[NodeBox, SubGraphPath, PathResolver](UOptimusNode* InNode)
-		{
-			UOptimusNode_GraphTerminal* EntryNode = Cast<UOptimusNode_GraphTerminal>(InNode); 
-			UOptimusNodeSubGraph* SubGraph = Cast<UOptimusNodeSubGraph>(PathResolver->ResolveGraphPath(SubGraphPath));
+	{
+		FOptimusActionScope ActionScope(*GetActionStack(), TEXT("Collapse Nodes to Sub-graph"));
+		
+		FName SubGraphName = DefaultSubGraphName;
+		SubGraphName = Optimus::GetUniqueNameForScope(this, SubGraphName);
+		
+		FOptimusNodeGraphAction_AddGraph* CreateGraph = new FOptimusNodeGraphAction_AddGraph(
+			GetCollectionPath(), EOptimusNodeGraphType::SubGraph, SubGraphName, INDEX_NONE,
+			[InputBindings, OutputBindings](UOptimusNodeGraph* InGraph) -> bool
+			{
+				UOptimusNodeSubGraph* SubGraph = Cast<UOptimusNodeSubGraph>(InGraph);
+				SubGraph->InputBindings = InputBindings;
+				SubGraph->OutputBindings = OutputBindings;
+				return true;
+			});
+		GetActionStack()->RunAction(CreateGraph);
 
-			SubGraph->EntryNode = EntryNode;
-			
-			EntryNode->TerminalType = EOptimusTerminalType::Entry;
-			EntryNode->OwningGraph = SubGraph;
-			return EntryNode->SetGraphPositionDirect({NodeBox.Min.X - 450.0f, NodeBox.GetCenter().Y});
-		});
+		FString SubGraphPath = ConstructSubGraphPath(SubGraphName.ToString());
 
-	Action->AddSubAction<FOptimusNodeGraphAction_AddNode>(
-		SubGraphPath, UOptimusNode_GraphTerminal::StaticClass(), UOptimusNode_GraphTerminal::ReturnNodeName,
-		[NodeBox, SubGraphPath, PathResolver](UOptimusNode* InNode)
+		// Create the entry and return nodes.
+		FBox2D NodeBox(ForceInit);
+		for (const UOptimusNode* Node: NodesToDuplicate)
 		{
-			UOptimusNode_GraphTerminal* ReturnNode = Cast<UOptimusNode_GraphTerminal>(InNode); 
-			UOptimusNodeSubGraph* SubGraph = Cast<UOptimusNodeSubGraph>(PathResolver->ResolveGraphPath(SubGraphPath));
-
-			SubGraph->ReturnNode = ReturnNode;
-			
-			ReturnNode->TerminalType = EOptimusTerminalType::Return;
-			ReturnNode->OwningGraph = SubGraph;
-			return ReturnNode->SetGraphPositionDirect({NodeBox.Max.X + 750.0f, NodeBox.GetCenter().Y});
-		});
+			NodeBox += Node->GetGraphPosition();
+		}
+		
+		GetActionStack()->RunAction<FOptimusNodeGraphAction_AddNode>(
+			SubGraphPath, UOptimusNode_GraphTerminal::StaticClass(), UOptimusNode_GraphTerminal::EntryNodeName,
+			[NodeBox](UOptimusNode* InNode)
+			{
+				UOptimusNode_GraphTerminal* EntryNode = Cast<UOptimusNode_GraphTerminal>(InNode); 
+				EntryNode->TerminalType = EOptimusTerminalType::Entry;
+				return EntryNode->SetGraphPositionDirect({NodeBox.Min.X - 450.0f, NodeBox.GetCenter().Y});
+			});
 
-	for (UOptimusNodeSubGraph* SubGraphToDuplicate: SubGraphsToDuplicate)
-	{
-		DuplicateSubGraph(SubGraphToDuplicate, SubGraphToDuplicate->GetFName());
-	}
-	
-	// Duplicate the nodes into the graph.
-	for (UOptimusNode* Node: NodesToDuplicate)
-	{
-		Action->AddSubAction<FOptimusNodeGraphAction_DuplicateNode>(
-			SubGraphPath, Node, Node->GetFName(), [](UOptimusNode*)
+		GetActionStack()->RunAction<FOptimusNodeGraphAction_AddNode>(
+			SubGraphPath, UOptimusNode_GraphTerminal::StaticClass(), UOptimusNode_GraphTerminal::ReturnNodeName,
+			[NodeBox](UOptimusNode* InNode)
 			{
-				return true;
+				UOptimusNode_GraphTerminal* ReturnNode = Cast<UOptimusNode_GraphTerminal>(InNode); 
+				ReturnNode->TerminalType = EOptimusTerminalType::Return;
+				return ReturnNode->SetGraphPositionDirect({NodeBox.Max.X + 750.0f, NodeBox.GetCenter().Y});
 			});
-	}
 
-	TArray<TPair<FString, FString>> NewNodePairs;
-	for (const UOptimusNodePair* NodePair : NodePairsToDuplicate)
-	{
-		// Genernerate new pair
-		const UOptimusNode *FirstNode = NodePair->GetFirst();
-		const UOptimusNode *SecondNode = NodePair->GetSecond();
+		for (UOptimusNodeSubGraph* SubGraphToDuplicate: SubGraphsToDuplicate)
+		{
+			DuplicateSubGraph(GetActionStack(), SubGraphPath, SubGraphToDuplicate, SubGraphToDuplicate->GetFName());
+		}
+		
+		// Duplicate the nodes into the graph.
+		for (UOptimusNode* Node: NodesToDuplicate)
+		{
+			GetActionStack()->RunAction<FOptimusNodeGraphAction_DuplicateNode>(
+				SubGraphPath, Node, Node->GetFName(), [](UOptimusNode*)
+				{
+					return true;
+				});
+		}
 
-		FString NewFirstNodePath = ConstructPath(SubGraphPath, FirstNode->GetName(), {});
-		FString NewSecondNodePath = ConstructPath(SubGraphPath, SecondNode->GetName(), {});
+		TArray<TPair<FString, FString>> NewNodePairs;
+		for (const UOptimusNodePair* NodePair : NodePairsToDuplicate)
+		{
+			// Genernerate new pair
+			const UOptimusNode *FirstNode = NodePair->GetFirst();
+			const UOptimusNode *SecondNode = NodePair->GetSecond();
 
-		NewNodePairs.Add({NewFirstNodePath, NewSecondNodePath});
-	}
+			FString NewFirstNodePath = ConstructPath(SubGraphPath, FirstNode->GetName(), {});
+			FString NewSecondNodePath = ConstructPath(SubGraphPath, SecondNode->GetName(), {});
 
-	for (const TPair<FString, FString>& NewPair : NewNodePairs)
-	{
-		Action->AddSubAction<FOptimusNodeGraphAction_AddNodePair>(NewPair.Key, NewPair.Value);
-	}
+			NewNodePairs.Add({NewFirstNodePath, NewSecondNodePath});
+		}
 
-	// Create the reference node and connect it.
-	FName GraphNodeRefName = DefaultSubGraphRefNodeName;
-	GraphNodeRefName = Optimus::GetUniqueNameForScope(this, GraphNodeRefName);
-	FOptimusNodeGraphAction_AddNode* AddSubGraphRefNodeAction = new FOptimusNodeGraphAction_AddNode(
-		GetGraphPath(), UOptimusNode_SubGraphReference::StaticClass(), GraphNodeRefName,
-		[NodeBox, SubGraphName](UOptimusNode* InNode)
+		for (const TPair<FString, FString>& NewPair : NewNodePairs)
 		{
-			UOptimusNode_SubGraphReference* SubGraphNode = Cast<UOptimusNode_SubGraphReference>(InNode); 
-			SubGraphNode->InitializeSerializedSubGraphName(SubGraphName);
-			return SubGraphNode->SetGraphPositionDirect(NodeBox.GetCenter());
-		});
-	
-	Action->AddSubAction(AddSubGraphRefNodeAction);
-	
-	// Remove all existing links in the original graph. This has to be done before we remove 
-	// nodes and add new links, otherwise the node removal and link creation to already-connected
-	// inputs will fail.
-	for (const UOptimusNodeLink* Link: InternalLinks)
-	{
-		Action->AddSubAction<FOptimusNodeGraphAction_RemoveLink>(Link);
-	}
-	for (const UOptimusNodeLink* Link: InputLinks)
-	{
-		Action->AddSubAction<FOptimusNodeGraphAction_RemoveLink>(Link);
-	}
-	for (const UOptimusNodeLink* Link: OutputLinks)
-	{
-		Action->AddSubAction<FOptimusNodeGraphAction_RemoveLink>(Link);
-	}
-	
-	// Create the links in the sub-graph between the duplicated nodes and the entry/return
-	// nodes using the internal/input/output links gathered previously.
-	for (const UOptimusNodeLink* Link: InternalLinks)
-	{
-		FString NodeOutputPinPath = FString::Printf(TEXT("%s/%s.%s"),
-			*SubGraphPath, *Link->NodeOutputPin->GetOwningNode()->GetName(), *Link->NodeOutputPin->GetUniqueName().ToString());
-		FString NodeInputPinPath = FString::Printf(TEXT("%s/%s.%s"),
-			*SubGraphPath, *Link->NodeInputPin->GetOwningNode()->GetName(), *Link->NodeInputPin->GetUniqueName().ToString());
-		Action->AddSubAction<FOptimusNodeGraphAction_AddLink>(NodeOutputPinPath, NodeInputPinPath);
-	}
-	for (const UOptimusNodeLink* Link: InputLinks)
-	{
-		FString EntryNodePinName = EntryNodePinNames[Link].ToString(); 
-		// Once for the Entry -> Sub-graph nodes, and another for Outer graph nodes -> Sub-graph ref node inputs.
-		FString NodeOutputPinPath = ConstructPath(SubGraphPath, UOptimusNode_GraphTerminal::EntryNodeName.ToString(), EntryNodePinName);
-		FString NodeInputPinPath = FString::Printf(TEXT("%s/%s.%s"),
-			*SubGraphPath, *Link->NodeInputPin->GetOwningNode()->GetName(), *Link->NodeInputPin->GetUniqueName().ToString());
-		Action->AddSubAction<FOptimusNodeGraphAction_AddLink>(NodeOutputPinPath, NodeInputPinPath);
+			GetActionStack()->RunAction<FOptimusNodeGraphAction_AddNodePair>(NewPair.Key, NewPair.Value);
+		}
 		
-		NodeOutputPinPath = Link->NodeOutputPin->GetPinPath();
-		NodeInputPinPath = FString::Printf(TEXT("%s/%s.%s"),
-			*GetGraphPath(), *GraphNodeRefName.ToString(), *EntryNodePinName);
-		Action->AddSubAction<FOptimusNodeGraphAction_AddLink>(NodeOutputPinPath, NodeInputPinPath);
-	}
-	for (const UOptimusNodeLink* Link: OutputLinks)
-	{
-		FString ReturnNodePinName = ReturnNodePinNames[Link].ToString();
+		// Create the reference node and connect it.
+		FName GraphNodeRefName = DefaultSubGraphRefNodeName;
+		GraphNodeRefName = Optimus::GetUniqueNameForScope(this, GraphNodeRefName);
 		
-		// Once for the Sub-graph nodes -> Return, and another for Sub-graph ref node outputs -> Outer graph nodes.
-		FString NodeOutputPinPath = FString::Printf(TEXT("%s/%s.%s"),
-			*SubGraphPath, *Link->NodeOutputPin->GetOwningNode()->GetName(), *Link->NodeOutputPin->GetUniqueName().ToString());
-		FString NodeInputPinPath =
-		ConstructPath(SubGraphPath, UOptimusNode_GraphTerminal::ReturnNodeName.ToString(), ReturnNodePinName);
-		Action->AddSubAction<FOptimusNodeGraphAction_AddLink>(NodeOutputPinPath, NodeInputPinPath);
+		{
+			// Workaround for extracting result from the undo stack
+			FOptimusCompoundAction* AddSubGraphRefNodeWrapperAction = new FOptimusCompoundAction;
+			WeakAddSubGraphRefNodeAction = AddSubGraphRefNodeWrapperAction->AddSubAction<FOptimusNodeGraphAction_AddNode>(
+				GetCollectionPath(), UOptimusNode_SubGraphReference::StaticClass(), GraphNodeRefName,
+				[NodeBox, SubGraphName](UOptimusNode* InNode)
+				{
+					UOptimusNode_SubGraphReference* SubGraphNode = Cast<UOptimusNode_SubGraphReference>(InNode); 
+					SubGraphNode->InitializeSerializedSubGraphName(SubGraphName);
+					return SubGraphNode->SetGraphPositionDirect(NodeBox.GetCenter());
+				});
+
+			GetActionStack()->RunAction(AddSubGraphRefNodeWrapperAction);	
+		}
+
+		// Remove the nodes we have collapsed
+		// Note: This also removes all existing links in the original graph. This has to be done before we add new links,
+		// otherwise the node removal and link creation to already-connected inputs will fail.
+		RemoveNodes(NodesToDuplicate);
 		
-		NodeOutputPinPath = FString::Printf(TEXT("%s/%s.%s"),
-			*GetGraphPath(), *GraphNodeRefName.ToString(), *ReturnNodePinName);
-		NodeInputPinPath = Link->NodeInputPin->GetPinPath();
-		Action->AddSubAction<FOptimusNodeGraphAction_AddLink>(NodeOutputPinPath, NodeInputPinPath);
-	}
+		// Create the links in the sub-graph between the duplicated nodes and the entry/return
+		// nodes using the internal/input/output links gathered previously.
+		for (const UOptimusNodeLink* Link: InternalLinks)
+		{
+			FString NodeOutputPinPath = FString::Printf(TEXT("%s/%s.%s"),
+				*SubGraphPath, *Link->NodeOutputPin->GetOwningNode()->GetName(), *Link->NodeOutputPin->GetUniqueName().ToString());
+			FString NodeInputPinPath = FString::Printf(TEXT("%s/%s.%s"),
+				*SubGraphPath, *Link->NodeInputPin->GetOwningNode()->GetName(), *Link->NodeInputPin->GetUniqueName().ToString());
+			GetActionStack()->RunAction<FOptimusNodeGraphAction_AddLink>(NodeOutputPinPath, NodeInputPinPath);
+		}
+		for (const UOptimusNodeLink* Link: InputLinks)
+		{
+			FString EntryNodePinName = EntryNodePinNames[Link].ToString(); 
+			// Once for the Entry -> Sub-graph nodes, and another for Outer graph nodes -> Sub-graph ref node inputs.
+			FString NodeOutputPinPath = ConstructPath(SubGraphPath, UOptimusNode_GraphTerminal::EntryNodeName.ToString(), EntryNodePinName);
+			FString NodeInputPinPath = FString::Printf(TEXT("%s/%s.%s"),
+				*SubGraphPath, *Link->NodeInputPin->GetOwningNode()->GetName(), *Link->NodeInputPin->GetUniqueName().ToString());
+			GetActionStack()->RunAction<FOptimusNodeGraphAction_AddLink>(NodeOutputPinPath, NodeInputPinPath);
+			
+			NodeOutputPinPath = Link->NodeOutputPin->GetPinPath();
+			NodeInputPinPath = FString::Printf(TEXT("%s/%s.%s"),
+				*GetCollectionPath(), *GraphNodeRefName.ToString(), *EntryNodePinName);
+			GetActionStack()->RunAction<FOptimusNodeGraphAction_AddLink>(NodeOutputPinPath, NodeInputPinPath);
+		}
+		for (const UOptimusNodeLink* Link: OutputLinks)
+		{
+			FString ReturnNodePinName = ReturnNodePinNames[Link].ToString();
+			
+			// Once for the Sub-graph nodes -> Return, and another for Sub-graph ref node outputs -> Outer graph nodes.
+			FString NodeOutputPinPath = FString::Printf(TEXT("%s/%s.%s"),
+				*SubGraphPath, *Link->NodeOutputPin->GetOwningNode()->GetName(), *Link->NodeOutputPin->GetUniqueName().ToString());
+			FString NodeInputPinPath =
+			ConstructPath(SubGraphPath, UOptimusNode_GraphTerminal::ReturnNodeName.ToString(), ReturnNodePinName);
+			GetActionStack()->RunAction<FOptimusNodeGraphAction_AddLink>(NodeOutputPinPath, NodeInputPinPath);
+			
+			NodeOutputPinPath = FString::Printf(TEXT("%s/%s.%s"),
+				*GetCollectionPath(), *GraphNodeRefName.ToString(), *ReturnNodePinName);
+			NodeInputPinPath = Link->NodeInputPin->GetPinPath();
+			GetActionStack()->RunAction<FOptimusNodeGraphAction_AddLink>(NodeOutputPinPath, NodeInputPinPath);
+		}
 
-	// Delete the node pairs that we have moved into the subgraph
-	for (UOptimusNodePair* NodePair : NodePairsToDuplicate)
-	{
-		Action->AddSubAction<FOptimusNodeGraphAction_RemoveNodePair>(NodePair);
-	}
-	
-	// Delete the existing nodes
-	for (UOptimusNode* Node: NodeSet)
-	{
-		Action->AddSubAction<FOptimusNodeGraphAction_RemoveNode>(Node);
+		// Closes ActionScope
 	}
-	
-	if (!GetActionStack()->RunAction(Action))
+
+	if (WeakAddSubGraphRefNodeAction.IsValid())
 	{
-		return nullptr;
+		IOptimusPathResolver* PathResolver = GetPathResolver();
+		return WeakAddSubGraphRefNodeAction.Pin()->GetNode(PathResolver);
 	}
-	
-	return AddSubGraphRefNodeAction->GetNode(PathResolver);
+
+	return nullptr;
 }
 
 
@@ -1539,26 +1456,6 @@ TArray<UOptimusNode*> UOptimusNodeGraph::ExpandCollapsedNodes(
 		return {};
 	}
 
-	// Figure out the non-clashing names to use, to avoid collisions during actual execution.
-	TSet<FName> ExistingObjects;
-	for (const UOptimusNode* Node: Nodes)
-	{
-		if (ensure(Node != nullptr))
-		{
-			ExistingObjects.Add(Node->GetFName());
-		}
-	}
-
-	auto MakeUniqueNodeName = [&ExistingObjects](FName InName)
-	{
-		while(ExistingObjects.Contains(InName))
-		{
-			InName.SetNumber(InName.GetNumber() + 1);
-		}
-		ExistingObjects.Add(InName);
-		return InName;
-	};
-	
 	TArray<UOptimusNodeLink*> ExternalLinks;	
 
 	for (const int32 LinkIndex : GetAllLinkIndexesToNode(InGraphReferenceNode))
@@ -1590,15 +1487,22 @@ TArray<UOptimusNode*> UOptimusNodeGraph::ExpandCollapsedNodes(
 	{
 		UOptimusNodeSubGraph* SubGraph = AsReferencerNode->GetReferencedSubGraph();
 		
-		TArray<UOptimusNode*> NodesToDuplicate = SubGraph->Nodes;
-		TArray<UOptimusNodePair*> NodePairsToDuplicate = SubGraph->NodePairs;
+		TArray<UOptimusNode*> NodesToConsider= SubGraph->Nodes;
 
-		NodesToDuplicate.RemoveAll([](const UOptimusNode* InNode)
+		NodesToConsider.RemoveAll([](const UOptimusNode* InNode)
 		{
 			return Cast<UOptimusNode_GraphTerminal>(InNode);
 		});
 
-
+		TArray<UOptimusNode*> NodesToDuplicate;
+		TArray<UOptimusNodePair*> NodePairsToDuplicate;
+		TArray<UOptimusNodeSubGraph*> SubGraphsToDuplicate;
+		
+		GatherRelatedObjects(NodesToConsider, NodesToDuplicate, NodePairsToDuplicate, SubGraphsToDuplicate);
+		
+		// Figure out the non-clashing names to use, to avoid collisions during actual execution.
+		Optimus::FUniqueNameGenerator UniqueNameGenerator(this);
+		
 		using FloatType = decltype(FVector2D::X);
 		FVector2D TopLeft{std::numeric_limits<FloatType>::max()};
 		FVector2D BotRight{std::numeric_limits<FloatType>::lowest()};
@@ -1609,13 +1513,13 @@ TArray<UOptimusNode*> UOptimusNodeGraph::ExpandCollapsedNodes(
 		{
 			TopLeft = FVector2D::Min(TopLeft, Node->GraphPosition);
 			BotRight = FVector2D::Max(BotRight, Node->GraphPosition);
-			NewNodeNameMap.Add(Node, MakeUniqueNodeName(Node->GetFName()));
+			NewNodeNameMap.Add(Node, UniqueNameGenerator.GetUniqueName(Node->GetFName()));
 		}
 
 		HalfBoundSize = (BotRight - TopLeft )/2 + FVector2D(100, 100);
 		Center = (BotRight + TopLeft) / 2;
 		
-		const FString GraphPath = GetGraphPath();
+		const FString CurrentGraphPath = GetCollectionPath();
 	
 		TArray<TPair<FString, FString>> NewNodePairs;
 		for (const UOptimusNodePair* NodePair : NodePairsToDuplicate)
@@ -1626,17 +1530,24 @@ TArray<UOptimusNode*> UOptimusNodeGraph::ExpandCollapsedNodes(
 
 			if (ensure(NewNodeNameMap.Contains(FirstNode) && NewNodeNameMap.Contains(SecondNode)))
 			{
-				FString NewFirstNodePath = ConstructPath(GraphPath, NewNodeNameMap[FirstNode].ToString(), {});
-				FString NewSecondNodePath = ConstructPath(GraphPath, NewNodeNameMap[SecondNode].ToString(), {});
+				FString NewFirstNodePath = ConstructPath(CurrentGraphPath, NewNodeNameMap[FirstNode].ToString(), {});
+				FString NewSecondNodePath = ConstructPath(CurrentGraphPath, NewNodeNameMap[SecondNode].ToString(), {});
 
 				NewNodePairs.Add({NewFirstNodePath, NewSecondNodePath});
 			}	
 		}	
+
+		// Using GraphName instead of the graph object as key here because this map is to be captured by a lambda to be used during undo/redo 
+		TMap<FName , FName> NewSubGraphNameMap;
+		for (UOptimusNodeSubGraph* SubGraphToDuplicate : SubGraphsToDuplicate)
+		{
+			NewSubGraphNameMap.Add(SubGraphToDuplicate->GetFName(), UniqueNameGenerator.GetUniqueName(SubGraphToDuplicate->GetFName()));
+		}
 		
 		TArray<TPair<FString, FString>> LinksToAdd;
 
 		// Record all links linking terminal pins and nodes in the outer graph, saving to LinksToAdd array
-		auto AddLink = [&LinksToAdd, GraphPath, AsRouterNode, NewNodeNameMap](UOptimusNodePin* InReferenceNodePin)
+		auto CollectLink = [&LinksToAdd, CurrentGraphPath, AsRouterNode, NewNodeNameMap](UOptimusNodePin* InReferenceNodePin)
 		{
 			const FOptimusRoutedNodePin PinCounterpart = AsRouterNode->GetPinCounterpart(InReferenceNodePin, {});
 
@@ -1647,7 +1558,7 @@ TArray<UOptimusNode*> UOptimusNodeGraph::ExpandCollapsedNodes(
 				for (const UOptimusNodePin* InternalPin : PinCounterpart.NodePin->GetConnectedPins())
 				{
 					const UOptimusNode* InternalNode = InternalPin->GetOwningNode();
-					FString InternalPinPath = ConstructPath(GraphPath, NewNodeNameMap[InternalNode].ToString(), InternalPin->GetUniqueName().ToString());
+					FString InternalPinPath = ConstructPath(CurrentGraphPath, NewNodeNameMap[InternalNode].ToString(), InternalPin->GetUniqueName().ToString());
 
 					if (ExternalPin->GetDirection() == EOptimusNodePinDirection::Output)
 					{
@@ -1663,7 +1574,7 @@ TArray<UOptimusNode*> UOptimusNodeGraph::ExpandCollapsedNodes(
 		
 		for (UOptimusNodePin* Pin : InGraphReferenceNode->GetPins())
 		{
-			AddLink(Pin);
+			CollectLink(Pin);
 		}
 		
 		// Unconnected component pins should be linked to outer graph explicitly if the subgraph default component pin is connected externally
@@ -1684,7 +1595,7 @@ TArray<UOptimusNode*> UOptimusNodeGraph::ExpandCollapsedNodes(
 					{
 						check(InternalPin->GetDirection() ==EOptimusNodePinDirection::Input);
 						
-						FString InternalPinPath = ConstructPath(GraphPath, NewNodeNameMap[Node].ToString(), InternalPin->GetUniqueName().ToString());
+						FString InternalPinPath = ConstructPath(CurrentGraphPath, NewNodeNameMap[Node].ToString(), InternalPin->GetUniqueName().ToString());
 						
 						LinksToAdd.Add({ExternalPinPath, InternalPinPath});
 					}
@@ -1707,70 +1618,66 @@ TArray<UOptimusNode*> UOptimusNodeGraph::ExpandCollapsedNodes(
 			const UOptimusNode* InputNode = InputPin->GetOwningNode();
 			
 			LinksToAdd.Add({
-				ConstructPath(GraphPath, NewNodeNameMap[OutputNode].ToString(), OutputPin->GetUniqueName().ToString()),
-				ConstructPath(GraphPath, NewNodeNameMap[InputNode].ToString(), InputPin->GetUniqueName().ToString())
+				ConstructPath(CurrentGraphPath, NewNodeNameMap[OutputNode].ToString(), OutputPin->GetUniqueName().ToString()),
+				ConstructPath(CurrentGraphPath, NewNodeNameMap[InputNode].ToString(), InputPin->GetUniqueName().ToString())
 			});
 		}
 
-		FOptimusCompoundAction *Action = new FOptimusCompoundAction(TEXT("Expand Collapsed Nodes"));
+		FOptimusActionScope ActionScope(*GetActionStack(), TEXT("Expand Collapsed Nodes"));
 
 		for (TPair<UOptimusNode*, FVector2D> NodePositionInfo: NodePositionInfos)
 		{
 			FVector2D Offset = NodePositionInfo.Value.GetSignVector() * HalfBoundSize;
 			
-			Action->AddSubAction<FOptimusNodeAction_MoveNode>(NodePositionInfo.Key, NodePositionInfo.Key->GetGraphPosition() + Offset);
+			GetActionStack()->RunAction<FOptimusNodeAction_MoveNode>(NodePositionInfo.Key, NodePositionInfo.Key->GetGraphPosition() + Offset);
 		}
 
-		
-		// Add all the pre-duplicate requirement actions first. This allows certain nodes to set up their
-		// operating environment correctly (e.g. a missing variable description for get variable nodes, etc.)
-		for (UOptimusNode* Node: NodesToDuplicate)
+		for (UOptimusNodeSubGraph* SubGraphToDuplicate : SubGraphsToDuplicate)
 		{
-			Node->PreDuplicateRequirementActions(this, Action);
+			DuplicateSubGraph(GetActionStack(), CurrentGraphPath, SubGraphToDuplicate, NewSubGraphNameMap[SubGraphToDuplicate->GetFName()]);
 		}
+		
 
-		TArray<TWeakPtr<FOptimusNodeGraphAction_DuplicateNode>> WeakDuplicateActions;
+		{
+			FOptimusCompoundAction *PreDuplicateRequirementAction = new FOptimusCompoundAction;	
+			// Add all the pre-duplicate requirement actions first. This allows certain nodes to set up their
+			// operating environment correctly (e.g. a missing variable description for get variable nodes, etc.)
+			for (UOptimusNode* Node: NodesToDuplicate)
+			{
+				Node->PreDuplicateRequirementActions(this, PreDuplicateRequirementAction);
+			}
+			GetActionStack()->RunAction(PreDuplicateRequirementAction);
+		}
 
 		// Duplicate the nodes and place them correctly
 		for (UOptimusNode* Node: NodesToDuplicate)
 		{
-			FOptimusNodeGraphAction_DuplicateNode *DuplicateNodeAction = new FOptimusNodeGraphAction_DuplicateNode(
-				GetGraphPath(), Node, NewNodeNameMap[Node],
-				[SourceNodeGraphPosition = Node->GraphPosition, Center, ExpandLocation](UOptimusNode *InNode) {
+			GetActionStack()->RunAction<FOptimusNodeGraphAction_DuplicateNode>(
+				GetCollectionPath(), Node, NewNodeNameMap[Node],
+				[SourceNodeGraphPosition = Node->GraphPosition, Center, ExpandLocation, NewSubGraphNameMap](UOptimusNode *InNode)
+				{
+					if (UOptimusNode_SubGraphReference* SubGraphNode = Cast<UOptimusNode_SubGraphReference>(InNode))
+					{
+						FName NewSubGraphName = NewSubGraphNameMap[SubGraphNode->GetSerializedSubGraphName()];
+						SubGraphNode->InitializeSerializedSubGraphName(NewSubGraphName);
+					}
 					InNode->SetGraphPositionDirect(SourceNodeGraphPosition - Center + ExpandLocation);
 					return true;
 			});
-		
-			WeakDuplicateActions.Add(StaticCastWeakPtr<FOptimusNodeGraphAction_DuplicateNode>(Action->AddSubAction(DuplicateNodeAction)));
 		}
 
+		// This will remove the subgraph if it is a subgraph node
+		RemoveNode(InGraphReferenceNode);
+		
 		for (const TPair<FString, FString>& NewPair : NewNodePairs)
 		{
-			Action->AddSubAction<FOptimusNodeGraphAction_AddNodePair>(NewPair.Key, NewPair.Value);
+			GetActionStack()->RunAction<FOptimusNodeGraphAction_AddNodePair>(NewPair.Key, NewPair.Value);
 		}
 
 		for (const TPair<FString, FString>& NewLink : LinksToAdd)
 		{
-			Action->AddSubAction<FOptimusNodeGraphAction_AddLink>(NewLink.Key, NewLink.Value);
-		}
-
-		RemoveNodesToAction(Action, SubGraph->GetAllNodes());
-
-		// This will remove the subgraph if it is a subgraph node
-		RemoveNodesToAction(Action, {InGraphReferenceNode});
-
-		if (!GetActionStack()->RunAction(Action))
-		{
-			return {};
-		}
-
-		IOptimusPathResolver* PathResolver = GetPathResolver();
-		TArray<UOptimusNode*> NewNodes;
-		for (TWeakPtr<FOptimusNodeGraphAction_DuplicateNode> WeakDuplicateAction : WeakDuplicateActions)
-		{
-			NewNodes.Add(WeakDuplicateAction.Pin()->GetNode(PathResolver));
+			GetActionStack()->RunAction<FOptimusNodeGraphAction_AddLink>(NewLink.Key, NewLink.Value);
 		}
-		return NewNodes;	
 	}
 	
 	return {};
@@ -1801,13 +1708,15 @@ bool UOptimusNodeGraph::ConvertToFunction(UOptimusNode* InSubGraphNode)
 	FName FunctionName = TEXT("NewFunction");
 	FunctionName = Optimus::GetUniqueNameForScope(Deformer, FunctionName);
 	
-	const FString FunctionGraphPath = GetFunctionGraphPath(FunctionName.ToString());
-	const FVector2D EntryNodePosition = SubGraph->EntryNode->GetGraphPosition();
-	const FVector2D ReturnNodePosition = SubGraph->ReturnNode->GetGraphPosition();
+	const FString FunctionGraphPath = GetFunctionGraphCollectionPath(FunctionName.ToString());
+	UOptimusNode_GraphTerminal* SubGraphEntryNode = SubGraph->GetTerminalNode(EOptimusTerminalType::Entry);
+	UOptimusNode_GraphTerminal* SubGraphReturnNode = SubGraph->GetTerminalNode(EOptimusTerminalType::Return);
+	const FVector2D EntryNodePosition = SubGraphEntryNode->GetGraphPosition();
+	const FVector2D ReturnNodePosition = SubGraphReturnNode->GetGraphPosition();
 	const FVector2D ReferenceNodePosition = SubGraphNode->GetGraphPosition();
 	const FName FunctionNodeName = Optimus::GetUniqueNameForScope(this, TEXT("FunctionReferenceNode"));
 	const FString FunctionNodeNameString = FunctionNodeName.ToString();
-	const FString CurrentGraphPath = GetGraphPath();
+	const FString CurrentGraphPath = GetCollectionPath();
 
 	TMap<FString, TArray<FString>> ExternalLinksToAdd;
 
@@ -1841,7 +1750,7 @@ bool UOptimusNodeGraph::ConvertToFunction(UOptimusNode* InSubGraphNode)
 	
 	for (UOptimusNodeLink* Link : SubGraph->GetAllLinks())
 	{
-		if (Link->GetNodeInputPin()->GetOwningNode() == SubGraph->ReturnNode)
+		if (Link->GetNodeInputPin()->GetOwningNode() == SubGraphReturnNode)
 		{
 			FString ReturnNodeInputPinPath = ConstructPath(
 				FunctionGraphPath,
@@ -1857,7 +1766,7 @@ bool UOptimusNodeGraph::ConvertToFunction(UOptimusNode* InSubGraphNode)
 
 			InternalLinksToAdd.FindOrAdd(InternalNodeOutputPinPath).Add(ReturnNodeInputPinPath);
 		}
-		else if (Link->GetNodeOutputPin()->GetOwningNode() == SubGraph->EntryNode)
+		else if (Link->GetNodeOutputPin()->GetOwningNode() == SubGraphEntryNode)
 		{
 			FString EntryNodeOutputPinPath = ConstructPath(
 				FunctionGraphPath,
@@ -1918,7 +1827,7 @@ bool UOptimusNodeGraph::ConvertToFunction(UOptimusNode* InSubGraphNode)
 
 	{
 		FOptimusActionScope(*GetActionStack(), TEXT("Convert To Function"));
-		GetActionStack()->RunAction<FOptimusNodeGraphAction_AddGraph>(Deformer, EOptimusNodeGraphType::Function, FunctionName, INDEX_NONE,
+		GetActionStack()->RunAction<FOptimusNodeGraphAction_AddGraph>(Deformer->GetCollectionPath(), EOptimusNodeGraphType::Function, FunctionName, INDEX_NONE,
 			[SubGraph](UOptimusNodeGraph* InGraph)
 			{
 				UOptimusFunctionNodeGraph* FunctionGraph = CastChecked<UOptimusFunctionNodeGraph>(InGraph);
@@ -1930,35 +1839,26 @@ bool UOptimusNodeGraph::ConvertToFunction(UOptimusNode* InSubGraphNode)
 	
 		GetActionStack()->RunAction<FOptimusNodeGraphAction_AddNode>(
 			FunctionGraphPath, UOptimusNode_GraphTerminal::StaticClass(), UOptimusNode_GraphTerminal::EntryNodeName,
-			[FunctionGraphPath, EntryNodePosition, PathResolver](UOptimusNode* InNode)
+			[EntryNodePosition](UOptimusNode* InNode)
 			{
 				UOptimusNode_GraphTerminal* EntryNode = Cast<UOptimusNode_GraphTerminal>(InNode); 
-				UOptimusNodeSubGraph* SubGraph = Cast<UOptimusNodeSubGraph>(PathResolver->ResolveGraphPath(FunctionGraphPath));
-
-				SubGraph->EntryNode = EntryNode;
-		
 				EntryNode->TerminalType = EOptimusTerminalType::Entry;
-				EntryNode->OwningGraph = SubGraph;
+
 				return EntryNode->SetGraphPositionDirect(EntryNodePosition);
 			});
 
 		GetActionStack()->RunAction<FOptimusNodeGraphAction_AddNode>(
 			FunctionGraphPath, UOptimusNode_GraphTerminal::StaticClass(), UOptimusNode_GraphTerminal::ReturnNodeName,
-			[FunctionGraphPath, ReturnNodePosition, PathResolver](UOptimusNode* InNode)
+			[ReturnNodePosition](UOptimusNode* InNode)
 			{
 				UOptimusNode_GraphTerminal* ReturnNode = Cast<UOptimusNode_GraphTerminal>(InNode); 
-				UOptimusNodeSubGraph* SubGraph = Cast<UOptimusNodeSubGraph>(PathResolver->ResolveGraphPath(FunctionGraphPath));
-
-				SubGraph->ReturnNode = ReturnNode;
-			
 				ReturnNode->TerminalType = EOptimusTerminalType::Return;
-				ReturnNode->OwningGraph = SubGraph;
 				return ReturnNode->SetGraphPositionDirect(ReturnNodePosition);
 			});
 
 		for (UOptimusNodeSubGraph* SubGraphToDuplicate : SubGraphsToDuplicate)
 		{
-			DuplicateSubGraph(SubGraphToDuplicate, SubGraphToDuplicate->GetFName());
+			DuplicateSubGraph(GetActionStack(), FunctionGraphPath, SubGraphToDuplicate, SubGraphToDuplicate->GetFName());
 		}	
 		
 		// Copy the nodes to the function graph
@@ -1988,7 +1888,7 @@ bool UOptimusNodeGraph::ConvertToFunction(UOptimusNode* InSubGraphNode)
 		
 		// Create the function ref node
 		GetActionStack()->RunAction<FOptimusNodeGraphAction_AddNode>(
-			GetGraphPath(), UOptimusNode_FunctionReference::StaticClass(), FunctionNodeName,
+			GetCollectionPath(), UOptimusNode_FunctionReference::StaticClass(), FunctionNodeName,
 			[ReferenceNodePosition, FunctionGraphPath, PathResolver](UOptimusNode* InNode)
 			{
 				UOptimusNode_FunctionReference* FunctionNode = Cast<UOptimusNode_FunctionReference>(InNode); 
@@ -1999,7 +1899,7 @@ bool UOptimusNodeGraph::ConvertToFunction(UOptimusNode* InSubGraphNode)
 			});
 
 		// Remove the subgraph node
-		RemoveNodes({SubGraphNode});
+		RemoveNode(SubGraphNode);
 		
 		// Restore links
 		for (const TPair<FString, TArray<FString>>& ExternalLinks : ExternalLinksToAdd )
@@ -2041,12 +1941,14 @@ bool UOptimusNodeGraph::ConvertToSubGraph(UOptimusNode* InFunctionNode)
 	SubGraphName = Optimus::GetUniqueNameForScope(this, SubGraphName);
 	
 	const FString SubGraphPath = ConstructSubGraphPath(SubGraphName.ToString());
-	const FVector2D EntryNodePosition = FunctionGraph->EntryNode->GetGraphPosition();
-	const FVector2D ReturnNodePosition = FunctionGraph->ReturnNode->GetGraphPosition();
+	UOptimusNode_GraphTerminal* FunctionGraphEntryNode = FunctionGraph->GetTerminalNode(EOptimusTerminalType::Entry);
+	UOptimusNode_GraphTerminal* FunctionGraphReturnNode = FunctionGraph->GetTerminalNode(EOptimusTerminalType::Return);
+	const FVector2D EntryNodePosition = FunctionGraphEntryNode->GetGraphPosition();
+	const FVector2D ReturnNodePosition = FunctionGraphReturnNode->GetGraphPosition();
 	const FVector2D ReferenceNodePosition = FunctionNode->GetGraphPosition();
 	const FName SubGraphNodeName = Optimus::GetUniqueNameForScope(this, DefaultSubGraphRefNodeName);
 	const FString SubGraphNodeNameString = SubGraphNodeName.ToString();
-	const FString CurrentGraphPath = GetGraphPath();
+	const FString CurrentGraphPath = GetCollectionPath();
 
 	TMap<FString, TArray<FString>> ExternalLinksToAdd;
 
@@ -2080,7 +1982,7 @@ bool UOptimusNodeGraph::ConvertToSubGraph(UOptimusNode* InFunctionNode)
 	
 	for (UOptimusNodeLink* Link : FunctionGraph->GetAllLinks())
 	{
-		if (Link->GetNodeInputPin()->GetOwningNode() == FunctionGraph->ReturnNode)
+		if (Link->GetNodeInputPin()->GetOwningNode() == FunctionGraphReturnNode)
 		{
 			FString ReturnNodeInputPinPath = ConstructPath(
 				SubGraphPath,
@@ -2096,7 +1998,7 @@ bool UOptimusNodeGraph::ConvertToSubGraph(UOptimusNode* InFunctionNode)
 
 			InternalLinksToAdd.FindOrAdd(InternalNodeOutputPinPath).Add(ReturnNodeInputPinPath);
 		}
-		else if (Link->GetNodeOutputPin()->GetOwningNode() == FunctionGraph->EntryNode)
+		else if (Link->GetNodeOutputPin()->GetOwningNode() == FunctionGraphEntryNode)
 		{
 			FString EntryNodeOutputPinPath = ConstructPath(
 				SubGraphPath,
@@ -2157,7 +2059,7 @@ bool UOptimusNodeGraph::ConvertToSubGraph(UOptimusNode* InFunctionNode)
 
 	{
 		FOptimusActionScope(*GetActionStack(), TEXT("Convert To SubGraph"));
-		GetActionStack()->RunAction<FOptimusNodeGraphAction_AddGraph>(this, EOptimusNodeGraphType::SubGraph, SubGraphName, INDEX_NONE,
+		GetActionStack()->RunAction<FOptimusNodeGraphAction_AddGraph>(GetCollectionPath(), EOptimusNodeGraphType::SubGraph, SubGraphName, INDEX_NONE,
 			[FunctionGraph](UOptimusNodeGraph* InGraph)
 			{
 				UOptimusNodeSubGraph* SubGraph = CastChecked<UOptimusNodeSubGraph>(InGraph);
@@ -2168,35 +2070,25 @@ bool UOptimusNodeGraph::ConvertToSubGraph(UOptimusNode* InFunctionNode)
 	
 		GetActionStack()->RunAction<FOptimusNodeGraphAction_AddNode>(
 			SubGraphPath, UOptimusNode_GraphTerminal::StaticClass(), UOptimusNode_GraphTerminal::EntryNodeName,
-			[SubGraphPath, EntryNodePosition, PathResolver](UOptimusNode* InNode)
+			[EntryNodePosition](UOptimusNode* InNode)
 			{
 				UOptimusNode_GraphTerminal* EntryNode = Cast<UOptimusNode_GraphTerminal>(InNode); 
-				UOptimusNodeSubGraph* SubGraph = Cast<UOptimusNodeSubGraph>(PathResolver->ResolveGraphPath(SubGraphPath));
-
-				SubGraph->EntryNode = EntryNode;
-		
 				EntryNode->TerminalType = EOptimusTerminalType::Entry;
-				EntryNode->OwningGraph = SubGraph;
 				return EntryNode->SetGraphPositionDirect(EntryNodePosition);
 			});
 
 		GetActionStack()->RunAction<FOptimusNodeGraphAction_AddNode>(
 			SubGraphPath, UOptimusNode_GraphTerminal::StaticClass(), UOptimusNode_GraphTerminal::ReturnNodeName,
-			[SubGraphPath, ReturnNodePosition, PathResolver](UOptimusNode* InNode)
+			[ReturnNodePosition](UOptimusNode* InNode)
 			{
 				UOptimusNode_GraphTerminal* ReturnNode = Cast<UOptimusNode_GraphTerminal>(InNode); 
-				UOptimusNodeSubGraph* SubGraph = Cast<UOptimusNodeSubGraph>(PathResolver->ResolveGraphPath(SubGraphPath));
-
-				SubGraph->ReturnNode = ReturnNode;
-			
 				ReturnNode->TerminalType = EOptimusTerminalType::Return;
-				ReturnNode->OwningGraph = SubGraph;
 				return ReturnNode->SetGraphPositionDirect(ReturnNodePosition);
 			});
 
 		for (UOptimusNodeSubGraph* SubGraphToDuplicate : SubGraphsToDuplicate)
 		{
-			DuplicateSubGraph(SubGraphToDuplicate, SubGraphToDuplicate->GetFName());
+			DuplicateSubGraph(GetActionStack(), SubGraphPath , SubGraphToDuplicate, SubGraphToDuplicate->GetFName());
 		}
 		
 		// Copy the nodes to the sub graph
@@ -2226,7 +2118,7 @@ bool UOptimusNodeGraph::ConvertToSubGraph(UOptimusNode* InFunctionNode)
 		
 		// Create the sub graph ref node
 		GetActionStack()->RunAction<FOptimusNodeGraphAction_AddNode>(
-			GetGraphPath(), UOptimusNode_SubGraphReference::StaticClass(), SubGraphNodeName,
+			GetCollectionPath(), UOptimusNode_SubGraphReference::StaticClass(), SubGraphNodeName,
 			[ReferenceNodePosition, SubGraphName](UOptimusNode* InNode)
 			{
 				UOptimusNode_SubGraphReference* SubGraphNode = Cast<UOptimusNode_SubGraphReference>(InNode);
@@ -2236,7 +2128,7 @@ bool UOptimusNodeGraph::ConvertToSubGraph(UOptimusNode* InFunctionNode)
 			});
 
 		// Remove the function node
-		RemoveNodes({FunctionNode});
+		RemoveNode(FunctionNode);
 		
 		// Restore links
 		for (const TPair<FString, TArray<FString>>& ExternalLinks : ExternalLinksToAdd )
@@ -3157,9 +3049,14 @@ TArray<int32> UOptimusNodeGraph::GetAllLinkIndexesToPin(
 
 FString UOptimusNodeGraph::ConstructSubGraphPath(const FString& InSubGraphName) const
 {
-	FString SubGraphPath = GetGraphPath() + TEXT("/") + InSubGraphName;
+	return ConstructSubGraphPath(GetCollectionPath(), InSubGraphName);
+}
+
+FString UOptimusNodeGraph::ConstructSubGraphPath(const FString& InGraphOwnerPath, const FString& InSubGraphName)
+{
+	FString SubGraphPath = InGraphOwnerPath + TEXT("/") + InSubGraphName;
 
-	return SubGraphPath;
+	return SubGraphPath;	
 }
 
 
@@ -3194,9 +3091,12 @@ IOptimusNodeGraphCollectionOwner* UOptimusNodeGraph::GetCollectionRoot() const
 
 FString UOptimusNodeGraph::GetCollectionPath() const
 {
-	if (const FString ParentPath = GetCollectionOwner()->GetCollectionPath(); !ParentPath.IsEmpty())
+	if (GetCollectionOwner())
 	{
-		return ParentPath + TEXT("/") + GetName();
+		if (const FString ParentPath = GetCollectionOwner()->GetCollectionPath(); !ParentPath.IsEmpty())
+		{
+			return ParentPath + TEXT("/") + GetName();
+		}
 	}
 	return GetName();
 }
diff --git a/Engine/Plugins/Animation/DeformerGraph/Source/OptimusCore/Private/OptimusNodePin.cpp b/Engine/Plugins/Animation/DeformerGraph/Source/OptimusCore/Private/OptimusNodePin.cpp
--- a/Engine/Plugins/Animation/DeformerGraph/Source/OptimusCore/Private/OptimusNodePin.cpp
+++ b/Engine/Plugins/Animation/DeformerGraph/Source/OptimusCore/Private/OptimusNodePin.cpp
@@ -148,7 +148,7 @@ FText UOptimusNodePin::GetTooltipText() const
 FString UOptimusNodePin::GetPinPath() const
 {
 	return UOptimusNodeGraph::ConstructPath(
-		GetOwningNode()->GetOwningGraph()->GetGraphPath(),
+		GetOwningNode()->GetOwningGraph()->GetCollectionPath(),
 		GetOwningNode()->GetName(),
 		GetUniqueName().ToString());
 }
diff --git a/Engine/Plugins/Animation/DeformerGraph/Source/OptimusCore/Private/OptimusNodeSubGraph.cpp b/Engine/Plugins/Animation/DeformerGraph/Source/OptimusCore/Private/OptimusNodeSubGraph.cpp
--- a/Engine/Plugins/Animation/DeformerGraph/Source/OptimusCore/Private/OptimusNodeSubGraph.cpp
+++ b/Engine/Plugins/Animation/DeformerGraph/Source/OptimusCore/Private/OptimusNodeSubGraph.cpp
@@ -45,7 +45,8 @@ EOptimusDataTypeUsageFlags UOptimusNodeSubGraph::GetTypeUsageFlags(const FOptimu
 
 UOptimusComponentSourceBinding* UOptimusNodeSubGraph::GetDefaultComponentBinding(const FOptimusPinTraversalContext& InTraversalContext) const
 {
-	if (!ensure(EntryNode.IsValid()))
+	UOptimusNode_GraphTerminal* EntryNode = GetTerminalNode(EOptimusTerminalType::Entry);
+	if (!ensure(EntryNode))
 	{
 		return nullptr;
 	}
@@ -53,6 +54,22 @@ UOptimusComponentSourceBinding* UOptimusNodeSubGraph::GetDefaultComponentBinding
 	return EntryNode->GetDefaultComponentBinding(InTraversalContext);
 }
 
+UOptimusNode_GraphTerminal* UOptimusNodeSubGraph::GetTerminalNode(EOptimusTerminalType InTerminalType) const
+{
+	for (UOptimusNode* Node : GetAllNodes())
+	{
+		if (UOptimusNode_GraphTerminal* TerminalNode = Cast<UOptimusNode_GraphTerminal>(Node))
+		{
+			if (TerminalNode->TerminalType == InTerminalType)
+			{
+				return TerminalNode;
+			}
+		}
+	}
+	
+	return nullptr;
+}
+
 #if WITH_EDITOR
 void UOptimusNodeSubGraph::PostEditChangeProperty(FPropertyChangedEvent& PropertyChangedEvent)
 {
@@ -119,8 +136,11 @@ FName UOptimusNodeSubGraph::GetSanitizedBindingName(FName InNewName, FName InOld
 
 	if (Name != InOldName)
 	{
-		Name = UOptimusNode::GetAvailablePinNameStable(EntryNode.Get(), Name);
-		Name = UOptimusNode::GetAvailablePinNameStable(ReturnNode.Get(), Name);
+		const UOptimusNode_GraphTerminal* EntryNode = GetTerminalNode(EOptimusTerminalType::Entry);
+		const UOptimusNode_GraphTerminal* ReturnNode = GetTerminalNode(EOptimusTerminalType::Return);
+
+		Name = UOptimusNode::GetAvailablePinNameStable(EntryNode, Name);
+		Name = UOptimusNode::GetAvailablePinNameStable(ReturnNode, Name);
 	}
 	
 	return Name;
@@ -155,6 +175,9 @@ void UOptimusNodeSubGraph::PropertyArrayPasted(const FPropertyChangedEvent& InPr
 
 void UOptimusNodeSubGraph::PropertyValueChanged(const FPropertyChangedEvent& InPropertyChangedEvent)
 {
+	UOptimusNode_GraphTerminal* EntryNode = GetTerminalNode(EOptimusTerminalType::Entry);
+	UOptimusNode_GraphTerminal* ReturnNode = GetTerminalNode(EOptimusTerminalType::Return);
+	
 	FOptimusParameterBindingArray* BindingArrayPtr = nullptr;
 	TArray<UOptimusNodePin*> BindingPins;
 	bool bAllowParameter = false;
diff --git a/Engine/Plugins/Animation/DeformerGraph/Source/OptimusCore/Public/IOptimusNonCopyableNode.h b/Engine/Plugins/Animation/DeformerGraph/Source/OptimusCore/Public/IOptimusNonCopyableNode.h
new file mode 100644
--- /dev/null
+++ b/Engine/Plugins/Animation/DeformerGraph/Source/OptimusCore/Public/IOptimusNonCopyableNode.h
@@ -0,0 +1,24 @@
+﻿// Copyright Epic Games, Inc. All Rights Reserved.
+
+#pragma once
+
+#include "UObject/Interface.h"
+
+#include "IOptimusNonCopyableNode.generated.h"
+
+UINTERFACE()
+class OPTIMUSCORE_API UOptimusNonCopyableNode :
+	public UInterface
+{
+	GENERATED_BODY()
+};
+
+/**
+* Interface that provides a mechanism to prevent node from being included in a subgraph
+*/
+class OPTIMUSCORE_API IOptimusNonCopyableNode
+{
+	GENERATED_BODY()
+
+};
+
diff --git a/Engine/Plugins/Animation/DeformerGraph/Source/OptimusCore/Public/OptimusActionStack.h b/Engine/Plugins/Animation/DeformerGraph/Source/OptimusCore/Public/OptimusActionStack.h
--- a/Engine/Plugins/Animation/DeformerGraph/Source/OptimusCore/Public/OptimusActionStack.h
+++ b/Engine/Plugins/Animation/DeformerGraph/Source/OptimusCore/Public/OptimusActionStack.h
@@ -83,6 +83,7 @@ public:
 
 	~FOptimusActionScope();
 
+	void SetTitle(const FString& InTitle);
 private:
 	UOptimusActionStack& ActionStack;
 };
diff --git a/Engine/Plugins/Animation/DeformerGraph/Source/OptimusCore/Public/OptimusHelpers.h b/Engine/Plugins/Animation/DeformerGraph/Source/OptimusCore/Public/OptimusHelpers.h
--- a/Engine/Plugins/Animation/DeformerGraph/Source/OptimusCore/Public/OptimusHelpers.h
+++ b/Engine/Plugins/Animation/DeformerGraph/Source/OptimusCore/Public/OptimusHelpers.h
@@ -53,11 +53,11 @@ namespace Optimus
 	struct FUniqueNameGenerator
 	{
 		explicit FUniqueNameGenerator(UObject* InScopeObject);
+		FName GetUniqueName(FName InName);
 		
+	private:	
 		UObject* ScopeObject = nullptr;
 		TArray<FName> GeneratedName;
-
-		FName GetUniqueName(FName InName);
 	};
 	
 	/** A small helper class to enable binary reads on an archive, since the 
diff --git a/Engine/Plugins/Animation/DeformerGraph/Source/OptimusCore/Public/OptimusNodeGraph.h b/Engine/Plugins/Animation/DeformerGraph/Source/OptimusCore/Public/OptimusNodeGraph.h
--- a/Engine/Plugins/Animation/DeformerGraph/Source/OptimusCore/Public/OptimusNodeGraph.h
+++ b/Engine/Plugins/Animation/DeformerGraph/Source/OptimusCore/Public/OptimusNodeGraph.h
@@ -80,7 +80,7 @@ public:
 	static const FName DefaultSubGraphRefNodeName;
 	
 	// Function Graphs are addressed in a special way
-	static FString GetFunctionGraphPath(const FString& InFunctionName);
+	static FString GetFunctionGraphCollectionPath(const FString& InFunctionName);
 
 	
 	// Check if the duplication took place at the asset level
@@ -93,8 +93,6 @@ public:
 
 	UOptimusNodeGraph *GetParentGraph() const;
 	
-	FString GetGraphPath() const;
-
 	/** Verify if the given name is a valid graph name. */
 	static bool IsValidUserGraphName(
 		const FString& InGraphName,
@@ -196,9 +194,9 @@ public:
 	bool RemoveNodes(
 		const TArray<UOptimusNode*>& InNodes
 	);
-	bool RemoveNodes(
-		const TArray<UOptimusNode*>& InNodes,
-		const FString& InActionName
+
+	int32 RemoveNodesAndCount(
+		const TArray<UOptimusNode*>& InNodes
 	);
 
 	UFUNCTION(BlueprintCallable, Category = OptimusNodeGraph)
@@ -390,6 +388,7 @@ public:
 
 	const TArray<UOptimusNode*>& GetAllNodes() const { return Nodes; }
 	const TArray<UOptimusNodeLink*>& GetAllLinks() const { return Links; }
+	const TArray<UOptimusNodePair*>& GetAllNodePairs() const { return NodePairs; }
 
 	UOptimusActionStack* GetActionStack() const;      
 
@@ -466,10 +465,9 @@ protected:
 		UOptimusNode* InNode
 	);
 	
-	bool RemoveNodesToAction(
-		FOptimusCompoundAction* InOutAction,
-		const TArray<UOptimusNode*>& InNodes
-		) const;
+	void RemoveGraph(
+		UOptimusNodeGraph* InNodeGraph
+		);
 	
 	// Remove a node directly. If a node still has connections this call will fail. 
 	bool RemoveNodeDirect(
@@ -535,7 +533,12 @@ private:
 		TArray<UOptimusNodeSubGraph*>& OutSubGraphs
 		);
 
-	bool DuplicateSubGraph(UOptimusNodeSubGraph* InSourceSubGraph, FName InNewGraphName);
+	static bool DuplicateSubGraph(
+		UOptimusActionStack* InActionStack,
+		const FString& InGraphOwnerPath,
+		UOptimusNodeSubGraph* InSourceSubGraph,
+		FName InNewGraphName
+		);
 
 	void RemoveNodePairByIndex(int32 NodePairIndex);
 	
@@ -564,6 +567,7 @@ private:
 		) const;
 
 	FString ConstructSubGraphPath(const FString& InSubGraphName) const;
+	static FString ConstructSubGraphPath(const FString& InGraphOwnerPath, const FString& InSubGraphName);
 	
 	UPROPERTY(NonTransactional)
 	TArray<TObjectPtr<UOptimusNode>> Nodes;
@@ -575,7 +579,7 @@ private:
 	UPROPERTY(NonTransactional)
 	TArray<TObjectPtr<UOptimusNodePair>> NodePairs;
 	
-	UPROPERTY()
+	UPROPERTY(NonTransactional)
 	TArray<TObjectPtr<UOptimusNodeGraph>> SubGraphs;
 
 	FOptimusGraphNotifyDelegate GraphNotifyDelegate;
diff --git a/Engine/Plugins/Animation/DeformerGraph/Source/OptimusCore/Public/OptimusNodeSubGraph.h b/Engine/Plugins/Animation/DeformerGraph/Source/OptimusCore/Public/OptimusNodeSubGraph.h
--- a/Engine/Plugins/Animation/DeformerGraph/Source/OptimusCore/Public/OptimusNodeSubGraph.h
+++ b/Engine/Plugins/Animation/DeformerGraph/Source/OptimusCore/Public/OptimusNodeSubGraph.h
@@ -7,6 +7,7 @@
 #include "OptimusNodeSubGraph.generated.h"
 
 class UOptimusNode_GraphTerminal;
+enum class EOptimusTerminalType;
 
 UCLASS()
 class OPTIMUSCORE_API UOptimusNodeSubGraph :
@@ -31,7 +32,8 @@ public:
 	EOptimusDataTypeUsageFlags GetTypeUsageFlags(const FOptimusDataDomain& InDataDomain) const override;
 
 	UOptimusComponentSourceBinding* GetDefaultComponentBinding(const FOptimusPinTraversalContext& InTraversalContext) const;
-	
+
+	UOptimusNode_GraphTerminal* GetTerminalNode(EOptimusTerminalType InTerminalType) const;
 
 	DECLARE_MULTICAST_DELEGATE_OneParam(FOnBindingArrayPasted, FName);
 	DECLARE_MULTICAST_DELEGATE_OneParam(FOnBindingValueChanged, FName);
@@ -58,11 +60,6 @@ public:
 	UPROPERTY(EditAnywhere, Category=Bindings)
 	FOptimusParameterBindingArray OutputBindings;
 
-	UPROPERTY()
-	TWeakObjectPtr<UOptimusNode_GraphTerminal> EntryNode; 
-
-	UPROPERTY()
-	TWeakObjectPtr<UOptimusNode_GraphTerminal> ReturnNode;
 
 
 private:
diff --git a/Engine/Plugins/Animation/DeformerGraph/Source/OptimusEditor/Private/OptimusEditor.cpp b/Engine/Plugins/Animation/DeformerGraph/Source/OptimusEditor/Private/OptimusEditor.cpp
--- a/Engine/Plugins/Animation/DeformerGraph/Source/OptimusEditor/Private/OptimusEditor.cpp
+++ b/Engine/Plugins/Animation/DeformerGraph/Source/OptimusEditor/Private/OptimusEditor.cpp
@@ -497,7 +497,21 @@ void FOptimusEditor::CopySelectedNodes() const
 
 bool FOptimusEditor::CanCopyNodes() const
 {
-	return !GetSelectedModelNodes().IsEmpty();
+	if (GraphEditorWidget->GetSelectedNodes().IsEmpty())
+	{
+		return false;
+	}
+	
+	for (UObject* Object : GraphEditorWidget->GetSelectedNodes())
+	{
+		UEdGraphNode* GraphNode = Cast<UEdGraphNode>(Object);
+		if (GraphNode && !GraphNode->CanDuplicateNode())
+		{
+			return false;
+		}
+	}
+	
+	return true;
 }
 
 
@@ -509,7 +523,10 @@ void FOptimusEditor::CutSelectedNodes() const
 	const TArray<UOptimusNode*> ModelNodes = GetSelectedModelNodes();
 	Clipboard.SetClipboardFromNodes(ModelNodes);
 
-	ModelGraph->RemoveNodes(ModelNodes, TEXT("Cut"));
+	FOptimusActionScope ActionScope(*GetActionStack(), {});
+	int32 NumCut = ModelGraph->RemoveNodesAndCount(ModelNodes);
+	const FString& ActionTitle = NumCut == 1 ? TEXT("Cut Node") : FString::Printf(TEXT("Cut %d Nodes"), NumCut);
+	ActionScope.SetTitle(ActionTitle);
 }
 
 
@@ -556,12 +573,7 @@ void FOptimusEditor::DuplicateNodes() const
 
 bool FOptimusEditor::CanDuplicateNodes() const
 {
-	if (IsGraphReadOnly())
-	{
-		return false;
-	}
-	
-	return !GetSelectedModelNodes().IsEmpty();
+	return CanCopyNodes() && !IsGraphReadOnly();
 }
 
 
@@ -669,11 +681,12 @@ bool FOptimusEditor::CanCollapseNodes() const
 
 void FOptimusEditor::ExpandCollapsedNode()
 {
-	FOptimusActionScope ActionScope(*GetActionStack(), TEXT("Expand Collapsed Nodes"));
 	UOptimusNodeGraph* ModelGraph = EditorGraph->GetModelGraph();
 	TArray<UObject*> NewNodes;
 	for (UOptimusNode* ModelNode: GetSelectedModelNodes())
 	{
+		// Note: currently each expand is one undo step due to limitation on move node not able to move a
+		// node that were removed as part of a previous expansion
 		ModelGraph->ExpandCollapsedNodes(ModelNode);
 	}
 }
@@ -686,7 +699,15 @@ bool FOptimusEditor::CanExpandCollapsedNode() const
 		return false;
 	}
 	
-	const TArray<UOptimusNode*> ModelNodes = GetSelectedModelNodes();	
+	const TArray<UOptimusNode*> ModelNodes = GetSelectedModelNodes();
+
+	// Note: currently we only allow one node expansion at a time due to limitation on move node not able to move a
+	// node that were removed as part of a previous expansion 
+	if (ModelNodes.Num() != 1)
+	{
+		return false;
+	}
+	
 	const UOptimusNodeGraph* ModelGraph = EditorGraph->GetModelGraph();	
 	for (UOptimusNode* ModelNode: ModelNodes)
 	{
diff --git a/Engine/Plugins/Animation/DeformerGraph/Source/OptimusEditor/Private/OptimusEditorClipboard.cpp b/Engine/Plugins/Animation/DeformerGraph/Source/OptimusEditor/Private/OptimusEditorClipboard.cpp
--- a/Engine/Plugins/Animation/DeformerGraph/Source/OptimusEditor/Private/OptimusEditorClipboard.cpp
+++ b/Engine/Plugins/Animation/DeformerGraph/Source/OptimusEditor/Private/OptimusEditorClipboard.cpp
@@ -440,6 +440,24 @@ bool FOptimusEditorClipboard::ProcessPostCreateObject(UObject* InRootOuter, UObj
 		{
 			UOptimusNodeSubGraph* SubGraph = Cast<UOptimusNodeSubGraph>(InNewObject);
 
+			// Objects created by the subgraph needs to init their transient data as well
+			TQueue<UOptimusNodeGraph*> SubGraphQueue;
+			SubGraphQueue.Enqueue(SubGraph);
+			UOptimusNodeGraph* WorkingSubGraph;
+			while(SubGraphQueue.Dequeue(WorkingSubGraph))
+			{
+				for (UOptimusNode* Node : WorkingSubGraph->GetAllNodes())
+				{
+					Node->InitializeTransientData();
+				}
+
+				for (UOptimusNodeGraph* SubGraphToAdd : WorkingSubGraph->GetGraphs())
+				{
+					SubGraphQueue.Enqueue(SubGraphToAdd);
+				}
+			}
+
+			
 			return Graph->AddGraphDirect(SubGraph, INDEX_NONE);
 		}
 	}
diff --git a/Engine/Plugins/Animation/DeformerGraph/Source/OptimusEditor/Private/OptimusEditorGraphNode.cpp b/Engine/Plugins/Animation/DeformerGraph/Source/OptimusEditor/Private/OptimusEditorGraphNode.cpp
--- a/Engine/Plugins/Animation/DeformerGraph/Source/OptimusEditor/Private/OptimusEditorGraphNode.cpp
+++ b/Engine/Plugins/Animation/DeformerGraph/Source/OptimusEditor/Private/OptimusEditorGraphNode.cpp
@@ -14,6 +14,7 @@
 #include "OptimusNode.h"
 #include "OptimusNodePin.h"
 #include "IOptimusNodeAdderPinProvider.h"
+#include "IOptimusNonCopyableNode.h"
 #include "IOptimusPinMutabilityDefiner.h"
 #include "OptimusComponentSource.h"
 #include "OptimusDataTypeRegistry.h"
@@ -265,6 +266,11 @@ void UOptimusEditorGraphNode::SyncDiagnosticStateWithModelNode()
 }
 
 
+bool UOptimusEditorGraphNode::CanDuplicateNode() const
+{
+	return Super::CanDuplicateNode() && (!ModelNode || Cast<IOptimusNonCopyableNode>(ModelNode) == nullptr);
+}
+
 bool UOptimusEditorGraphNode::CanUserDeleteNode() const
 {
 	return Super::CanUserDeleteNode() && (!ModelNode || ModelNode->CanUserDeleteNode());
diff --git a/Engine/Plugins/Animation/DeformerGraph/Source/OptimusEditor/Private/OptimusEditorGraphNode.h b/Engine/Plugins/Animation/DeformerGraph/Source/OptimusEditor/Private/OptimusEditorGraphNode.h
--- a/Engine/Plugins/Animation/DeformerGraph/Source/OptimusEditor/Private/OptimusEditorGraphNode.h
+++ b/Engine/Plugins/Animation/DeformerGraph/Source/OptimusEditor/Private/OptimusEditorGraphNode.h
@@ -44,6 +44,7 @@ public:
 	FOptimusNodePinExpansionChanged& OnNodePinExpansionChanged() { return NodePinExpansionChanged; }
 	
 	// UEdGraphNode overrides
+	bool CanDuplicateNode() const override;
 	bool CanUserDeleteNode() const override;
 	FText GetNodeTitle(ENodeTitleType::Type InTitleType) const override;
 	void GetNodeContextMenuActions(UToolMenu* InMenu, UGraphNodeContextMenuContext* InContext) const override;
diff --git a/Engine/Plugins/Animation/DeformerGraph/Source/OptimusEditor/Private/OptimusEditorGraphSchemaActions.cpp b/Engine/Plugins/Animation/DeformerGraph/Source/OptimusEditor/Private/OptimusEditorGraphSchemaActions.cpp
--- a/Engine/Plugins/Animation/DeformerGraph/Source/OptimusEditor/Private/OptimusEditorGraphSchemaActions.cpp
+++ b/Engine/Plugins/Animation/DeformerGraph/Source/OptimusEditor/Private/OptimusEditorGraphSchemaActions.cpp
@@ -159,7 +159,7 @@ FOptimusSchemaAction_Graph::FOptimusSchemaAction_Graph(
 		), 
 		GraphType(InGraph->GetGraphType())
 {
-	GraphPath = InGraph->GetGraphPath();
+	GraphPath = InGraph->GetCollectionPath();
 }
 
 
diff --git a/Engine/Plugins/Animation/DeformerGraph/Source/OptimusEditor/Private/Widgets/SOptimusEditorGraphExplorer.cpp b/Engine/Plugins/Animation/DeformerGraph/Source/OptimusEditor/Private/Widgets/SOptimusEditorGraphExplorer.cpp
--- a/Engine/Plugins/Animation/DeformerGraph/Source/OptimusEditor/Private/Widgets/SOptimusEditorGraphExplorer.cpp
+++ b/Engine/Plugins/Animation/DeformerGraph/Source/OptimusEditor/Private/Widgets/SOptimusEditorGraphExplorer.cpp
@@ -308,18 +308,6 @@ TSharedRef<SWidget> SOptimusEditorGraphExplorer::OnCreateWidgetForAction(FCreate
 	return SNew(SOptimusEditorGraphExplorerItem, InCreateData, OptimusEditor.Pin());
 }
 
-static FText GetGraphSubCategory(UOptimusNodeGraph* InGraph)	
-{
-	if (InGraph->GetGraphType() == EOptimusNodeGraphType::ExternalTrigger)
-	{
-		return FText::FromString(TEXT("Triggered Graphs"));
-	}
-	else
-	{
-		return FText::GetEmpty();
-	}
-}
-
 void SOptimusEditorGraphExplorer::CollectAllActions(FGraphActionListBuilderBase& OutAllActions)
 {
 	TSharedPtr<FOptimusEditor> Editor = OptimusEditor.Pin();
@@ -337,11 +325,10 @@ void SOptimusEditorGraphExplorer::CollectAllActions(FGraphActionListBuilderBase&
 	
 	for (UOptimusNodeGraph* Graph : Deformer->GetGraphs())
 	{
-		FText GraphCategory = GetGraphSubCategory(Graph);
 		TSharedPtr<FOptimusSchemaAction_Graph> GraphAction = MakeShared<FOptimusSchemaAction_Graph>(Graph);
 		OutAllActions.AddAction(GraphAction);
 
-		CollectChildGraphActions(OutAllActions, Graph, GraphCategory);
+		CollectChildGraphActions(OutAllActions, Graph, {});
 	}
 	
 	for (UOptimusComponentSourceBinding* Binding : Deformer->GetComponentBindings())
