Fix incorrect buffer type crash when falling back to default empty buffer.
This caused an assert on Vulkan.
#rnx

diff --git a/Engine/Plugins/Animation/DeformerGraph/Source/OptimusCore/Private/DataInterfaces/OptimusDataInterfaceHalfEdge.cpp b/Engine/Plugins/Animation/DeformerGraph/Source/OptimusCore/Private/DataInterfaces/OptimusDataInterfaceHalfEdge.cpp
--- a/Engine/Plugins/Animation/DeformerGraph/Source/OptimusCore/Private/DataInterfaces/OptimusDataInterfaceHalfEdge.cpp
+++ b/Engine/Plugins/Animation/DeformerGraph/Source/OptimusCore/Private/DataInterfaces/OptimusDataInterfaceHalfEdge.cpp
@@ -11,6 +11,7 @@
 #include "Rendering/SkeletalMeshRenderData.h"
 #include "ShaderCompilerCore.h"
 #include "ShaderParameterMetadataBuilder.h"
+#include "SystemTextures.h"
 #include "SkeletalRenderPublic.h"
 
 #include UE_INLINE_GENERATED_CPP_BY_NAME(OptimusDataInterfaceHalfEdge)
@@ -233,7 +234,7 @@ bool FOptimusHalfEdgeDataProviderProxy::IsValid(FValidationData const& InValidat
 
 void FOptimusHalfEdgeDataProviderProxy::AllocateResources(FRDGBuilder& GraphBuilder)
 {
-	FallbackSRV = GraphBuilder.CreateSRV(GraphBuilder.RegisterExternalBuffer(GWhiteVertexBufferWithRDG->Buffer), PF_R16G16B16A16_UINT);
+	FallbackSRV = GraphBuilder.CreateSRV(FRDGBufferSRVDesc(GSystemTextures.GetDefaultStructuredBuffer(GraphBuilder, sizeof(int32))));
 	
 	const int32 LodIndex = SkeletalMeshObject->GetLOD();
 	FSkeletalMeshRenderData const& SkeletalMeshRenderData = SkeletalMeshObject->GetSkeletalMeshRenderData();
