[Deformer Graph] fix an issue with the unique name generator not generating unique names.

diff --git a/Engine/Plugins/Animation/DeformerGraph/Source/OptimusCore/Private/OptimusHelpers.cpp b/Engine/Plugins/Animation/DeformerGraph/Source/OptimusCore/Private/OptimusHelpers.cpp
--- a/Engine/Plugins/Animation/DeformerGraph/Source/OptimusCore/Private/OptimusHelpers.cpp
+++ b/Engine/Plugins/Animation/DeformerGraph/Source/OptimusCore/Private/OptimusHelpers.cpp
@@ -28,7 +28,7 @@ Optimus::FUniqueNameGenerator::FUniqueNameGenerator(UObject* InScopeObject)
 FName Optimus::FUniqueNameGenerator::GetUniqueName(FName InName)
 {
 	FName Result = Optimus::GetUniqueNameForScope(ScopeObject, InName);
-	GenerateUniqueNameFromExistingNames(Result, GeneratedName);
+	Result = GenerateUniqueNameFromExistingNames(Result, GeneratedName);
 
 	// Result be usable at this point since the name number strictly increases.
 	// Only take the slow route if there is still a name collision for mysterious reasons
