[Deformer Graph] fix extreme lag when changing value on a constant node

diff --git a/Engine/Plugins/Animation/DeformerGraph/Source/OptimusEditor/Private/OptimusEditor.cpp b/Engine/Plugins/Animation/DeformerGraph/Source/OptimusEditor/Private/OptimusEditor.cpp
--- a/Engine/Plugins/Animation/DeformerGraph/Source/OptimusEditor/Private/OptimusEditor.cpp
+++ b/Engine/Plugins/Animation/DeformerGraph/Source/OptimusEditor/Private/OptimusEditor.cpp
@@ -1346,12 +1346,14 @@ void FOptimusEditor::OnDeformerModified(
 	case EOptimusGlobalNotifyType::ResourceIndexChanged:
 	case EOptimusGlobalNotifyType::VariableRenamed:
 	case EOptimusGlobalNotifyType::VariableIndexChanged:
-	case EOptimusGlobalNotifyType::ConstantValueChanged:
 	case EOptimusGlobalNotifyType::NodeTypeAdded: 
 	case EOptimusGlobalNotifyType::NodeTypeRemoved: 
 		RefreshEvent.Broadcast();
 		break;
 
+	case EOptimusGlobalNotifyType::ConstantValueChanged:
+		break;
+		
 	case EOptimusGlobalNotifyType::DataTypeChanged:
 		OnDataTypeChanged();
 		break;
