[Deformer Graph] Make sure subgraphs within a function is also uneditable

diff --git a/Engine/Plugins/Animation/DeformerGraph/Source/OptimusCore/Private/Nodes/OptimusNode_CustomComputeKernel.cpp b/Engine/Plugins/Animation/DeformerGraph/Source/OptimusCore/Private/Nodes/OptimusNode_CustomComputeKernel.cpp
--- a/Engine/Plugins/Animation/DeformerGraph/Source/OptimusCore/Private/Nodes/OptimusNode_CustomComputeKernel.cpp
+++ b/Engine/Plugins/Animation/DeformerGraph/Source/OptimusCore/Private/Nodes/OptimusNode_CustomComputeKernel.cpp
@@ -187,12 +187,7 @@ void UOptimusNode_CustomComputeKernel::SetShaderText(const FString& NewText)
 
 bool UOptimusNode_CustomComputeKernel::IsShaderTextReadOnly() const
 {
-	if (GetOwningGraph()->GetGraphType() == EOptimusNodeGraphType::Function)
-	{
-		return true;
-	}
-
-	return false;
+	return GetOwningGraph()->IsReadOnly();
 }
 
 #endif // WITH_EDITOR
diff --git a/Engine/Plugins/Animation/DeformerGraph/Source/OptimusCore/Private/OptimusNodeGraph.cpp b/Engine/Plugins/Animation/DeformerGraph/Source/OptimusCore/Private/OptimusNodeGraph.cpp
--- a/Engine/Plugins/Animation/DeformerGraph/Source/OptimusCore/Private/OptimusNodeGraph.cpp
+++ b/Engine/Plugins/Animation/DeformerGraph/Source/OptimusCore/Private/OptimusNodeGraph.cpp
@@ -233,6 +233,21 @@ FString UOptimusNodeGraph::ConstructPath(const FString& GraphPath, const FString
 	return Result;
 }
 
+bool UOptimusNodeGraph::IsReadOnly() const
+{
+	if (IsFunctionGraph())
+	{
+		return true;
+	}
+
+	if (UOptimusNodeGraph* GraphOwner = Cast<UOptimusNodeGraph>(GetCollectionOwner()))
+	{
+		return GraphOwner->IsReadOnly();
+	}
+
+	return false;
+}
+
 int32 UOptimusNodeGraph::GetGraphIndex() const
 {
 	return GetCollectionOwner()->GetGraphs().IndexOfByKey(this);
diff --git a/Engine/Plugins/Animation/DeformerGraph/Source/OptimusCore/Public/OptimusNodeGraph.h b/Engine/Plugins/Animation/DeformerGraph/Source/OptimusCore/Public/OptimusNodeGraph.h
--- a/Engine/Plugins/Animation/DeformerGraph/Source/OptimusCore/Public/OptimusNodeGraph.h
+++ b/Engine/Plugins/Animation/DeformerGraph/Source/OptimusCore/Public/OptimusNodeGraph.h
@@ -115,7 +115,9 @@ public:
 	{
 		return GraphType == EOptimusNodeGraphType::Function; 
 	}
-	
+
+	bool IsReadOnly() const;
+
 	UFUNCTION(BlueprintCallable, Category = OptimusNodeGraph)
 	int32 GetGraphIndex() const;
 
diff --git a/Engine/Plugins/Animation/DeformerGraph/Source/OptimusEditor/Private/OptimusEditor.cpp b/Engine/Plugins/Animation/DeformerGraph/Source/OptimusEditor/Private/OptimusEditor.cpp
--- a/Engine/Plugins/Animation/DeformerGraph/Source/OptimusEditor/Private/OptimusEditor.cpp
+++ b/Engine/Plugins/Animation/DeformerGraph/Source/OptimusEditor/Private/OptimusEditor.cpp
@@ -1238,12 +1238,7 @@ FGraphAppearanceInfo FOptimusEditor::GetGraphAppearance() const
 
 bool FOptimusEditor::IsGraphReadOnly() const
 {
-	if (EditorGraph->GetModelGraph()->GetGraphType() == EOptimusNodeGraphType::Function)
-	{
-		return true;
-	}
-
-	return false;
+	return EditorGraph->GetModelGraph()->IsReadOnly();
 }
 
 bool FOptimusEditor::IsGraphEditable() const
@@ -1261,10 +1256,10 @@ bool FOptimusEditor::IsPropertyReadOnly(const FPropertyAndParent& InPropertyAndP
 			continue;
 		}
 
-		// Everything a function graph owns should be read only
-		if (SelectedObject->GetTypedOuter<UOptimusFunctionNodeGraph>())
+		// Everything a read-only graph owns should be read only
+		if (UOptimusNodeGraph* OwnerGraph = SelectedObject->GetTypedOuter<UOptimusNodeGraph>())
 		{
-			return true;
+			return OwnerGraph->IsReadOnly();
 		}
 		
 	}
