[Deformer Graph] Fix editor crash when clearing preview mesh while the transform manipulator is shown in the viewport

diff --git a/Engine/Plugins/Animation/DeformerGraph/Source/OptimusEditor/Private/OptimusEditor.cpp b/Engine/Plugins/Animation/DeformerGraph/Source/OptimusEditor/Private/OptimusEditor.cpp
--- a/Engine/Plugins/Animation/DeformerGraph/Source/OptimusEditor/Private/OptimusEditor.cpp
+++ b/Engine/Plugins/Animation/DeformerGraph/Source/OptimusEditor/Private/OptimusEditor.cpp
@@ -1005,6 +1005,13 @@ void FOptimusEditor::HandlePreviewMeshChanged(
 		EditableSkeleton.Reset();
 	}
 	GetPersonaToolkit()->GetPreviewScene()->SetEditableSkeleton(EditableSkeleton);
+	
+	// Vanilla preview scene relies on the skeleton tree widget to do this step, see SSkeletonTree::OnSelectionChanged
+	// But since we don't have that widget, this has to be done manually
+	if (InOldPreviewMesh != InNewPreviewMesh || InNewPreviewMesh == nullptr)
+	{
+		GetPersonaToolkit()->GetPreviewScene()->DeselectAll();
+	}
 }
 
 
