[Deformer Graph] fix copy kernel dispatched later than the kernels that depends on it. Previously all copy kernels were dispatched at the end, which causes regression in some old asset

+ minor fix , delete function graph was not working

diff --git a/Engine/Plugins/Animation/DeformerGraph/Source/OptimusCore/Private/OptimusDeformer.cpp b/Engine/Plugins/Animation/DeformerGraph/Source/OptimusCore/Private/OptimusDeformer.cpp
--- a/Engine/Plugins/Animation/DeformerGraph/Source/OptimusCore/Private/OptimusDeformer.cpp
+++ b/Engine/Plugins/Animation/DeformerGraph/Source/OptimusCore/Private/OptimusDeformer.cpp
@@ -150,8 +150,10 @@ UOptimusNodeGraph* UOptimusDeformer::GetUpdateGraph() const
 
 bool UOptimusDeformer::RemoveGraph(UOptimusNodeGraph* InGraph)
 {
-	if (UOptimusNodeSubGraph* SubGraph = Cast<UOptimusNodeSubGraph>(InGraph))
+	// Plain subgraph maps 1:1 to a subgraph reference node
+	if (InGraph->GetGraphType() == EOptimusNodeGraphType::SubGraph)
 	{
+		UOptimusNodeSubGraph* SubGraph = CastChecked<UOptimusNodeSubGraph>(InGraph);
 		if (UOptimusNode* Node = GetSubGraphReferenceNode(SubGraph))
 		{
 			FOptimusActionScope ActionScope(*GetActionStack(), TEXT("Remove SubGraph"));
@@ -2035,7 +2037,7 @@ TArray<FOptimusComputeGraphInfo> UOptimusDeformer::CompileNodeGraphToComputeGrap
 						}
 						else
 						{
-							check(InstancedNode.LoopIndex == 0);
+							// InstancedNode.LoopIndex should be at its max (i.e. outside of a loop);
 							
 							UOptimusNodePin* LoopTerminalInputPin = nullptr;
 							if (LoopTerminalToSkip.Contains(OtherRoutedNode))
@@ -2555,6 +2557,9 @@ TArray<FOptimusComputeGraphInfo> UOptimusDeformer::CompileNodeGraphToComputeGrap
 		};
 
 		TArray<FKernelWithDataBindings> BoundKernels;
+		// Copy Kernel uses this map to look up the earliest point it can dispatch
+		// It needs to be dispatched after its input kernels but before all its output kernels
+		TMap<FOptimusInstancedNode, UComputeKernel*> InstancedNodeToComputeKernel;
 		
 		for (const FOptimusInstancedNode& InstancedNode : InstancedNodes )
 		{
@@ -2582,6 +2587,9 @@ TArray<FOptimusComputeGraphInfo> UOptimusDeformer::CompileNodeGraphToComputeGrap
 					KernelDataInterface,
 					BoundKernel.InputDataBindings, BoundKernel.OutputDataBindings
 				);
+
+				InstancedNodeToComputeKernel.Add(InstancedNode, BoundKernel.Kernel);
+				
 				if (FText* ErrorMessage = KernelSourceResult.TryGet<FText>())
 				{
 					AddDiagnostic(EOptimusDiagnosticLevel::Error,
@@ -2629,20 +2637,31 @@ TArray<FOptimusComputeGraphInfo> UOptimusDeformer::CompileNodeGraphToComputeGrap
 		
 		if (InNodeGraph->GetGraphType() == GraphInfo.GraphType)
 		{
+			TArray<FKernelWithDataBindings> BoundCopyKernels;
+			TArray<UComputeKernel*> InsertAfterComputeKernelLookUpArray;
 			for (TPair<FOptimusInstancedPin, TArray<FOptimusInstancedPin>> OutputLink : LinksToInsertCopyKernel)
 			{
 				// Create a copy kernel per source pin, that copies from 1 source pin to multiple target pins
 				const FOptimusInstancedPin& SourceInstancedPin = OutputLink.Key;
+
+				// Indicates that this copy kernel should run immediately after the found compute kernel
+				// Note: nullptr means that the copy kernel has no kernel node dependency and thus should run before everything else
+				UComputeKernel* InsertAfterComputeKernel = nullptr;
+				if (UComputeKernel** ComputeKnernel = InstancedNodeToComputeKernel.Find(SourceInstancedPin.InstancedNode))
+				{
+					InsertAfterComputeKernel = *ComputeKnernel;
+				}
+				
 				const FShaderValueTypeHandle ValueType = SourceInstancedPin.Pin->GetDataType()->ShaderValueType;
 				
-				FKernelWithDataBindings BoundKernel;
+				FKernelWithDataBindings BoundCopyKernel;
 
-				BoundKernel.Kernel = NewObject<UComputeKernel>(this);
+				BoundCopyKernel.Kernel = NewObject<UComputeKernel>(this);
 
-				FOptimus_InterfaceBindingMap& InputDataBindings = BoundKernel.InputDataBindings;
-				FOptimus_InterfaceBindingMap& OutputDataBindings = BoundKernel.OutputDataBindings;
+				FOptimus_InterfaceBindingMap& InputDataBindings = BoundCopyKernel.InputDataBindings;
+				FOptimus_InterfaceBindingMap& OutputDataBindings = BoundCopyKernel.OutputDataBindings;
 
-				UOptimusKernelSource* KernelSource = NewObject<UOptimusKernelSource>(BoundKernel.Kernel);
+				UOptimusKernelSource* KernelSource = NewObject<UOptimusKernelSource>(BoundCopyKernel.Kernel);
 				FString SourceText;
 				SourceText = TEXT("if (Index >= ReadNumThreads().x) return;\n");
 
@@ -2732,15 +2751,47 @@ TArray<FOptimusComputeGraphInfo> UOptimusDeformer::CompileNodeGraphToComputeGrap
 
 				static const FString CopyKernelName = TEXT("CopyKernel");
 				static const FIntVector GroupSize = FIntVector(64, 1, 1);
-				FString CookedSource = Optimus::GetCookedKernelSource(BoundKernel.Kernel->GetPathName(), SourceText, CopyKernelName, GroupSize);
+				FString CookedSource = Optimus::GetCookedKernelSource(BoundCopyKernel.Kernel->GetPathName(), SourceText, CopyKernelName, GroupSize);
 				KernelSource->SetSource(CookedSource);
 				KernelSource->EntryPoint = CopyKernelName;
 				KernelSource->GroupSize = GroupSize;
-				BoundKernel.Kernel->KernelSource = KernelSource;
+				BoundCopyKernel.Kernel->KernelSource = KernelSource;
+
+				BoundCopyKernels.Add(BoundCopyKernel);
+				InsertAfterComputeKernelLookUpArray.Add(InsertAfterComputeKernel);
+			}
+
+			check(BoundCopyKernels.Num() == InsertAfterComputeKernelLookUpArray.Num());
+			
+			// Insert copy kernels immediately after the kernel that they are copying from
+			for (int32 CopyKernelIndex = BoundCopyKernels.Num()-1 ; CopyKernelIndex >= 0; CopyKernelIndex--)
+			{
+				const FKernelWithDataBindings& BoundCopyKernel = BoundCopyKernels[CopyKernelIndex];
+				UComputeKernel* InsertAfterComputeKernel = InsertAfterComputeKernelLookUpArray[CopyKernelIndex];
+
+				int32 BoundKernelIndex = BoundKernels.IndexOfByPredicate([InsertAfterComputeKernel]( const FKernelWithDataBindings& InBoundKernel)
+				{
+					return InBoundKernel.Kernel == InsertAfterComputeKernel;
+				});
+
+				int32 InsertIndex = INDEX_NONE;
+				if (BoundKernelIndex == INDEX_NONE)
+				{
+					// Insert to the beginning if there is no kernel dependency
+					InsertIndex = 0;
+				}
+				else
+				{
+					// By Default insert after the kernel that this copy kernel is copying from
+					InsertIndex = BoundKernelIndex + 1;
+				}
 				
-				BoundKernels.Add(BoundKernel);
-				ComputeGraph->KernelInvocations.Add(BoundKernel.Kernel);
-				ComputeGraph->KernelToNode.Add(nullptr);
+				if (ensure(InsertIndex >= 0) && ensure(InsertIndex <= BoundKernels.Num()))
+				{
+					BoundKernels.Insert(BoundCopyKernel, InsertIndex);
+					ComputeGraph->KernelInvocations.Insert(BoundCopyKernel.Kernel, InsertIndex);
+					ComputeGraph->KernelToNode.Insert(nullptr, InsertIndex);
+				}
 			}
 		}
 
