Fix for vulkan buffer/structured buffer mismatch.
Also touched shader code to ensure that it layout is updated.

diff --git a/Engine/Plugins/Animation/DeformerGraph/Source/OptimusCore/Private/DataInterfaces/OptimusDataInterfaceHalfEdge.cpp b/Engine/Plugins/Animation/DeformerGraph/Source/OptimusCore/Private/DataInterfaces/OptimusDataInterfaceHalfEdge.cpp
--- a/Engine/Plugins/Animation/DeformerGraph/Source/OptimusCore/Private/DataInterfaces/OptimusDataInterfaceHalfEdge.cpp
+++ b/Engine/Plugins/Animation/DeformerGraph/Source/OptimusCore/Private/DataInterfaces/OptimusDataInterfaceHalfEdge.cpp
@@ -113,12 +113,12 @@ BEGIN_SHADER_PARAMETER_STRUCT(FHalfEdgeDataInterfaceParameters, )
 	SHADER_PARAMETER(uint32, NumTriangles)
 	SHADER_PARAMETER(uint32, IndexBufferStart)
 	SHADER_PARAMETER(uint32, InputStreamStart)
-	SHADER_PARAMETER_RDG_BUFFER_SRV(Buffer<int>, VertexToEdgeBuffer)
-	SHADER_PARAMETER_RDG_BUFFER_SRV(Buffer<int>, EdgeToTwinEdgeBuffer)
+	SHADER_PARAMETER_RDG_BUFFER_SRV(StructuredBuffer<int>, VertexToEdgeBuffer)
+	SHADER_PARAMETER_RDG_BUFFER_SRV(StructuredBuffer<int>, EdgeToTwinEdgeBuffer)
 
 	SHADER_PARAMETER(uint32, bUseBufferFromRenderData)
-	SHADER_PARAMETER_SRV(Buffer<int>, RenderDataVertexToEdgeBuffer)
-	SHADER_PARAMETER_SRV(Buffer<int>, RenderDataEdgeToTwinEdgeBuffer)
+	SHADER_PARAMETER_SRV(StructuredBuffer<int>, RenderDataVertexToEdgeBuffer)
+	SHADER_PARAMETER_SRV(StructuredBuffer<int>, RenderDataEdgeToTwinEdgeBuffer)
 END_SHADER_PARAMETER_STRUCT()
 
 void UOptimusHalfEdgeDataInterface::GetShaderParameters(TCHAR const* UID, FShaderParametersMetadataBuilder& InOutBuilder, FShaderParametersMetadataAllocations& InOutAllocations) const
@@ -274,7 +274,7 @@ void FOptimusHalfEdgeDataProviderProxy::GatherDispatchData(FDispatchData const&
 	FRHIShaderResourceView* RenderDataVertexToEdgeBufferSRV = LodRenderData->HalfEdgeBuffer.GetVertexToEdgeBufferSRV();
 	FRHIShaderResourceView* RenderDataEdgeToTwinEdgeBufferSRV = LodRenderData->HalfEdgeBuffer.GetEdgeToTwinEdgeBufferSRV();
 	
-	FRHIShaderResourceView* NullSRVBinding = GWhiteVertexBufferWithSRV->ShaderResourceViewRHI.GetReference();
+	FRHIShaderResourceView* NullSRVBinding = GEmptyStructuredBufferWithUAV->ShaderResourceViewRHI.GetReference();
 
 	const TStridedView<FParameters> ParameterArray = MakeStridedParameterView<FParameters>(InDispatchData);
 	for (int32 InvocationIndex = 0; InvocationIndex < ParameterArray.Num(); ++InvocationIndex)
