[Deformer Graph] avoid crash when skeletal mesh using the mesh deformer does not have half edge data cooked and cannot fallback to construct half edge data on the fly

diff --git a/Engine/Plugins/Animation/DeformerGraph/Source/OptimusCore/Private/DataInterfaces/OptimusDataInterfaceHalfEdge.cpp b/Engine/Plugins/Animation/DeformerGraph/Source/OptimusCore/Private/DataInterfaces/OptimusDataInterfaceHalfEdge.cpp
--- a/Engine/Plugins/Animation/DeformerGraph/Source/OptimusCore/Private/DataInterfaces/OptimusDataInterfaceHalfEdge.cpp
+++ b/Engine/Plugins/Animation/DeformerGraph/Source/OptimusCore/Private/DataInterfaces/OptimusDataInterfaceHalfEdge.cpp
@@ -221,7 +221,7 @@ bool FOptimusHalfEdgeDataProviderProxy::IsValid(FValidationData const& InValidat
 	// Invalid if there is no cooked buffer and run time generation isn't possible either
 	const FRawStaticIndexBuffer16or32Interface* IndexBuffer = LodRenderData->MultiSizeIndexContainer.GetIndexBuffer();
 	const FPositionVertexBuffer& VertexBuffer = LodRenderData->StaticVertexBuffers.PositionVertexBuffer;
-	if (!LodRenderData->HalfEdgeBuffer.IsInitialized())
+	if (!LodRenderData->HalfEdgeBuffer.IsReadyForRendering())
 	{
 		if (!IndexBuffer->GetNeedsCPUAccess() || !VertexBuffer.GetAllowCPUAccess())
 		{
