[Deformer Graph] fix crash when deleting a node while drag moving it

diff --git a/Engine/Plugins/Animation/DeformerGraph/Source/OptimusEditor/Private/Widgets/SOptimusEditorGraphNode.cpp b/Engine/Plugins/Animation/DeformerGraph/Source/OptimusEditor/Private/Widgets/SOptimusEditorGraphNode.cpp
--- a/Engine/Plugins/Animation/DeformerGraph/Source/OptimusEditor/Private/Widgets/SOptimusEditorGraphNode.cpp
+++ b/Engine/Plugins/Animation/DeformerGraph/Source/OptimusEditor/Private/Widgets/SOptimusEditorGraphNode.cpp
@@ -715,7 +715,11 @@ void SOptimusEditorGraphNode::EndUserInteraction() const
 		for (UOptimusEditorGraphNode* SelectedNode : SelectedNodes)
 		{
 			FVector2D Position(SelectedNode->NodePosX, SelectedNode->NodePosY);
-			SelectedNode->ModelNode->SetGraphPosition(Position);
+			// It is possible that the model node is deleted during drag
+			if (SelectedNode->ModelNode)
+			{
+				SelectedNode->ModelNode->SetGraphPosition(Position);
+			}
 		}
 	}
 }
