[Deformer Graph] Fix crash when collapsing a link between component source pin and a group pin

diff --git a/Engine/Plugins/Animation/DeformerGraph/Source/OptimusCore/Private/OptimusNodeGraph.cpp b/Engine/Plugins/Animation/DeformerGraph/Source/OptimusCore/Private/OptimusNodeGraph.cpp
--- a/Engine/Plugins/Animation/DeformerGraph/Source/OptimusCore/Private/OptimusNodeGraph.cpp
+++ b/Engine/Plugins/Animation/DeformerGraph/Source/OptimusCore/Private/OptimusNodeGraph.cpp
@@ -1281,8 +1281,17 @@ UOptimusNode* UOptimusNodeGraph::CollapseNodesToSubGraph(
 			const UOptimusNodePin* InputPin = Link->GetNodeInputPin();
 			FOptimusParameterBinding Binding;
 			Binding.Name = GetUniqueBindingName(InputPin->GetFName());
-			Binding.DataType = InputPin->GetDataType();
-			Binding.DataDomain = InputPin->GetDataDomain();
+			if (!InputPin->IsGroupingPin())
+			{
+				Binding.DataType = InputPin->GetDataType();
+				Binding.DataDomain = InputPin->GetDataDomain();
+			}
+			else
+			{
+				// Grouping pin don't have a type currently so use the output pin instead
+				Binding.DataType = Link->GetNodeOutputPin()->GetDataType();
+				Binding.DataDomain = Link->GetNodeOutputPin()->GetDataDomain();
+			}
 			InputBindings.Add(Binding);
 			BindingNames.Add(Binding.Name);
 			EntryNodePinNames.Add(Link) = Binding.Name;
