[DeformerGraph] fixed a regression: variable nodes wasn't working at all because we changed how unique value names are made. Their look up was still based on the old convention for value names

diff --git a/Engine/Plugins/Animation/DeformerGraph/Source/OptimusCore/Private/DataInterfaces/OptimusDataInterfaceGraph.cpp b/Engine/Plugins/Animation/DeformerGraph/Source/OptimusCore/Private/DataInterfaces/OptimusDataInterfaceGraph.cpp
--- a/Engine/Plugins/Animation/DeformerGraph/Source/OptimusCore/Private/DataInterfaces/OptimusDataInterfaceGraph.cpp
+++ b/Engine/Plugins/Animation/DeformerGraph/Source/OptimusCore/Private/DataInterfaces/OptimusDataInterfaceGraph.cpp
@@ -7,6 +7,7 @@
 #include "ComputeFramework/ShaderParameterMetadataAllocation.h"
 #include "ComputeFramework/ShaderParamTypeDefinition.h"
 #include "OptimusDeformerInstance.h"
+#include "OptimusHelpers.h"
 #include "OptimusVariableDescription.h"
 #include "ShaderParameterMetadataBuilder.h"
 
@@ -93,6 +94,17 @@ UComputeDataProvider* UOptimusGraphDataInterface::CreateDataProvider(TObjectPtr<
 	UOptimusGraphDataProvider* Provider = NewObject<UOptimusGraphDataProvider>();
 	Provider->MeshComponent = Cast<UMeshComponent>(InBinding);
 	Provider->Variables = Variables;
+
+	for (FOptimusGraphVariableDescription& Variable : Provider->Variables)
+	{
+		// When source object was introduced, we also appended a unique index to the value name provided by each value provider
+		// so instead of using the name directly, we need to do this extra step
+		if (!Variable.SourceObject.IsNull())
+		{
+			Variable.CachedSourceValueName = Optimus::ExtractSourceValueName(Variable.Name);
+		}
+	}
+	
 	Provider->ParameterBufferSize = ParameterBufferSize;
 	return Provider;
 }
@@ -157,7 +169,33 @@ FOptimusGraphDataProviderProxy::FOptimusGraphDataProviderProxy(UOptimusDeformerI
 			{
 				if (VariableValue != nullptr)
 				{
-					if (Variable.ValueType == VariableValue->DataType->ShaderValueType && Variable.Name == VariableValue->VariableName.GetPlainNameString())
+					if (Variable.ValueType != VariableValue->DataType->ShaderValueType)
+					{
+						continue;
+					}
+
+					bool bNameMatch = false;
+					
+					// Once upon a time when these values had no source objects, they also just have simple names
+					// so we can directly use the name to find the matching variable
+					if (Variable.SourceObject.IsNull())
+					{
+						if (Variable.Name == VariableValue->VariableName.GetPlainNameString())
+						{
+							bNameMatch = true;
+						}	
+					}
+					else
+					{
+						// When source object was introduced, we also appended a unique index to the value name
+						// so instead of using the name directly we use the source value name
+						if (Variable.CachedSourceValueName == VariableValue->VariableName.GetPlainNameString())
+						{
+							bNameMatch = true;
+						}	
+					}
+
+					if (bNameMatch)
 					{
 						if (ensure(ParameterData.Num() >= Variable.Offset + VariableValue->ValueData.Num()))
 						{
diff --git a/Engine/Plugins/Animation/DeformerGraph/Source/OptimusCore/Private/DataInterfaces/OptimusDataInterfaceGraph.h b/Engine/Plugins/Animation/DeformerGraph/Source/OptimusCore/Private/DataInterfaces/OptimusDataInterfaceGraph.h
--- a/Engine/Plugins/Animation/DeformerGraph/Source/OptimusCore/Private/DataInterfaces/OptimusDataInterfaceGraph.h
+++ b/Engine/Plugins/Animation/DeformerGraph/Source/OptimusCore/Private/DataInterfaces/OptimusDataInterfaceGraph.h
@@ -31,6 +31,9 @@ struct FOptimusGraphVariableDescription
 
 	UPROPERTY()
 	TSoftObjectPtr<UObject> SourceObject;
+
+	// Cached SourceValueName extracted from Name, Initialized by the DataProvider for when SourceObject is not null
+	FString CachedSourceValueName;
 };
 
 /** Compute Framework Data Interface used for marshaling compute graph parameters and variables. */
diff --git a/Engine/Plugins/Animation/DeformerGraph/Source/OptimusCore/Private/Nodes/OptimusNode_ComputeKernelBase.cpp b/Engine/Plugins/Animation/DeformerGraph/Source/OptimusCore/Private/Nodes/OptimusNode_ComputeKernelBase.cpp
--- a/Engine/Plugins/Animation/DeformerGraph/Source/OptimusCore/Private/Nodes/OptimusNode_ComputeKernelBase.cpp
+++ b/Engine/Plugins/Animation/DeformerGraph/Source/OptimusCore/Private/Nodes/OptimusNode_ComputeKernelBase.cpp
@@ -415,7 +415,7 @@ TOptional<FText> UOptimusNode_ComputeKernelBase::ProcessInputPinForComputeKernel
 				// Value nodes bind the single graph data interface.
 				DataInterfaceFuncIndex = InValueNodes.Find(OutputNode);
 				check(DataInterfaceFuncIndex != INDEX_NONE);
-				DataFunctionName = Optimus::MakeUniqueValueName(ValueProvider->GetValueDebugName(), DataInterfaceFuncIndex);
+				DataFunctionName = Optimus::MakeUniqueValueName(ValueProvider->GetValueName(), DataInterfaceFuncIndex);
 			}
 		}
 		else if (Cast<const UOptimusNode_LoopTerminal>(OutputNode))
diff --git a/Engine/Plugins/Animation/DeformerGraph/Source/OptimusCore/Private/Nodes/OptimusNode_ConstantValue.cpp b/Engine/Plugins/Animation/DeformerGraph/Source/OptimusCore/Private/Nodes/OptimusNode_ConstantValue.cpp
--- a/Engine/Plugins/Animation/DeformerGraph/Source/OptimusCore/Private/Nodes/OptimusNode_ConstantValue.cpp
+++ b/Engine/Plugins/Animation/DeformerGraph/Source/OptimusCore/Private/Nodes/OptimusNode_ConstantValue.cpp
@@ -101,7 +101,7 @@ void UOptimusNode_ConstantValue::PostEditChangeProperty(FPropertyChangedEvent& P
 #endif // WITH_EDITOR
 
 
-FString UOptimusNode_ConstantValue::GetValueDebugName() const
+FString UOptimusNode_ConstantValue::GetValueName() const
 {
 	return GetName();
 }
diff --git a/Engine/Plugins/Animation/DeformerGraph/Source/OptimusCore/Private/Nodes/OptimusNode_ConstantValue.h b/Engine/Plugins/Animation/DeformerGraph/Source/OptimusCore/Private/Nodes/OptimusNode_ConstantValue.h
--- a/Engine/Plugins/Animation/DeformerGraph/Source/OptimusCore/Private/Nodes/OptimusNode_ConstantValue.h
+++ b/Engine/Plugins/Animation/DeformerGraph/Source/OptimusCore/Private/Nodes/OptimusNode_ConstantValue.h
@@ -65,7 +65,7 @@ public:
 #endif
 
 	// IOptimusValueProvider overrides 
-	FString GetValueDebugName() const override;
+	FString GetValueName() const override;
 	FOptimusDataTypeRef GetValueType() const override;
 	FShaderValueType::FValue GetShaderValue() const override;
 
diff --git a/Engine/Plugins/Animation/DeformerGraph/Source/OptimusCore/Private/Nodes/OptimusNode_GetVariable.cpp b/Engine/Plugins/Animation/DeformerGraph/Source/OptimusCore/Private/Nodes/OptimusNode_GetVariable.cpp
--- a/Engine/Plugins/Animation/DeformerGraph/Source/OptimusCore/Private/Nodes/OptimusNode_GetVariable.cpp
+++ b/Engine/Plugins/Animation/DeformerGraph/Source/OptimusCore/Private/Nodes/OptimusNode_GetVariable.cpp
@@ -195,7 +195,7 @@ TOptional<FText> UOptimusNode_GetVariable::ValidateForCompile(const FOptimusPinT
 	return {};
 }
 
-FString UOptimusNode_GetVariable::GetValueDebugName() const
+FString UOptimusNode_GetVariable::GetValueName() const
 {
 	if (const UOptimusVariableDescription* Var = VariableDesc.Get())
 	{
diff --git a/Engine/Plugins/Animation/DeformerGraph/Source/OptimusCore/Private/Nodes/OptimusNode_GetVariable.h b/Engine/Plugins/Animation/DeformerGraph/Source/OptimusCore/Private/Nodes/OptimusNode_GetVariable.h
--- a/Engine/Plugins/Animation/DeformerGraph/Source/OptimusCore/Private/Nodes/OptimusNode_GetVariable.h
+++ b/Engine/Plugins/Animation/DeformerGraph/Source/OptimusCore/Private/Nodes/OptimusNode_GetVariable.h
@@ -54,7 +54,7 @@ public:
 	TOptional<FText> ValidateForCompile(const FOptimusPinTraversalContext& InContext) const override;
 	
 	// IOptimusValueProvider overrides 
-	FString GetValueDebugName() const override;
+	FString GetValueName() const override;
 	FOptimusDataTypeRef GetValueType() const override;
 	FShaderValueType::FValue GetShaderValue() const override;
 
diff --git a/Engine/Plugins/Animation/DeformerGraph/Source/OptimusCore/Private/OptimusDeformer.cpp b/Engine/Plugins/Animation/DeformerGraph/Source/OptimusCore/Private/OptimusDeformer.cpp
--- a/Engine/Plugins/Animation/DeformerGraph/Source/OptimusCore/Private/OptimusDeformer.cpp
+++ b/Engine/Plugins/Animation/DeformerGraph/Source/OptimusCore/Private/OptimusDeformer.cpp
@@ -2161,7 +2161,7 @@ TArray<FOptimusComputeGraphInfo> UOptimusDeformer::CompileNodeGraphToComputeGrap
 		if (IOptimusValueProvider const* ValueProvider = Cast<const IOptimusValueProvider>(ValueNode))
 		{
 			FOptimusGraphVariableDescription& ValueNodeDescription = ValueNodeDescriptions.AddDefaulted_GetRef();
-			ValueNodeDescription.Name = Optimus::MakeUniqueValueName(ValueProvider->GetValueDebugName(), ValueNodeIndex);
+			ValueNodeDescription.Name = Optimus::MakeUniqueValueName(ValueProvider->GetValueName(), ValueNodeIndex);
 			ValueNodeDescription.ValueType = ValueProvider->GetValueType()->ShaderValueType;
 			ValueNodeDescription.SourceObject = ValueNode;
 
diff --git a/Engine/Plugins/Animation/DeformerGraph/Source/OptimusCore/Private/OptimusHelpers.cpp b/Engine/Plugins/Animation/DeformerGraph/Source/OptimusCore/Private/OptimusHelpers.cpp
--- a/Engine/Plugins/Animation/DeformerGraph/Source/OptimusCore/Private/OptimusHelpers.cpp
+++ b/Engine/Plugins/Animation/DeformerGraph/Source/OptimusCore/Private/OptimusHelpers.cpp
@@ -308,3 +308,15 @@ FString Optimus::MakeUniqueValueName(const FString& InValueName, int32 InUniqueI
 {
 	return InValueName + TEXT("_") + FString::FromInt(InUniqueIndex);
 }
+
+FString Optimus::ExtractSourceValueName(const FString& InUniqueValueName)
+{
+	int32 LastUnderscoreIndex;
+	InUniqueValueName.FindLastChar(TEXT('_'), LastUnderscoreIndex);
+	if (ensure(LastUnderscoreIndex != INDEX_NONE))
+	{
+		return InUniqueValueName.Left(LastUnderscoreIndex);
+	}
+
+	return InUniqueValueName;
+}
diff --git a/Engine/Plugins/Animation/DeformerGraph/Source/OptimusCore/Public/IOptimusValueProvider.h b/Engine/Plugins/Animation/DeformerGraph/Source/OptimusCore/Public/IOptimusValueProvider.h
--- a/Engine/Plugins/Animation/DeformerGraph/Source/OptimusCore/Public/IOptimusValueProvider.h
+++ b/Engine/Plugins/Animation/DeformerGraph/Source/OptimusCore/Public/IOptimusValueProvider.h
@@ -24,7 +24,7 @@ class IOptimusValueProvider
 
 public:
 	// Returns the value name.
-	virtual FString GetValueDebugName() const = 0;
+	virtual FString GetValueName() const = 0;
 	// Returns the value type.
 	virtual FOptimusDataTypeRef GetValueType() const = 0;
 	// Returns the stored value as a shader-compatible value.
diff --git a/Engine/Plugins/Animation/DeformerGraph/Source/OptimusCore/Public/OptimusHelpers.h b/Engine/Plugins/Animation/DeformerGraph/Source/OptimusCore/Public/OptimusHelpers.h
--- a/Engine/Plugins/Animation/DeformerGraph/Source/OptimusCore/Public/OptimusHelpers.h
+++ b/Engine/Plugins/Animation/DeformerGraph/Source/OptimusCore/Public/OptimusHelpers.h
@@ -128,4 +128,5 @@ namespace Optimus
 	FName GenerateUniqueNameFromExistingNames(FName InBaseName, const TArray<FName>& InExistingNames);
 
 	FString MakeUniqueValueName(const FString& InValueName, int32 InUniqueIndex);
+	FString ExtractSourceValueName(const FString& InUniqueValueName);
 }
