[Deformer Graph & Mesh Deformer] expose the buffers that a mesh deformer overrides as API

diff --git a/Engine/Plugins/Animation/DeformerGraph/Source/OptimusCore/Private/DataInterfaces/OptimusDataInterfaceSkinnedMeshWrite.cpp b/Engine/Plugins/Animation/DeformerGraph/Source/OptimusCore/Private/DataInterfaces/OptimusDataInterfaceSkinnedMeshWrite.cpp
--- a/Engine/Plugins/Animation/DeformerGraph/Source/OptimusCore/Private/DataInterfaces/OptimusDataInterfaceSkinnedMeshWrite.cpp
+++ b/Engine/Plugins/Animation/DeformerGraph/Source/OptimusCore/Private/DataInterfaces/OptimusDataInterfaceSkinnedMeshWrite.cpp
@@ -13,6 +13,7 @@
 #include "SkeletalMeshDeformerHelpers.h"
 #include "SkeletalRenderPublic.h"
 #include "DataDrivenShaderPlatformInfo.h"
+#include "Animation/MeshDeformerInstance.h"
 
 #include UE_INLINE_GENERATED_CPP_BY_NAME(OptimusDataInterfaceSkinnedMeshWrite)
 
@@ -119,6 +120,25 @@ UComputeDataProvider* UOptimusSkinnedMeshWriteDataInterface::CreateDataProvider(
 	return Provider;
 }
 
+EMeshDeformerOutputBuffer UOptimusSkinnedMeshWriteDataInterface::GetOutputBuffer(int32 InBoundOutputFunctionIndex) const
+{
+	// Maps to the index of functions in GetSupportedOutputs
+	if (InBoundOutputFunctionIndex == 0)
+	{
+		return EMeshDeformerOutputBuffer::SkinnedMeshPosition;
+	}
+	if (InBoundOutputFunctionIndex == 1 || InBoundOutputFunctionIndex == 2)
+	{
+		return EMeshDeformerOutputBuffer::SkinnedMeshTangents;
+	}
+	if (InBoundOutputFunctionIndex == 3)
+	{
+		return EMeshDeformerOutputBuffer::SkinnedMeshVertexColor;
+	}
+
+	return EMeshDeformerOutputBuffer::None;
+}
+
 
 FComputeDataProviderRenderProxy* UOptimusSkinnedMeshWriteDataProvider::GetRenderProxy()
 {
diff --git a/Engine/Plugins/Animation/DeformerGraph/Source/OptimusCore/Private/DataInterfaces/OptimusDataInterfaceSkinnedMeshWrite.h b/Engine/Plugins/Animation/DeformerGraph/Source/OptimusCore/Private/DataInterfaces/OptimusDataInterfaceSkinnedMeshWrite.h
--- a/Engine/Plugins/Animation/DeformerGraph/Source/OptimusCore/Private/DataInterfaces/OptimusDataInterfaceSkinnedMeshWrite.h
+++ b/Engine/Plugins/Animation/DeformerGraph/Source/OptimusCore/Private/DataInterfaces/OptimusDataInterfaceSkinnedMeshWrite.h
@@ -2,6 +2,7 @@
 
 #pragma once
 
+#include "IOptimusOutputBufferWriter.h"
 #include "OptimusComputeDataInterface.h"
 #include "ComputeFramework/ComputeDataProvider.h"
 
@@ -15,7 +16,9 @@ class USkinnedMeshComponent;
 
 /** Compute Framework Data Interface for writing skinned mesh. */
 UCLASS(Category = ComputeFramework)
-class OPTIMUSCORE_API UOptimusSkinnedMeshWriteDataInterface : public UOptimusComputeDataInterface
+class OPTIMUSCORE_API UOptimusSkinnedMeshWriteDataInterface :
+	public UOptimusComputeDataInterface,
+	public IOptimusOutputBufferWriter
 {
 	GENERATED_BODY()
 
@@ -39,6 +42,9 @@ public:
 	UComputeDataProvider* CreateDataProvider(TObjectPtr<UObject> InBinding, uint64 InInputMask, uint64 InOutputMask) const override;
 	//~ End UComputeDataInterface Interface
 
+	//~ Begin IOptimusOutputBufferWriter Interface
+	EMeshDeformerOutputBuffer GetOutputBuffer(int32 InBoundOutputFunctionIndex) const override;
+	//~ End IOptimusOutputBufferWriter Interface 
 private:
 	static TCHAR const* TemplateFilePath;
 };
diff --git a/Engine/Plugins/Animation/DeformerGraph/Source/OptimusCore/Private/OptimusComputeGraph.cpp b/Engine/Plugins/Animation/DeformerGraph/Source/OptimusCore/Private/OptimusComputeGraph.cpp
--- a/Engine/Plugins/Animation/DeformerGraph/Source/OptimusCore/Private/OptimusComputeGraph.cpp
+++ b/Engine/Plugins/Animation/DeformerGraph/Source/OptimusCore/Private/OptimusComputeGraph.cpp
@@ -2,6 +2,7 @@
 
 #include "OptimusComputeGraph.h"
 
+#include "IOptimusOutputBufferWriter.h"
 #include "Components/MeshComponent.h"
 #include "ComputeFramework/ComputeKernelCompileResult.h"
 #include "HAL/FileManager.h"
@@ -14,6 +15,7 @@
 #include "OptimusHelpers.h"
 #include "OptimusNode.h"
 #include "OptimusObjectVersion.h"
+#include "Animation/MeshDeformerInstance.h"
 #include "Misc/UObjectToken.h"
 
 #include UE_INLINE_GENERATED_CPP_BY_NAME(OptimusComputeGraph)
@@ -132,4 +134,23 @@ void UOptimusComputeGraph::OnKernelCompilationComplete(int32 InKernelIndex, FCom
 	}
 }
 
+EMeshDeformerOutputBuffer UOptimusComputeGraph::GetOutputBuffers() const
+{
+	EMeshDeformerOutputBuffer Result = EMeshDeformerOutputBuffer::None;
+	
+	for (const FComputeGraphEdge& GraphEdge : GraphEdges)
+	{
+		if (!GraphEdge.bKernelInput)
+		{
+			const UComputeDataInterface* DataInterface = DataInterfaces[GraphEdge.DataInterfaceIndex];
+			if (const IOptimusOutputBufferWriter* BufferOverrider = Cast<const IOptimusOutputBufferWriter>(DataInterface))
+			{
+				Result |= BufferOverrider->GetOutputBuffer(GraphEdge.DataInterfaceBindingIndex);
+			}
+		}
+	}
+
+	return Result;
+}
+
 #undef LOCTEXT_NAMESPACE
diff --git a/Engine/Plugins/Animation/DeformerGraph/Source/OptimusCore/Private/OptimusComputeGraph.h b/Engine/Plugins/Animation/DeformerGraph/Source/OptimusCore/Private/OptimusComputeGraph.h
--- a/Engine/Plugins/Animation/DeformerGraph/Source/OptimusCore/Private/OptimusComputeGraph.h
+++ b/Engine/Plugins/Animation/DeformerGraph/Source/OptimusCore/Private/OptimusComputeGraph.h
@@ -7,6 +7,7 @@
 
 #include "OptimusComputeGraph.generated.h"
 
+enum class EMeshDeformerOutputBuffer : uint8;
 class UOptimusDeformer;
 class UOptimusNode;
 
@@ -24,6 +25,7 @@ public:
 	// UComputeGraph overrides
 	void OnKernelCompilationComplete(int32 InKernelIndex, FComputeKernelCompileResults const& InCompileResults) override;
 
+	EMeshDeformerOutputBuffer GetOutputBuffers() const;
 protected:
 	// Lookup into Graphs array from the UComputeGraph kernel index. 
 	UPROPERTY()
diff --git a/Engine/Plugins/Animation/DeformerGraph/Source/OptimusCore/Private/OptimusDeformerInstance.cpp b/Engine/Plugins/Animation/DeformerGraph/Source/OptimusCore/Private/OptimusDeformerInstance.cpp
--- a/Engine/Plugins/Animation/DeformerGraph/Source/OptimusCore/Private/OptimusDeformerInstance.cpp
+++ b/Engine/Plugins/Animation/DeformerGraph/Source/OptimusCore/Private/OptimusDeformerInstance.cpp
@@ -563,6 +563,21 @@ void UOptimusDeformerInstance::EnqueueWork(FEnqueueWorkDesc const& InDesc)
 	}
 }
 
+EMeshDeformerOutputBuffer UOptimusDeformerInstance::GetOutputBuffers() const
+{
+	EMeshDeformerOutputBuffer Result = EMeshDeformerOutputBuffer::None;
+
+	for (const FOptimusDeformerInstanceExecInfo& ExecInfo : ComputeGraphExecInfos)
+	{
+		if (const UOptimusComputeGraph* ComputeGraph = Cast<UOptimusComputeGraph>(ExecInfo.ComputeGraph))
+		{
+			Result |= ComputeGraph->GetOutputBuffers();
+		}
+	}
+
+	return Result;
+}
+
 namespace
 {
 	template <typename T>
diff --git a/Engine/Plugins/Animation/DeformerGraph/Source/OptimusCore/Public/IOptimusOutputBufferWriter.h b/Engine/Plugins/Animation/DeformerGraph/Source/OptimusCore/Public/IOptimusOutputBufferWriter.h
new file mode 100644
--- /dev/null
+++ b/Engine/Plugins/Animation/DeformerGraph/Source/OptimusCore/Public/IOptimusOutputBufferWriter.h
@@ -0,0 +1,31 @@
+﻿// Copyright Epic Games, Inc. All Rights Reserved.
+
+#pragma once
+
+#include "UObject/Interface.h"
+
+#include "IOptimusOutputBufferWriter.generated.h"
+
+
+enum class EMeshDeformerOutputBuffer : uint8;
+
+UINTERFACE()
+class OPTIMUSCORE_API UOptimusOutputBufferWriter :
+	public UInterface
+{
+	GENERATED_BODY()
+};
+
+
+/** An interface that should be implemented by data interfaces that overrides some engine internal buffers
+ */
+class IOptimusOutputBufferWriter
+{
+	GENERATED_BODY()
+
+public:
+	/**
+	 * Returns the output buffer that it writes to for the given data interface output function index
+	 */
+	virtual EMeshDeformerOutputBuffer GetOutputBuffer(int32 InBoundOutputFunctionIndex) const = 0;
+};
diff --git a/Engine/Plugins/Animation/DeformerGraph/Source/OptimusCore/Public/OptimusDeformerInstance.h b/Engine/Plugins/Animation/DeformerGraph/Source/OptimusCore/Public/OptimusDeformerInstance.h
--- a/Engine/Plugins/Animation/DeformerGraph/Source/OptimusCore/Public/OptimusDeformerInstance.h
+++ b/Engine/Plugins/Animation/DeformerGraph/Source/OptimusCore/Public/OptimusDeformerInstance.h
@@ -256,6 +256,7 @@ protected:
 	void AllocateResources() override;
 	void ReleaseResources() override;
 	void EnqueueWork(FEnqueueWorkDesc const& InDesc) override;
+	EMeshDeformerOutputBuffer GetOutputBuffers() const override;
 
 private:
 	/** The Mesh Component that owns this Mesh Deformer Instance. */
diff --git a/Engine/Source/Runtime/Engine/Classes/Animation/MeshDeformerInstance.h b/Engine/Source/Runtime/Engine/Classes/Animation/MeshDeformerInstance.h
--- a/Engine/Source/Runtime/Engine/Classes/Animation/MeshDeformerInstance.h
+++ b/Engine/Source/Runtime/Engine/Classes/Animation/MeshDeformerInstance.h
@@ -7,6 +7,15 @@
 
 class FSceneInterface;
 
+enum class EMeshDeformerOutputBuffer : uint8
+{
+	None = 0,
+	SkinnedMeshPosition = 1 << 0,
+	SkinnedMeshTangents = 1 << 1,
+	SkinnedMeshVertexColor = 1 << 2,
+};
+
+ENUM_CLASS_FLAGS(EMeshDeformerOutputBuffer);
 
 /**
  * Base class for mesh deformers instance settings.
@@ -29,6 +38,8 @@ class UMeshDeformerInstance : public UObject
 	GENERATED_BODY()
 
 public:
+
+
 	/** Called to allocate any persistent render resources */
 	ENGINE_API virtual void AllocateResources() PURE_VIRTUAL(, );
 
@@ -65,4 +76,7 @@ public:
 
 	/** Enqueue the mesh deformer workload on a scene. */
 	ENGINE_API virtual void EnqueueWork(FEnqueueWorkDesc const& InDesc) PURE_VIRTUAL(, );
+	
+	/** Return the buffers that this deformer can potentially write to */
+	ENGINE_API virtual EMeshDeformerOutputBuffer GetOutputBuffers() const PURE_VIRTUAL(, return EMeshDeformerOutputBuffer::None; );
 };
