[Deformer Graph] Fixed variables sharing the same base name getting values from the first variable with that base name. Each of them should hold a unique value

diff --git a/Engine/Plugins/Animation/DeformerGraph/Source/OptimusCore/Private/DataInterfaces/OptimusDataInterfaceGraph.cpp b/Engine/Plugins/Animation/DeformerGraph/Source/OptimusCore/Private/DataInterfaces/OptimusDataInterfaceGraph.cpp
--- a/Engine/Plugins/Animation/DeformerGraph/Source/OptimusCore/Private/DataInterfaces/OptimusDataInterfaceGraph.cpp
+++ b/Engine/Plugins/Animation/DeformerGraph/Source/OptimusCore/Private/DataInterfaces/OptimusDataInterfaceGraph.cpp
@@ -180,6 +180,10 @@ FOptimusGraphDataProviderProxy::FOptimusGraphDataProviderProxy(UOptimusDeformerI
 					// so we can directly use the name to find the matching variable
 					if (Variable.SourceObject.IsNull())
 					{
+						// Using GetPlainNameString here because back then variables in the graph data interface were
+						// also generated using GetPlainNameString. This certainly creates an issue where
+						// multiple variables were sharing the same name, but at least the first matching variable
+						// would still work
 						if (Variable.Name == VariableValue->VariableName.GetPlainNameString())
 						{
 							bNameMatch = true;
@@ -189,7 +193,7 @@ FOptimusGraphDataProviderProxy::FOptimusGraphDataProviderProxy(UOptimusDeformerI
 					{
 						// When source object was introduced, we also appended a unique index to the value name
 						// so instead of using the name directly we use the source value name
-						if (Variable.CachedSourceValueName == VariableValue->VariableName.GetPlainNameString())
+						if (Variable.CachedSourceValueName == VariableValue->VariableName)
 						{
 							bNameMatch = true;
 						}	
diff --git a/Engine/Plugins/Animation/DeformerGraph/Source/OptimusCore/Private/Nodes/OptimusNode_GetVariable.cpp b/Engine/Plugins/Animation/DeformerGraph/Source/OptimusCore/Private/Nodes/OptimusNode_GetVariable.cpp
--- a/Engine/Plugins/Animation/DeformerGraph/Source/OptimusCore/Private/Nodes/OptimusNode_GetVariable.cpp
+++ b/Engine/Plugins/Animation/DeformerGraph/Source/OptimusCore/Private/Nodes/OptimusNode_GetVariable.cpp
@@ -199,7 +199,7 @@ FString UOptimusNode_GetVariable::GetValueName() const
 {
 	if (const UOptimusVariableDescription* Var = VariableDesc.Get())
 	{
-		return Var->VariableName.GetPlainNameString();
+		return Var->VariableName.ToString();
 	}
 
 	return {};
