[Deformer Graph] Fixed assertion when opening a user defined struct used by a deformer graph before the deformer graph is opened/loaded in editor

diff --git a/Engine/Plugins/Animation/DeformerGraph/Source/OptimusEditor/Private/OptimusEditorModule.cpp b/Engine/Plugins/Animation/DeformerGraph/Source/OptimusEditor/Private/OptimusEditorModule.cpp
--- a/Engine/Plugins/Animation/DeformerGraph/Source/OptimusEditor/Private/OptimusEditorModule.cpp
+++ b/Engine/Plugins/Animation/DeformerGraph/Source/OptimusEditor/Private/OptimusEditorModule.cpp
@@ -139,7 +139,7 @@ void FOptimusEditorModule::PreChange(const UUserDefinedStruct* Changed,
 			if (Asset.IsInstanceOf(UOptimusDeformer::StaticClass()))
 			{
 				// Only care about loaded assets, which may have active instances
-				if (UOptimusDeformer* DeformerAsset = CastChecked<UOptimusDeformer>(Asset.FastGetAsset(false)))
+				if (UOptimusDeformer* DeformerAsset = Cast<UOptimusDeformer>(Asset.FastGetAsset(false)))
 				{
 					DeformerAsset->SetAllInstancesCanbeActive(false);
 				}
@@ -176,7 +176,7 @@ void FOptimusEditorModule::PostChange(const UUserDefinedStruct* Changed,
 				if (Asset.IsInstanceOf(UOptimusDeformer::StaticClass()))
 				{
 					// Only care about loaded assets, which may have active instances
-					if (UOptimusDeformer* DeformerAsset = CastChecked<UOptimusDeformer>(Asset.FastGetAsset(false)))
+					if (UOptimusDeformer* DeformerAsset = Cast<UOptimusDeformer>(Asset.FastGetAsset(false)))
 					{
 						DeformerAsset->Compile();
 						DeformerAsset->SetAllInstancesCanbeActive(true);
