[Deformer Graph] minor fix for half edge data interface clearing data that the render proxy still relies on during destruction. Which can happen during editor load if there is a deformer in the default map.

diff --git a/Engine/Plugins/Animation/DeformerGraph/Source/OptimusCore/Private/DataInterfaces/OptimusDataInterfaceHalfEdge.cpp b/Engine/Plugins/Animation/DeformerGraph/Source/OptimusCore/Private/DataInterfaces/OptimusDataInterfaceHalfEdge.cpp
--- a/Engine/Plugins/Animation/DeformerGraph/Source/OptimusCore/Private/DataInterfaces/OptimusDataInterfaceHalfEdge.cpp
+++ b/Engine/Plugins/Animation/DeformerGraph/Source/OptimusCore/Private/DataInterfaces/OptimusDataInterfaceHalfEdge.cpp
@@ -173,13 +173,17 @@ UComputeDataProvider* UOptimusHalfEdgeDataInterface::CreateDataProvider(TObjectP
 void UOptimusHalfEdgeDataProvider::BeginDestroy()
 {
 	Super::BeginDestroy();
-	OnDemandHalfEdgeBuffers.Reset();
 	DestroyFence.BeginFence();
 }
 
 bool UOptimusHalfEdgeDataProvider::IsReadyForFinishDestroy()
 {
-	return Super::IsReadyForFinishDestroy() && DestroyFence.IsFenceComplete();
+	if (Super::IsReadyForFinishDestroy() && DestroyFence.IsFenceComplete())
+	{
+		OnDemandHalfEdgeBuffers.Reset();
+		return true;
+	}
+	return false;
 }
 
 FComputeDataProviderRenderProxy* UOptimusHalfEdgeDataProvider::GetRenderProxy()
