[Deformer Graph] Added warning for kernels with no execution domain

diff --git a/Engine/Plugins/Animation/DeformerGraph/Source/OptimusCore/Private/Nodes/OptimusNode_ComputeKernelBase.cpp b/Engine/Plugins/Animation/DeformerGraph/Source/OptimusCore/Private/Nodes/OptimusNode_ComputeKernelBase.cpp
--- a/Engine/Plugins/Animation/DeformerGraph/Source/OptimusCore/Private/Nodes/OptimusNode_ComputeKernelBase.cpp
+++ b/Engine/Plugins/Animation/DeformerGraph/Source/OptimusCore/Private/Nodes/OptimusNode_ComputeKernelBase.cpp
@@ -244,6 +244,11 @@ FOptimus_ComputeKernelResult UOptimusNode_ComputeKernelBase::CreateComputeKernel
 
 TOptional<FText> UOptimusNode_ComputeKernelBase::ValidateForCompile(const FOptimusPinTraversalContext& InContext) const
 {
+	if (!GetExecutionDomain().IsDefined())
+	{
+		return LOCTEXT("NoExecutionDomain", "Kernel has undefined Execution Domain, please use the Details panel to select a valid domain");
+	}
+	
 	auto GetStructTypeDefFromPin = [](const UOptimusNodePin* InPin) -> TOptional<FText>
 	{
 		const FOptimusDataTypeHandle TypeHandle = InPin->GetDataType();
diff --git a/Engine/Plugins/Animation/DeformerGraph/Source/OptimusEditor/Private/OptimusDetailsCustomization.cpp b/Engine/Plugins/Animation/DeformerGraph/Source/OptimusEditor/Private/OptimusDetailsCustomization.cpp
--- a/Engine/Plugins/Animation/DeformerGraph/Source/OptimusEditor/Private/OptimusDetailsCustomization.cpp
+++ b/Engine/Plugins/Animation/DeformerGraph/Source/OptimusEditor/Private/OptimusDetailsCustomization.cpp
@@ -300,6 +300,18 @@ void FOptimusExecutionDomainCustomization::CustomizeHeader(
 
 					return LOCTEXT("MultipleValues", "Multiple Values");
 				})
+				.ColorAndOpacity_Lambda([InPropertyHandle, this]()
+				{
+					if (TOptional<FOptimusExecutionDomain> ExecutionDomain = TryGetSingleExecutionDomain(InPropertyHandle))
+					{
+						if (!ExecutionDomain->IsDefined())
+						{
+							return FSlateColor(FLinearColor::Red);
+						}
+					}
+
+					return FSlateColor::UseStyle();
+				})
 			]
 		]
 		+ SWidgetSwitcher::Slot()
