[Deformer Graph] minor fix, crash caused by removing item from an array while iterating it.

diff --git a/Engine/Plugins/Animation/DeformerGraph/Source/OptimusCore/Private/Nodes/OptimusNode_AnimAttributeDataInterface.cpp b/Engine/Plugins/Animation/DeformerGraph/Source/OptimusCore/Private/Nodes/OptimusNode_AnimAttributeDataInterface.cpp
--- a/Engine/Plugins/Animation/DeformerGraph/Source/OptimusCore/Private/Nodes/OptimusNode_AnimAttributeDataInterface.cpp
+++ b/Engine/Plugins/Animation/DeformerGraph/Source/OptimusCore/Private/Nodes/OptimusNode_AnimAttributeDataInterface.cpp
@@ -185,12 +185,9 @@ void UOptimusNode_AnimAttributeDataInterface::UpdatePinNames()
 
 void UOptimusNode_AnimAttributeDataInterface::ClearOutputPins()
 {
-	for (UOptimusNodePin* Pin : GetPins())
+	for (UOptimusNodePin* Pin : GetPinsByDirection(EOptimusNodePinDirection::Output, false))
 	{
-		if (Pin->GetDirection() == EOptimusNodePinDirection::Output)
-		{
-			RemovePin(Pin);
-		}
+		RemovePin(Pin);
 	}
 }
 
