[Deformer Graph] fixing old code in anim attribute node relying on old property change event behavior

the new property change event behavior was introduced in 33462523

#rnx

diff --git a/Engine/Plugins/Animation/DeformerGraph/Source/OptimusCore/Private/Nodes/OptimusNode_AnimAttributeDataInterface.cpp b/Engine/Plugins/Animation/DeformerGraph/Source/OptimusCore/Private/Nodes/OptimusNode_AnimAttributeDataInterface.cpp
--- a/Engine/Plugins/Animation/DeformerGraph/Source/OptimusCore/Private/Nodes/OptimusNode_AnimAttributeDataInterface.cpp
+++ b/Engine/Plugins/Animation/DeformerGraph/Source/OptimusCore/Private/Nodes/OptimusNode_AnimAttributeDataInterface.cpp
@@ -41,7 +41,7 @@ void UOptimusNode_AnimAttributeDataInterface::PostEditChangeChainProperty(FPrope
 {
 	Super::PostEditChangeChainProperty(PropertyChangedEvent);
 
-	if (PropertyChangedEvent.ChangeType == EPropertyChangeType::ValueSet)
+	if (PropertyChangedEvent.ChangeType == EPropertyChangeType::Unspecified)
 	{
 		static FProperty* NameProperty =
 			FOptimusAnimAttributeDescription::StaticStruct()->FindPropertyByName(
@@ -70,22 +70,11 @@ void UOptimusNode_AnimAttributeDataInterface::PostEditChangeChainProperty(FPrope
 			UpdatePinTypes();
 		}
 
-	}
-	else if (PropertyChangedEvent.ChangeType & (EPropertyChangeType::ArrayAdd | EPropertyChangeType::Duplicate |
-												EPropertyChangeType::ArrayRemove | EPropertyChangeType::ArrayMove))
-	{
 		if (PropertyChangedEvent.PropertyChain.GetTail()->GetValue()->GetFName() == GET_MEMBER_NAME_STRING_CHECKED(FOptimusAnimAttributeArray, InnerArray))
 		{
 			RefreshOutputPins();
 		}
 	}
-	else if(PropertyChangedEvent.ChangeType == EPropertyChangeType::ArrayClear)
-	{
-		if (PropertyChangedEvent.PropertyChain.GetTail()->GetValue()->GetFName() == GET_MEMBER_NAME_STRING_CHECKED(FOptimusAnimAttributeArray, InnerArray))
-		{
-			ClearOutputPins();
-		}
-	}
 }
 
 #endif
