Animation: Fix proper class icon not showing for the skelmesh editor

Also fix the class thumbnail icon not defined for Deformer Graph.

#rnx

diff --git a/Engine/Plugins/Animation/DeformerGraph/Source/OptimusEditor/Private/OptimusEditorStyle.cpp b/Engine/Plugins/Animation/DeformerGraph/Source/OptimusEditor/Private/OptimusEditorStyle.cpp
--- a/Engine/Plugins/Animation/DeformerGraph/Source/OptimusEditor/Private/OptimusEditorStyle.cpp
+++ b/Engine/Plugins/Animation/DeformerGraph/Source/OptimusEditor/Private/OptimusEditorStyle.cpp
@@ -19,6 +19,7 @@ FOptimusEditorStyle::FOptimusEditorStyle() :
 	static const FVector2D IconSize16x12(16.0f, 12.0f);
 	static const FVector2D IconSize16x16(16.0f, 16.0f);
 	static const FVector2D IconSize20x20(20.0f, 20.0f);
+	static const FVector2D IconSize64x64(64.0f, 64.0f);
 
 	static const FSlateColor DefaultForeground(FLinearColor(0.72f, 0.72f, 0.72f, 1.f));
 
@@ -28,6 +29,7 @@ FOptimusEditorStyle::FOptimusEditorStyle() :
 	// Asset icons
 	{
 		Set("ClassIcon.OptimusDeformer", new IMAGE_BRUSH_SVG("Icons/DeformerGraph_16", IconSize16x16, DefaultForeground));
+		Set("ClassThumbnail.OptimusDeformer", new IMAGE_BRUSH_SVG("Icons/DeformerGraph_64", IconSize64x64, DefaultForeground));
 	}
 
 	// Text editor styles
diff --git a/Engine/Source/Editor/EditorStyle/Private/StarshipStyle.cpp b/Engine/Source/Editor/EditorStyle/Private/StarshipStyle.cpp
--- a/Engine/Source/Editor/EditorStyle/Private/StarshipStyle.cpp
+++ b/Engine/Source/Editor/EditorStyle/Private/StarshipStyle.cpp
@@ -6112,6 +6112,7 @@ void FStarshipEditorStyle::FStyle::SetupClassIconsAndThumbnails()
 			{TEXT("SceneCapture2D")},
 			{TEXT("SceneCaptureCube")},
 			{TEXT("SceneComponent")},
+			{TEXT("SkeletalMesh")},
 			{TEXT("SkeletalMeshActor")},
 			{TEXT("Skeleton")},
 			{TEXT("SkyAtmosphere")},
