[Deformer Graph] fixed empty deformer graph causing OOB due to empty vertex factory when setting a preview mesh

#rnx

diff --git a/Engine/Plugins/Animation/DeformerGraph/Source/OptimusCore/Private/OptimusDeformerInstance.cpp b/Engine/Plugins/Animation/DeformerGraph/Source/OptimusCore/Private/OptimusDeformerInstance.cpp
--- a/Engine/Plugins/Animation/DeformerGraph/Source/OptimusCore/Private/OptimusDeformerInstance.cpp
+++ b/Engine/Plugins/Animation/DeformerGraph/Source/OptimusCore/Private/OptimusDeformerInstance.cpp
@@ -347,6 +347,8 @@ void UOptimusDeformerInstance::SetInstanceSettings(UOptimusDeformerInstanceSetti
 
 void UOptimusDeformerInstance::SetupFromDeformer(UOptimusDeformer* InDeformer)
 {
+	const EMeshDeformerOutputBuffer PreviousOutputBuffer = GetOutputBuffers();
+	
 	// If we're doing a recompile, ditch all stored render resources.
 	ReleaseResources();
 
@@ -461,6 +463,13 @@ void UOptimusDeformerInstance::SetupFromDeformer(UOptimusDeformer* InDeformer)
 
 	if (UMeshComponent* Ptr = MeshComponent.Get())
 	{
+		// In case we are writing to different buffers, notify the mesh component such that it can recreate render state and allocate necessary
+		// passthrough vertex factories
+		const EMeshDeformerOutputBuffer CurrentOutputBuffer = GetOutputBuffers();
+		if (CurrentOutputBuffer != PreviousOutputBuffer)
+		{
+			Ptr->MarkRenderStateDirty();
+		}
 		Ptr->MarkRenderDynamicDataDirty();
 	}
 }
