[Deformer Graph] Added support for array type variables

UI work to follow next

diff --git a/Engine/Plugins/Animation/DeformerGraph/Source/OptimusCore/Private/DataInterfaces/OptimusDataInterfaceAnimAttribute.cpp b/Engine/Plugins/Animation/DeformerGraph/Source/OptimusCore/Private/DataInterfaces/OptimusDataInterfaceAnimAttribute.cpp
--- a/Engine/Plugins/Animation/DeformerGraph/Source/OptimusCore/Private/DataInterfaces/OptimusDataInterfaceAnimAttribute.cpp
+++ b/Engine/Plugins/Animation/DeformerGraph/Source/OptimusCore/Private/DataInterfaces/OptimusDataInterfaceAnimAttribute.cpp
@@ -723,7 +723,7 @@ void FOptimusAnimAttributeDataProviderProxy::AllocateResources(FRDGBuilder& Grap
 		for (int32 ArrayIndex = 0; ArrayIndex < AttributeArrayMetadata.Num(); ArrayIndex++)
 		{
 			const FArrayMetadata& ArrayMetadata = AttributeArrayMetadata[ArrayIndex];
-			const TArray<uint8>& ArrayData = AttributeArrayData[ArrayIndex];
+			const TArray<uint8>& ArrayData = AttributeArrayData[ArrayIndex].ArrayOfValues;
 			
 			ArrayBuffers.Add(
 				GraphBuilder.CreateBuffer(
diff --git a/Engine/Plugins/Animation/DeformerGraph/Source/OptimusCore/Private/DataInterfaces/OptimusDataInterfaceAnimAttribute.h b/Engine/Plugins/Animation/DeformerGraph/Source/OptimusCore/Private/DataInterfaces/OptimusDataInterfaceAnimAttribute.h
--- a/Engine/Plugins/Animation/DeformerGraph/Source/OptimusCore/Private/DataInterfaces/OptimusDataInterfaceAnimAttribute.h
+++ b/Engine/Plugins/Animation/DeformerGraph/Source/OptimusCore/Private/DataInterfaces/OptimusDataInterfaceAnimAttribute.h
@@ -170,7 +170,7 @@ struct FOptimusAnimAttributeRuntimeData
 
 	UScriptStruct* AttributeType = nullptr;
 
-	FShaderValueType::FValue CachedDefaultValue;
+	FShaderValueContainer CachedDefaultValue;
 };
 
 /** Compute Framework Data Provider for reading animation attributes on skeletal mesh. */
@@ -224,7 +224,7 @@ public:
 public:
 	TArray<uint8> AttributeBuffer;
 	TArray<FArrayMetadata> AttributeArrayMetadata;
-	TArray<TArray<uint8>> AttributeArrayData;
+	TArray<FArrayShaderValue> AttributeArrayData;
 
 private:
 	TArray<FRDGBuffer*> ArrayBuffers;
diff --git a/Engine/Plugins/Animation/DeformerGraph/Source/OptimusCore/Private/DataInterfaces/OptimusDataInterfaceGraph.cpp b/Engine/Plugins/Animation/DeformerGraph/Source/OptimusCore/Private/DataInterfaces/OptimusDataInterfaceGraph.cpp
--- a/Engine/Plugins/Animation/DeformerGraph/Source/OptimusCore/Private/DataInterfaces/OptimusDataInterfaceGraph.cpp
+++ b/Engine/Plugins/Animation/DeformerGraph/Source/OptimusCore/Private/DataInterfaces/OptimusDataInterfaceGraph.cpp
@@ -9,6 +9,8 @@
 #include "OptimusDeformerInstance.h"
 #include "OptimusHelpers.h"
 #include "OptimusVariableDescription.h"
+#include "RenderGraphBuilder.h"
+#include "RenderGraphResources.h"
 #include "ShaderParameterMetadataBuilder.h"
 
 #include UE_INLINE_GENERATED_CPP_BY_NAME(OptimusDataInterfaceGraph)
@@ -92,10 +94,41 @@ void UOptimusGraphDataInterface::GetHLSL(FString& OutHLSL, FString const& InData
 UComputeDataProvider* UOptimusGraphDataInterface::CreateDataProvider(TObjectPtr<UObject> InBinding, uint64 InInputMask, uint64 InOutputMask) const
 {
 	UOptimusGraphDataProvider* Provider = NewObject<UOptimusGraphDataProvider>();
-	Provider->MeshComponent = Cast<UMeshComponent>(InBinding);
-	Provider->Variables = Variables;
+	Provider->Init(Cast<UMeshComponent>(InBinding), Variables, ParameterBufferSize);
 
-	for (FOptimusGraphVariableDescription& Variable : Provider->Variables)
+
+	
+	
+	return Provider;
+}
+
+void UOptimusGraphDataInterface::PostLoad()
+{
+	Super::PostLoad();
+
+PRAGMA_DISABLE_DEPRECATION_WARNINGS
+	for (FOptimusGraphVariableDescription& VariableDescription : Variables)
+	{
+		if (!VariableDescription.Value_DEPRECATED.IsEmpty())
+		{
+			FOptimusDataTypeHandle DataType = FOptimusDataTypeRegistry::Get().FindType(VariableDescription.ValueType);
+			VariableDescription.ShaderValue = DataType->MakeShaderValue();
+			check(VariableDescription.ShaderValue.ArrayList.Num() == 0);
+			VariableDescription.ShaderValue.ShaderValue = VariableDescription.Value_DEPRECATED;
+			VariableDescription.Value_DEPRECATED.Reset();
+		}
+	}
+PRAGMA_ENABLE_DEPRECATION_WARNINGS
+}
+
+
+void UOptimusGraphDataProvider::Init(UMeshComponent* InMeshComponent, const TArray<FOptimusGraphVariableDescription>& InVariables, int32 InParameterBufferSize)
+{
+	MeshComponent = InMeshComponent;
+	Variables = InVariables;
+	ParameterBufferSize = InParameterBufferSize;
+
+	for (FOptimusGraphVariableDescription& Variable : Variables)
 	{
 		// When source object was introduced, we also appended a unique index to the value name provided by each value provider
 		// so instead of using the name directly, we need to do this extra step
@@ -105,20 +138,42 @@ UComputeDataProvider* UOptimusGraphDataInterface::CreateDataProvider(TObjectPtr<
 		}
 	}
 	
-	Provider->ParameterBufferSize = ParameterBufferSize;
-	return Provider;
-}
+	int32 TotalNumArrays = 0;
+	for (FOptimusGraphVariableDescription& Variable : Variables)
+	{
+		const TArray<FOptimusDataTypeRegistry::FArrayMetadata>& TypeArrayMetadata = FOptimusDataTypeRegistry::Get().FindArrayMetadata(FOptimusDataTypeRegistry::Get().FindType(Variable.ValueType)->TypeName);
 
+		Variable.CachedArrayIndexStart = TotalNumArrays;
+		TotalNumArrays += TypeArrayMetadata.Num();
+	}
+	
+	ParameterArrayMetadata.AddDefaulted(TotalNumArrays);
+	
+	for (FOptimusGraphVariableDescription const& Variable : Variables)
+	{
+		const TArray<FOptimusDataTypeRegistry::FArrayMetadata>& TypeArrayMetadata = FOptimusDataTypeRegistry::Get().FindArrayMetadata(FOptimusDataTypeRegistry::Get().FindType(Variable.ValueType)->TypeName);
+
+		for (int32 ArrayIndex = 0; ArrayIndex < TypeArrayMetadata.Num(); ArrayIndex++)
+		{
+			const int32 TopLevelArrayIndex = Variable.CachedArrayIndexStart + ArrayIndex;
+			if (ensure(ParameterArrayMetadata.IsValidIndex(TopLevelArrayIndex)))
+			{
+				ParameterArrayMetadata[TopLevelArrayIndex].Offset = Variable.Offset + TypeArrayMetadata[ArrayIndex].ShaderValueOffset;
+				ParameterArrayMetadata[TopLevelArrayIndex].ElementSize = TypeArrayMetadata[ArrayIndex].ElementShaderValueSize;
+			}
+		}
+	}
+}
 
-void UOptimusGraphDataProvider::SetConstant(TSoftObjectPtr<UObject> InSourceObject, TArray<uint8> const& InValue)
+void UOptimusGraphDataProvider::SetConstant(TSoftObjectPtr<UObject> InSourceObject, FShaderValueContainer const& InValue)
 {
 	for (int32 VariableIndex = 0; VariableIndex < Variables.Num(); ++VariableIndex)
 	{
 		if (Variables[VariableIndex].SourceObject == InSourceObject)
 		{
-			if (ensure(Variables[VariableIndex].Value.Num() == InValue.Num()))
+			if (ensure(FShaderValueContainer::IsSameType(Variables[VariableIndex].ShaderValue, InValue)))
 			{
-				Variables[VariableIndex].Value = InValue;
+				Variables[VariableIndex].ShaderValue = InValue;
 				break;
 			}
 		}
@@ -127,7 +182,7 @@ void UOptimusGraphDataProvider::SetConstant(TSoftObjectPtr<UObject> InSourceObje
 
 FComputeDataProviderRenderProxy* UOptimusGraphDataProvider::GetRenderProxy()
 {
-	return new FOptimusGraphDataProviderProxy(DeformerInstance, Variables, ParameterBufferSize);
+	return new FOptimusGraphDataProviderProxy(DeformerInstance, Variables, ParameterBufferSize, ParameterArrayMetadata);
 }
 
 void UOptimusGraphDataProvider::SetDeformerInstance(UOptimusDeformerInstance* InInstance)
@@ -140,26 +195,47 @@ UOptimusDeformerInstance* UOptimusGraphDataProvider::GetDeformerInstance() const
 	return DeformerInstance;
 }
 
-FOptimusGraphDataProviderProxy::FOptimusGraphDataProviderProxy(UOptimusDeformerInstance const* DeformerInstance, TArray<FOptimusGraphVariableDescription> const& Variables, int32 ParameterBufferSize)
+FOptimusGraphDataProviderProxy::FOptimusGraphDataProviderProxy(
+	UOptimusDeformerInstance const* DeformerInstance,
+	TArray<FOptimusGraphVariableDescription> const& Variables,
+	int32 ParameterBufferSize,
+	TArray<UOptimusGraphDataProvider::FArrayMetadata> const& InParameterArrayMetadata)
 {
 	// Get all variables from deformer instance and fill buffer.
 	ParameterData.AddZeroed(ParameterBufferSize);
 
+	ParameterArrayMetadata = InParameterArrayMetadata;
+	ParameterArrayData.AddDefaulted(ParameterArrayMetadata.Num());
+	
 	if (DeformerInstance == nullptr)
 	{
 		return;
 	}
 
+	auto CopyVariableToBuffer = [this](int32 InOffset, int32 InArrayIndexStart, const FShaderValueContainer& InShaderValue)
+	{
+    	if (ensure(ParameterData.Num() >= InOffset + InShaderValue.ShaderValue.Num()))
+    	{
+    		FMemory::Memcpy(&ParameterData[InOffset], InShaderValue.ShaderValue.GetData(), InShaderValue.ShaderValue.Num());
+    	
+    		for (int32 ArrayIndex = 0; ArrayIndex < InShaderValue.ArrayList.Num(); ArrayIndex++)
+    		{
+    			const int32 ToplevelArrayIndex = InArrayIndexStart+ ArrayIndex;
+    			if (ensure(ParameterArrayData.IsValidIndex(ToplevelArrayIndex)))
+    			{
+    				ParameterArrayData[ToplevelArrayIndex] = InShaderValue.ArrayList[ArrayIndex];
+    			}
+    		}
+    	}	
+	};
+
 	TArray<UOptimusVariableDescription*> const& VariableValues = DeformerInstance->GetVariables();
 	for (FOptimusGraphVariableDescription const& Variable : Variables)
 	{
-		if (Variable.Value.Num())
+		if (Variable.ShaderValue.IsValid())
 		{
 			// Use the constant value.
-			if (ensure(ParameterData.Num() >= Variable.Offset + Variable.Value.Num()))
-			{
-				FMemory::Memcpy(&ParameterData[Variable.Offset], Variable.Value.GetData(), Variable.Value.Num());
-			}
+			CopyVariableToBuffer(Variable.Offset, Variable.CachedArrayIndexStart, Variable.ShaderValue);
 		}
 		else
 		{
@@ -201,10 +277,8 @@ FOptimusGraphDataProviderProxy::FOptimusGraphDataProviderProxy(UOptimusDeformerI
 
 					if (bNameMatch)
 					{
-						if (ensure(ParameterData.Num() >= Variable.Offset + VariableValue->ValueData.Num()))
-						{
-							FMemory::Memcpy(&ParameterData[Variable.Offset], VariableValue->ValueData.GetData(), VariableValue->ValueData.Num());
-						}
+						CopyVariableToBuffer(Variable.Offset, Variable.CachedArrayIndexStart, VariableValue->CachedShaderValue);
+						
 						break;
 					}
 				}
@@ -229,11 +303,38 @@ bool FOptimusGraphDataProviderProxy::IsValid(FValidationData const& InValidation
 	return true;
 }
 
+void FOptimusGraphDataProviderProxy::AllocateResources(FRDGBuilder& GraphBuilder)
+{
+	ParameterArrayBuffers.Reset();
+	ParameterArrayBufferSRVs.Reset();
+
+
+	for (int32 ArrayIndex = 0; ArrayIndex < ParameterArrayMetadata.Num(); ArrayIndex++)
+	{
+		const UOptimusGraphDataProvider::FArrayMetadata& ArrayMetadata = ParameterArrayMetadata[ArrayIndex];
+		const TArray<uint8>& ArrayData = ParameterArrayData[ArrayIndex].ArrayOfValues;
+		
+		ParameterArrayBuffers.Add(
+			GraphBuilder.CreateBuffer(
+				FRDGBufferDesc::CreateStructuredDesc(
+					ArrayMetadata.ElementSize,
+					FMath::Max(ArrayData.Num() / ArrayMetadata.ElementSize, 1)), TEXT("Optimus.GraphDataInterfaceInnerBuffer")));
+		ParameterArrayBufferSRVs.Add(GraphBuilder.CreateSRV(ParameterArrayBuffers.Last()));
+		GraphBuilder.QueueBufferUpload(ParameterArrayBuffers.Last(), ArrayData.GetData(), ArrayData.Num(), ERDGInitialDataFlags::None);	
+	}
+}
+
 void FOptimusGraphDataProviderProxy::GatherDispatchData(FDispatchData const& InDispatchData)
 {
 	for (int32 InvocationIndex = 0; InvocationIndex < InDispatchData.NumInvocations; ++InvocationIndex)
 	{
-		void* ParameterBuffer = (void*)(InDispatchData.ParameterBuffer + InDispatchData.ParameterBufferOffset + InDispatchData.ParameterBufferStride * InvocationIndex);
+		uint8* ParameterBuffer = (uint8*)(InDispatchData.ParameterBuffer + InDispatchData.ParameterBufferOffset + InDispatchData.ParameterBufferStride * InvocationIndex);
 		FMemory::Memcpy(ParameterBuffer, ParameterData.GetData(), ParameterData.Num());
+
+		for (int32 ArrayIndex = 0; ArrayIndex < ParameterArrayMetadata.Num(); ArrayIndex++)
+		{
+			const UOptimusGraphDataProvider::FArrayMetadata& ArrayMetadata = ParameterArrayMetadata[ArrayIndex];
+			*((FRDGBufferSRV**)(ParameterBuffer + ArrayMetadata.Offset)) = ParameterArrayBufferSRVs[ArrayIndex];
+		}
 	}
 }
diff --git a/Engine/Plugins/Animation/DeformerGraph/Source/OptimusCore/Private/DataInterfaces/OptimusDataInterfaceGraph.h b/Engine/Plugins/Animation/DeformerGraph/Source/OptimusCore/Private/DataInterfaces/OptimusDataInterfaceGraph.h
--- a/Engine/Plugins/Animation/DeformerGraph/Source/OptimusCore/Private/DataInterfaces/OptimusDataInterfaceGraph.h
+++ b/Engine/Plugins/Animation/DeformerGraph/Source/OptimusCore/Private/DataInterfaces/OptimusDataInterfaceGraph.h
@@ -3,6 +3,7 @@
 #pragma once
 
 #include "IOptimusDeformerInstanceAccessor.h"
+#include "OptimusDataTypeRegistry.h"
 #include "ComputeFramework/ComputeDataInterface.h"
 #include "ComputeFramework/ComputeDataProvider.h"
 #include "ComputeFramework/ShaderParamTypeDefinition.h"
@@ -10,6 +11,8 @@
 
 class UOptimusDeformerInstance;
 class UMeshComponent;
+class FRDGBuffer;
+class FRDGBufferSRV;
 
 /** */
 USTRUCT()
@@ -24,16 +27,24 @@ struct FOptimusGraphVariableDescription
 	FShaderValueTypeHandle ValueType;
 
 	UPROPERTY()
-	TArray<uint8> Value;
-
+    FShaderValueContainer ShaderValue;
+	
 	UPROPERTY()
 	int32 Offset = 0;
 
 	UPROPERTY()
 	TSoftObjectPtr<UObject> SourceObject;
-
-	// Cached SourceValueName extracted from Name, Initialized by the DataProvider for when SourceObject is not null
+	
+	// Cache below are set by the data provider, computed from serialized data
+	
+	// SourceValueName extracted from Name, Initialized by the DataProvider for when SourceObject is not null
 	FString CachedSourceValueName;
+
+	int32 CachedArrayIndexStart = 0;
+	
+	// Deprecated 
+	UPROPERTY(meta = (DeprecatedProperty, DeprecationMessage = "Use ShaderValue instead"))
+	TArray<uint8> Value_DEPRECATED;
 };
 
 /** Compute Framework Data Interface used for marshaling compute graph parameters and variables. */
@@ -54,6 +65,7 @@ public:
 	UComputeDataProvider* CreateDataProvider(TObjectPtr<UObject> InBinding, uint64 InInputMask, uint64 InOutputMask) const override;
 	//~ End UComputeDataInterface Interface
 
+	void PostLoad() override;
 private:
 	UPROPERTY()
 	TArray<FOptimusGraphVariableDescription> Variables;
@@ -79,7 +91,19 @@ public:
 
 	int32 ParameterBufferSize = 0;
 
-	void SetConstant(TSoftObjectPtr<UObject> InSourceObject, TArray<uint8> const& InValue);
+
+	struct FArrayMetadata
+	{
+		int32 Offset;
+		int32 ElementSize;
+	};
+	
+	TArray<FArrayMetadata> ParameterArrayMetadata;
+	
+	
+	void Init(UMeshComponent* InMeshComponent, const TArray<FOptimusGraphVariableDescription>& InVariables, int32 InParameterBufferSize);
+	
+	void SetConstant(TSoftObjectPtr<UObject> InSourceObject, FShaderValueContainer const& InValue);
 	
 	//~ Begin UComputeDataProvider Interface
 	FComputeDataProviderRenderProxy* GetRenderProxy() override;
@@ -98,13 +122,24 @@ private:
 class FOptimusGraphDataProviderProxy : public FComputeDataProviderRenderProxy
 {
 public:
-	FOptimusGraphDataProviderProxy(UOptimusDeformerInstance const* DeformerInstance, TArray<FOptimusGraphVariableDescription> const& Variables, int32 ParameterBufferSize);
+	FOptimusGraphDataProviderProxy(
+		UOptimusDeformerInstance const* DeformerInstance,
+		TArray<FOptimusGraphVariableDescription> const& Variables, 
+		int32 ParameterBufferSize,
+		TArray<UOptimusGraphDataProvider::FArrayMetadata> const& InParameterArrayMetadata);
 
 	//~ Begin FComputeDataProviderRenderProxy Interface
 	bool IsValid(FValidationData const& InValidationData) const override;
+	void AllocateResources(FRDGBuilder& GraphBuilder) override;
 	void GatherDispatchData(FDispatchData const& InDispatchData) override;
 	//~ End FComputeDataProviderRenderProxy Interface
 
 private:
 	TArray<uint8> ParameterData;
+
+	TArray<UOptimusGraphDataProvider::FArrayMetadata> ParameterArrayMetadata;
+	TArray<FArrayShaderValue> ParameterArrayData;
+
+	TArray<FRDGBuffer*> ParameterArrayBuffers;
+	TArray<FRDGBufferSRV*> ParameterArrayBufferSRVs;
 };
diff --git a/Engine/Plugins/Animation/DeformerGraph/Source/OptimusCore/Private/Nodes/OptimusNode_ConstantValue.cpp b/Engine/Plugins/Animation/DeformerGraph/Source/OptimusCore/Private/Nodes/OptimusNode_ConstantValue.cpp
--- a/Engine/Plugins/Animation/DeformerGraph/Source/OptimusCore/Private/Nodes/OptimusNode_ConstantValue.cpp
+++ b/Engine/Plugins/Animation/DeformerGraph/Source/OptimusCore/Private/Nodes/OptimusNode_ConstantValue.cpp
@@ -118,7 +118,7 @@ FOptimusDataTypeRef UOptimusNode_ConstantValue::GetValueType() const
 }
 
 
-FShaderValueType::FValue UOptimusNode_ConstantValue::GetShaderValue() const
+FShaderValueContainer UOptimusNode_ConstantValue::GetShaderValue() const
 {
 	const UOptimusNodePin *ValuePin = FindPinFromPath({TEXT("Value")});
 	if (ensure(ValuePin))
@@ -128,7 +128,7 @@ FShaderValueType::FValue UOptimusNode_ConstantValue::GetShaderValue() const
 		if (ensure(ValueProperty) && ensure(DataType.IsValid()))
 		{
 			TArrayView<const uint8> ValueData(ValueProperty->ContainerPtrToValuePtr<uint8>(this), ValueProperty->GetSize());
-			FShaderValueType::FValue ValueResult = DataType->MakeShaderValue();
+			FShaderValueContainer ValueResult = DataType->MakeShaderValue();
 
 			if (DataType->ConvertPropertyValueToShader(ValueData, ValueResult))
 			{
diff --git a/Engine/Plugins/Animation/DeformerGraph/Source/OptimusCore/Private/Nodes/OptimusNode_ConstantValue.h b/Engine/Plugins/Animation/DeformerGraph/Source/OptimusCore/Private/Nodes/OptimusNode_ConstantValue.h
--- a/Engine/Plugins/Animation/DeformerGraph/Source/OptimusCore/Private/Nodes/OptimusNode_ConstantValue.h
+++ b/Engine/Plugins/Animation/DeformerGraph/Source/OptimusCore/Private/Nodes/OptimusNode_ConstantValue.h
@@ -67,7 +67,7 @@ public:
 	// IOptimusValueProvider overrides 
 	FString GetValueName() const override;
 	FOptimusDataTypeRef GetValueType() const override;
-	FShaderValueType::FValue GetShaderValue() const override;
+	FShaderValueContainer GetShaderValue() const override;
 
 	// IOptimusGeneratedClassDefiner implementation
 	FTopLevelAssetPath GetAssetPathForClassDefiner() const override;
diff --git a/Engine/Plugins/Animation/DeformerGraph/Source/OptimusCore/Private/Nodes/OptimusNode_GetVariable.cpp b/Engine/Plugins/Animation/DeformerGraph/Source/OptimusCore/Private/Nodes/OptimusNode_GetVariable.cpp
--- a/Engine/Plugins/Animation/DeformerGraph/Source/OptimusCore/Private/Nodes/OptimusNode_GetVariable.cpp
+++ b/Engine/Plugins/Animation/DeformerGraph/Source/OptimusCore/Private/Nodes/OptimusNode_GetVariable.cpp
@@ -217,18 +217,10 @@ FOptimusDataTypeRef UOptimusNode_GetVariable::GetValueType() const
 }
 
 
-FShaderValueType::FValue UOptimusNode_GetVariable::GetShaderValue() const
+FShaderValueContainer UOptimusNode_GetVariable::GetShaderValue() const
 {
-	if (const UOptimusVariableDescription* Var = VariableDesc.Get();
-		Var && ensure(Var->DataType.IsValid()) && ensure(GetPins().Num() == 1))
-	{
-		FShaderValueType::FValue ValueResult = Var->DataType->MakeShaderValue();
-
-		if (Var->DataType->ConvertPropertyValueToShader(Var->ValueData, ValueResult))
-		{
-			return ValueResult;
-		}
-	}
+	// Unused, Graph data interface is responsible for collecting value from UOptimusVariableDescription
+	checkNoEntry();
 
 	return {};
 }
diff --git a/Engine/Plugins/Animation/DeformerGraph/Source/OptimusCore/Private/Nodes/OptimusNode_GetVariable.h b/Engine/Plugins/Animation/DeformerGraph/Source/OptimusCore/Private/Nodes/OptimusNode_GetVariable.h
--- a/Engine/Plugins/Animation/DeformerGraph/Source/OptimusCore/Private/Nodes/OptimusNode_GetVariable.h
+++ b/Engine/Plugins/Animation/DeformerGraph/Source/OptimusCore/Private/Nodes/OptimusNode_GetVariable.h
@@ -56,7 +56,7 @@ public:
 	// IOptimusValueProvider overrides 
 	FString GetValueName() const override;
 	FOptimusDataTypeRef GetValueType() const override;
-	FShaderValueType::FValue GetShaderValue() const override;
+	FShaderValueContainer GetShaderValue() const override;
 
 	//IOptimusPinMutabilityDefiner overrides 
 	EOptimusPinMutability GetOutputPinMutability(const UOptimusNodePin* InPin) const override { return EOptimusPinMutability::Mutable; };
diff --git a/Engine/Plugins/Animation/DeformerGraph/Source/OptimusCore/Private/OptimusDataType.cpp b/Engine/Plugins/Animation/DeformerGraph/Source/OptimusCore/Private/OptimusDataType.cpp
--- a/Engine/Plugins/Animation/DeformerGraph/Source/OptimusCore/Private/OptimusDataType.cpp
+++ b/Engine/Plugins/Animation/DeformerGraph/Source/OptimusCore/Private/OptimusDataType.cpp
@@ -93,7 +93,7 @@ FProperty* FOptimusDataType::CreateProperty(
 
 bool FOptimusDataType::ConvertPropertyValueToShader(
 	TArrayView<const uint8> InValue,
-	FShaderValueType::FValueView OutConvertedValue
+	FShaderValueContainer& OutConvertedValue
 	) const
 {
 	const FOptimusDataTypeRegistry::PropertyValueConvertFuncT PropertyConversionFunc =
@@ -108,9 +108,9 @@ bool FOptimusDataType::ConvertPropertyValueToShader(
 	}
 }
 
-FShaderValueType::FValue FOptimusDataType::MakeShaderValue() const
+FShaderValueContainer FOptimusDataType::MakeShaderValue() const
 {
-	return {ShaderValueSize, GetNumArrays()};
+	return FShaderValueContainer(ShaderValueSize, GetNumArrays());
 }
 
 bool FOptimusDataType::CanCreateProperty() const
diff --git a/Engine/Plugins/Animation/DeformerGraph/Source/OptimusCore/Private/OptimusDataTypeRegistry.cpp b/Engine/Plugins/Animation/DeformerGraph/Source/OptimusCore/Private/OptimusDataTypeRegistry.cpp
--- a/Engine/Plugins/Animation/DeformerGraph/Source/OptimusCore/Private/OptimusDataTypeRegistry.cpp
+++ b/Engine/Plugins/Animation/DeformerGraph/Source/OptimusCore/Private/OptimusDataTypeRegistry.cpp
@@ -57,7 +57,7 @@ static bool IsStructHashable(const UScriptStruct* InStructType)
 
 template<typename SourceT, typename DestT>
 static bool ConvertPropertyValuePOD(
-		TArrayView<const uint8> InRawValue, FShaderValueType::FValueView OutShaderValue
+		TArrayView<const uint8> InRawValue, FShaderValueContainerView OutShaderValue
 	)
 {
 	if (ensure(InRawValue.Num() == sizeof(SourceT)) &&
@@ -69,6 +69,18 @@ static bool ConvertPropertyValuePOD(
 	return false;
 }
 
+// Special logic for things like StructuredBuffer<float3> which should be packed as buffer of float4 in Vulkan
+static FOptimusDataTypeHandle GetArrayElementDataTypeForStructuredBuffer(FOptimusDataTypeHandle InDataType)
+{
+	FOptimusDataTypeHandle Result = InDataType;
+	if (InDataType->ShaderValueType->Type != EShaderFundamentalType::Struct && InDataType->ShaderValueType->VectorElemCount == 3)
+	{
+		Result = FOptimusDataTypeRegistry::Get().FindType(FShaderValueType::Get(InDataType->ShaderValueType->Type, 4));
+	}
+
+	return Result;
+}
+
 EOptimusDataTypeUsageFlags FOptimusDataTypeRegistry::GetStructTypeUsageFlag(UScriptStruct* InStruct)
 {
 	struct FTypeValidator 
@@ -231,7 +243,7 @@ void FOptimusDataTypeRegistry::RegisterBuiltinTypes()
 	    *FBoolProperty::StaticClass(),
 	    FText::FromString(TEXT("Bool")),
 	    FShaderValueType::Get(EShaderFundamentalType::Bool),
-		[](UStruct *InScope, FName InName) {
+		[](FFieldVariant InScope, FName InName) {
 		    FBoolProperty* Property = new FBoolProperty(InScope, InName, RF_Public);
 		    Property->SetBoolSize(sizeof(bool), true);
 			return Property;
@@ -245,7 +257,7 @@ void FOptimusDataTypeRegistry::RegisterBuiltinTypes()
 	    *FIntProperty::StaticClass(),
 	    FText::FromString(TEXT("Int")),
 	    FShaderValueType::Get(EShaderFundamentalType::Int),
-	    [](UStruct* InScope, FName InName) {
+	    [](FFieldVariant InScope, FName InName) {
 		    FIntProperty* Property = new FIntProperty(InScope, InName, RF_Public);
 			Property->SetPropertyFlags(CPF_HasGetValueTypeHash);
 		    return Property;
@@ -285,7 +297,7 @@ void FOptimusDataTypeRegistry::RegisterBuiltinTypes()
 		*FUInt32Property::StaticClass(),
 	    FText::FromString(TEXT("UInt")),
 		FShaderValueType::Get(EShaderFundamentalType::Uint),
-		[](UStruct* InScope, FName InName) {
+		[](FFieldVariant InScope, FName InName) {
 			FUInt32Property* Property = new FUInt32Property(InScope, InName, RF_Public);
 			Property->SetPropertyFlags(CPF_HasGetValueTypeHash);
 			return Property;
@@ -327,7 +339,7 @@ void FOptimusDataTypeRegistry::RegisterBuiltinTypes()
 	    *FFloatProperty::StaticClass(),
 	    FText::FromString(TEXT("Float")),
 	    FShaderValueType::Get(EShaderFundamentalType::Float),
-	    [](UStruct* InScope, FName InName) {
+	    [](FFieldVariant InScope, FName InName) {
 		    FFloatProperty* Property = new FFloatProperty(InScope, InName, RF_Public);
 		    Property->SetPropertyFlags(CPF_HasGetValueTypeHash);
 #if WITH_EDITOR
@@ -347,7 +359,7 @@ void FOptimusDataTypeRegistry::RegisterBuiltinTypes()
 	    *FDoubleProperty::StaticClass(),
 	    FText::FromString(TEXT("Float")),
 	    FShaderValueType::Get(EShaderFundamentalType::Float),
-	    [](UStruct* InScope, FName InName) {
+	    [](FFieldVariant InScope, FName InName) {
 		    FDoubleProperty* Property = new FDoubleProperty(InScope, InName, RF_Public);
 		    Property->SetPropertyFlags(CPF_HasGetValueTypeHash);
 #if WITH_EDITOR
@@ -419,7 +431,7 @@ void FOptimusDataTypeRegistry::RegisterBuiltinTypes()
 	    FShaderValueType::Get(EShaderFundamentalType::Float, 4, 4),
 	    [](
 	    	TArrayView<const uint8> InRawValue,
-			FShaderValueType::FValueView OutShaderValue) -> bool
+			FShaderValueContainerView OutShaderValue) -> bool
 	    {
 	    	if (ensure(InRawValue.Num() == TBaseStructure<FTransform>::Get()->GetCppStructOps()->GetSize()) &&
 				ensure(OutShaderValue.ShaderValue.Num() == FShaderValueType::Get(EShaderFundamentalType::Float, 4, 4)->GetResourceElementSize()))
@@ -458,6 +470,16 @@ void FOptimusDataTypeRegistry::RegisterBuiltinTypes()
 
 	
 	Registry.TypeWithAtomicSupport = {FIntProperty::StaticClass()->GetFName()};
+	
+	// Scan available built-in types to see if we can create array type for them
+	// Currently only support variable usage + no nested array
+
+	TMap<FName, FTypeInfo> AlreadyRegisteredTypes = Registry.RegisteredTypes;
+	
+	for (const TPair<FName, FTypeInfo>& Type : AlreadyRegisteredTypes)
+	{
+		Registry.RegisterArrayTypeIfApplicable(Type.Value.DataType);
+	}
 }
 
 
@@ -528,7 +550,7 @@ bool FOptimusDataTypeRegistry::RegisterStructType(UScriptStruct* InStructType)
 	const bool bIsHashable = IsStructHashable(InStructType);
 
 	PropertyCreateFuncT PropertyCreateFunc;
-	PropertyCreateFunc = [bIsHashable, InStructType](UStruct* InScope, FName InName) -> FProperty *
+	PropertyCreateFunc = [bIsHashable, InStructType](FFieldVariant InScope, FName InName) -> FProperty *
 	{
 		auto Property = new FStructProperty(InScope, InName, RF_Public);
 		Property->Struct = InStructType;
@@ -583,10 +605,7 @@ bool FOptimusDataTypeRegistry::RegisterStructType(UScriptStruct* InStructType)
 			// Special logic for things like StructuredBuffer<float3> which should be packed as buffer of float4 in Vulkan
 			if (bIsArrayMember)
 			{
-				if (DataType->ShaderValueType->Type != EShaderFundamentalType::Struct && DataType->ShaderValueType->VectorElemCount == 3)
-				{
-					DataType = FindType(FShaderValueType::Get(DataType->ShaderValueType->Type, 4));
-				}
+				DataType = GetArrayElementDataTypeForStructuredBuffer(DataType);
 			}
 			
 			if (DataType->GetNumArrays() > 0)
@@ -716,7 +735,7 @@ bool FOptimusDataTypeRegistry::RegisterStructType(UScriptStruct* InStructType)
 		}
 		
 		PropertyValueConvertFunc = [ConversionEntries, ExpectedPropertySize, ExpectedShaderValueSize](
-				TArrayView<const uint8> InRawValue, FShaderValueType::FValueView OutShaderValue
+				TArrayView<const uint8> InRawValue, FShaderValueContainerView OutShaderValue
 			) -> bool
 		{
 			if (ensure(InRawValue.Num() == ExpectedPropertySize) &&
@@ -736,7 +755,7 @@ bool FOptimusDataTypeRegistry::RegisterStructType(UScriptStruct* InStructType)
 
 						// Convert each element, store them in a separate buffer that is to be uploaded
 						FScriptArrayHelper ArrayHelper(ArrayProperty, PropertyRawValue);
-						TArray<uint8>& Buffer = OutShaderValue.ArrayList[ConversionInfo.ArrayIndex];
+						TArray<uint8>& Buffer = OutShaderValue.ArrayList[ConversionInfo.ArrayIndex].ArrayOfValues;
 						Buffer.AddZeroed(ConversionInfo.ShaderValueSize * ArrayHelper.Num());
 						
 						for (int32 Index = 0; Index < ArrayHelper.Num(); Index++)
@@ -745,7 +764,7 @@ bool FOptimusDataTypeRegistry::RegisterStructType(UScriptStruct* InStructType)
 							uint8* ShaderValuePtr = Buffer.GetData() + ConversionInfo.ShaderValueSize * Index;
 
 							// Nested buffer is not possible
-							TArray<TArray<uint8>> DummyBufferList;
+							TArray<FArrayShaderValue> DummyBufferList;
 
 							if (!ConversionInfo.PropertyInfo.ConvertFunc(
 									{ElementPtr, ArrayProperty->Inner->GetSize()},
@@ -824,6 +843,103 @@ bool FOptimusDataTypeRegistry::RegisterStructType(UScriptStruct* InStructType)
 	return false;
 }
 
+bool FOptimusDataTypeRegistry::RegisterArrayTypeIfApplicable(FOptimusDataTypeHandle InElementDataType)
+{
+	// For now only allow array type for variables
+	if (!EnumHasAnyFlags(InElementDataType->UsageFlags, EOptimusDataTypeUsageFlags::Variable))
+	{
+		return false;
+	}
+
+	// Nested array is not supported at the moment for array variables
+	if (InElementDataType->GetNumArrays() != 0)
+	{
+		return false;
+	}
+	
+	PropertyCreateFuncT ElementPropertyCreateFunc = FindPropertyCreateFunc(InElementDataType->TypeName);
+	PropertyCreateFuncT	ArrayPropertyCreateFunc = [ElementPropertyCreateFunc](FFieldVariant InScope, FName InName) -> FProperty*
+	{
+		FArrayProperty* ArrayProperty = new FArrayProperty(InScope, InName, RF_NoFlags);
+		ArrayProperty->Inner = ElementPropertyCreateFunc(ArrayProperty, TEXT("Inner"));
+		
+		return ArrayProperty;
+	};
+
+	// Making sure we are copying property value of a type into shader value of a equal or larger type, see comment for GetArrayElementDataTypeForStructuredBuffer
+	FOptimusDataTypeHandle InnerDataTypeForStructuredBuffer = GetArrayElementDataTypeForStructuredBuffer(InElementDataType);
+	check(InnerDataTypeForStructuredBuffer->ShaderValueSize >= InElementDataType->ShaderValueSize);
+	PropertyValueConvertFuncT ElementPropertyValueConvertFunc = FindPropertyValueConvertFunc(InnerDataTypeForStructuredBuffer->TypeName);
+	check(ElementPropertyValueConvertFunc);
+	int32 ElementShaderValueSize = InnerDataTypeForStructuredBuffer->ShaderValueSize;
+
+	// Nested array is not supported at the moment for array variables
+	check(InnerDataTypeForStructuredBuffer->GetNumArrays() == 0);
+	
+	PropertyValueConvertFuncT ArrayPropertyValueConvertFunc = [
+		ArrayPropertyCreateFunc,
+		ElementShaderValueSize,
+		ElementPropertyValueConvertFunc
+		] (
+		TArrayView<const uint8> InRawValue,
+		FShaderValueContainerView OutShaderValue
+		) -> bool
+	{
+		const TUniquePtr<FArrayProperty> LocalArrayProperty(CastField<FArrayProperty>(ArrayPropertyCreateFunc(nullptr, NAME_None)));
+
+		const FArrayProperty* ArrayProperty = LocalArrayProperty.Get();
+		// Convert each element, store them in a separate buffer that is to be uploaded
+		FScriptArrayHelper ArrayHelper(ArrayProperty, InRawValue.GetData());
+		TArray<uint8>& Buffer = OutShaderValue.ArrayList[0].ArrayOfValues;
+		Buffer.AddZeroed(ElementShaderValueSize * ArrayHelper.Num());
+				
+		for (int32 Index = 0; Index < ArrayHelper.Num(); Index++)
+		{
+			uint8* ElementPtr = ArrayHelper.GetRawPtr(Index);
+			uint8* ShaderValuePtr = Buffer.GetData() + ElementShaderValueSize * Index;
+
+			// Nested buffer is not possible
+			TArray<FArrayShaderValue> DummyBufferList;
+
+			if (!ElementPropertyValueConvertFunc(
+					{ElementPtr, ArrayProperty->Inner->GetSize()},
+					{
+						{ShaderValuePtr, ElementShaderValueSize},
+						DummyBufferList
+					}))
+			{
+				return false;
+			}
+		}	
+		return true;
+	};
+
+	TSharedRef<FOptimusDataType> ArrayDataType = MakeShared<FOptimusDataType>();
+	
+	const FTypeInfo Info
+	{
+		ArrayDataType,
+		ArrayPropertyCreateFunc,
+		ArrayPropertyValueConvertFunc,
+		{FArrayMetadata(ElementShaderValueSize, 0 )}
+	};
+
+	{
+		*ArrayDataType = *InElementDataType;
+		ArrayDataType->TypeName = *(TEXT("TArray<") + InElementDataType->TypeName.ToString() + TEXT(">"));
+		ArrayDataType->DisplayName = FText::FromString(InElementDataType->DisplayName.ToString() + TEXT(" Array"));
+		ArrayDataType->ShaderValueType = FShaderValueType::MakeDynamicArrayType(InnerDataTypeForStructuredBuffer->ShaderValueType);
+		
+		// Unused field
+		ArrayDataType->ShaderValueSize = 0;	
+	}
+	
+	RegisteredTypes.Add(ArrayDataType->TypeName, Info);
+	RegistrationOrder.Add(ArrayDataType->TypeName);
+
+	return true;
+}
+
 void FOptimusDataTypeRegistry::RefreshStructType(UUserDefinedStruct* InStructType)
 {
 	FName TypeName = Optimus::GetTypeName(InStructType);
@@ -916,7 +1032,7 @@ bool FOptimusDataTypeRegistry::RegisterType(
 		{
 			const bool bIsHashable = IsStructHashable(InStructType);
 
-			PropertyCreateFunc = [bIsHashable, InStructType](UStruct* InScope, FName InName) -> FProperty *
+			PropertyCreateFunc = [bIsHashable, InStructType](FFieldVariant InScope, FName InName) -> FProperty *
 			{
 				auto Property = new FStructProperty(InScope, InName, RF_Public);
 				Property->Struct = InStructType;
@@ -970,7 +1086,7 @@ bool FOptimusDataTypeRegistry::RegisterType(
 
 			PropertyValueConvertFunc = [ConversionEntries, ExpectedPropertySize, ExpectedShaderValueSize](
 				TArrayView<const uint8> InRawValue, 
-				FShaderValueType::FValueView OutShaderValue
+				FShaderValueContainerView OutShaderValue
 				) -> bool
 			{
 				// we can be copying a smaller property into a larger shader side array element
@@ -1076,7 +1192,7 @@ bool FOptimusDataTypeRegistry::RegisterType(
 		{
 			const bool bIsHashable = IsStructHashable(InStructType);
 
-			PropertyCreateFunc = [bIsHashable, InStructType](UStruct* InScope, FName InName) -> FProperty *
+			PropertyCreateFunc = [bIsHashable, InStructType](FFieldVariant InScope, FName InName) -> FProperty *
 			{
 				FStructProperty* Property = new FStructProperty(InScope, InName, RF_Public);
 				Property->Struct = InStructType;
@@ -1132,7 +1248,7 @@ bool FOptimusDataTypeRegistry::RegisterType(
 		{
 			const bool bIsHashable = IsStructHashable(InStructType);
 			
-			PropertyCreateFunc = [bIsHashable, InStructType](UStruct* InScope, FName InName) -> FProperty* {
+			PropertyCreateFunc = [bIsHashable, InStructType](FFieldVariant InScope, FName InName) -> FProperty* {
 				FStructProperty* Property = new FStructProperty(InScope, InName, RF_Public);
 				Property->Struct = InStructType;
 				Property->ElementSize = InStructType->GetStructureSize();
@@ -1179,7 +1295,7 @@ bool FOptimusDataTypeRegistry::RegisterType(
 		PropertyCreateFuncT PropertyCreateFunc;
 		if (EnumHasAnyFlags(InUsageFlags, EOptimusDataTypeUsageFlags::Variable))
 		{
-			PropertyCreateFunc = [InClassType](UStruct* InScope, FName InName) -> FProperty* {
+			PropertyCreateFunc = [InClassType](FFieldVariant InScope, FName InName) -> FProperty* {
 				FObjectProperty* Property = new FObjectProperty(InScope, InName, RF_Public);
 				Property->SetPropertyClass(InClassType);
 				Property->SetPropertyFlags(CPF_HasGetValueTypeHash);
diff --git a/Engine/Plugins/Animation/DeformerGraph/Source/OptimusCore/Private/OptimusDeformer.cpp b/Engine/Plugins/Animation/DeformerGraph/Source/OptimusCore/Private/OptimusDeformer.cpp
--- a/Engine/Plugins/Animation/DeformerGraph/Source/OptimusCore/Private/OptimusDeformer.cpp
+++ b/Engine/Plugins/Animation/DeformerGraph/Source/OptimusCore/Private/OptimusDeformer.cpp
@@ -421,8 +421,7 @@ bool UOptimusDeformer::SetVariableDataType(
 		return false;
 	}
 	
-	// Make sure the value data container is still large enough to hold the property value.
-	InVariableDesc->ValueData.Reset();
+	// Make sure the value data container matches the property value.
 	InVariableDesc->EnsureValueContainer();
 	return true;
 }
@@ -2167,7 +2166,7 @@ TArray<FOptimusComputeGraphInfo> UOptimusDeformer::CompileNodeGraphToComputeGrap
 
 			if (UOptimusNode_ConstantValue const* ConstantNode = Cast<const UOptimusNode_ConstantValue>(ValueNode))
 			{
-				ValueNodeDescription.Value = ConstantNode->GetShaderValue().ShaderValue;
+				ValueNodeDescription.ShaderValue = ConstantNode->GetShaderValue();
 			}
 		}
 	}
@@ -3068,7 +3067,7 @@ void UOptimusDeformer::Notify(EOptimusGlobalNotifyType InNotifyType, UObject* In
 	case EOptimusGlobalNotifyType::ConstantValueChanged:
 		if (UOptimusNode_ConstantValue* ConstantValue = Cast<UOptimusNode_ConstantValue>(InObject))
 		{
-			ConstantValueUpdateDelegate.Broadcast(ConstantValue, ConstantValue->GetShaderValue().ShaderValue);
+			ConstantValueUpdateDelegate.Broadcast(ConstantValue, ConstantValue->GetShaderValue());
 		}
 		
 		break;
diff --git a/Engine/Plugins/Animation/DeformerGraph/Source/OptimusCore/Private/OptimusDeformerInstance.cpp b/Engine/Plugins/Animation/DeformerGraph/Source/OptimusCore/Private/OptimusDeformerInstance.cpp
--- a/Engine/Plugins/Animation/DeformerGraph/Source/OptimusCore/Private/OptimusDeformerInstance.cpp
+++ b/Engine/Plugins/Animation/DeformerGraph/Source/OptimusCore/Private/OptimusDeformerInstance.cpp
@@ -456,8 +456,8 @@ void UOptimusDeformerInstance::SetupFromDeformer(UOptimusDeformer* InDeformer)
 		VariableDescriptionCopy->Guid = VariableDescription->Guid;
 		VariableDescriptionCopy->VariableName = VariableDescription->VariableName;
 		VariableDescriptionCopy->DataType = VariableDescription->DataType;
-		VariableDescriptionCopy->DefaultValue = nullptr; // No need to copy the default value.
-		VariableDescriptionCopy->ValueData = VariableDescription->ValueData;
+		VariableDescriptionCopy->DefaultValue = nullptr; // No need to copy the default value, we directly copy the shader value of the default value below
+		VariableDescriptionCopy->CachedShaderValue = VariableDescription->CachedShaderValue;
 		Variables->Descriptions.Add(VariableDescriptionCopy);
 	}
 
@@ -601,9 +601,7 @@ namespace
 				if (ensure(Property->GetSize() == sizeof(T)))
 				{
 					const uint8* ValueBytes = reinterpret_cast<const uint8*>(&InValue);
-					FShaderValueType::FValue ValueResult = WantedType->MakeShaderValue();
-					WantedType->ConvertPropertyValueToShader(TArrayView<const uint8>(ValueBytes, sizeof(T)), ValueResult);
-					VariableDesc->ValueData = MoveTemp(ValueResult.ShaderValue);
+					WantedType->ConvertPropertyValueToShader(TArrayView<const uint8>(ValueBytes, sizeof(T)), VariableDesc->CachedShaderValue);
 				}
 
 				return true;
@@ -673,7 +671,7 @@ bool UOptimusDeformerInstance::EnqueueTriggerGraph(FName InTriggerGraphName)
 }
 
 
-void UOptimusDeformerInstance::SetConstantValueDirect(TSoftObjectPtr<UObject> InSourceObject, TArray<uint8> const& InValue)
+void UOptimusDeformerInstance::SetConstantValueDirect(TSoftObjectPtr<UObject> InSourceObject, FShaderValueContainer const& InValue)
 {
 	// Poke constants into the UGraphDataProvider objects.
 	// This is an editor only operation when constant nodes are edited in the graph and we want to see the result without a full compile step.
diff --git a/Engine/Plugins/Animation/DeformerGraph/Source/OptimusCore/Private/OptimusValueContainer.cpp b/Engine/Plugins/Animation/DeformerGraph/Source/OptimusCore/Private/OptimusValueContainer.cpp
--- a/Engine/Plugins/Animation/DeformerGraph/Source/OptimusCore/Private/OptimusValueContainer.cpp
+++ b/Engine/Plugins/Animation/DeformerGraph/Source/OptimusCore/Private/OptimusValueContainer.cpp
@@ -106,7 +106,7 @@ FOptimusDataTypeRef UOptimusValueContainer::GetValueType() const
 	return {};	
 }
 
-FShaderValueType::FValue UOptimusValueContainer::GetShaderValue() const
+FShaderValueContainer UOptimusValueContainer::GetShaderValue() const
 {
 	UOptimusValueContainerGeneratorClass* Class = Cast<UOptimusValueContainerGeneratorClass>(GetClass());
 	const FProperty* ValueProperty = Class->PropertyLink;
@@ -115,7 +115,7 @@ FShaderValueType::FValue UOptimusValueContainer::GetShaderValue() const
 	if (ensure(ValueProperty) && ensure(DataType.IsValid()))
 	{
 		TArrayView<const uint8> ValueData(ValueProperty->ContainerPtrToValuePtr<uint8>(this), ValueProperty->GetSize());
-		FShaderValueType::FValue ValueResult = DataType->MakeShaderValue();
+		FShaderValueContainer ValueResult = DataType->MakeShaderValue();
 		if (DataType->ConvertPropertyValueToShader(ValueData, ValueResult))
 		{
 			return ValueResult;
diff --git a/Engine/Plugins/Animation/DeformerGraph/Source/OptimusCore/Private/OptimusVariableDescription.cpp b/Engine/Plugins/Animation/DeformerGraph/Source/OptimusCore/Private/OptimusVariableDescription.cpp
--- a/Engine/Plugins/Animation/DeformerGraph/Source/OptimusCore/Private/OptimusVariableDescription.cpp
+++ b/Engine/Plugins/Animation/DeformerGraph/Source/OptimusCore/Private/OptimusVariableDescription.cpp
@@ -10,39 +10,18 @@
 #include UE_INLINE_GENERATED_CPP_BY_NAME(OptimusVariableDescription)
 
 
-bool UOptimusVariableDescription::EnsureValueContainer()
+void UOptimusVariableDescription::EnsureValueContainer()
 {
-	bool bValueContainerChanged = false;
-	
 	// Check if the current default value storage matches, otherwise create a matching default value storage, otherwise
 	// if the variable type changes, we end up with mismatch in storage vs type.
 	const UClass* RequiredClass = UOptimusValueContainerGeneratorClass::GetClassForType(GetPackage(), DataType);
 
-	bool bValueContainerNeedsZeroing = false;
 	if (!DefaultValue || DefaultValue->GetClass() != RequiredClass)
 	{
 		DefaultValue = UOptimusValueContainer::MakeValueContainer(this, DataType);
-		bValueContainerChanged = true;
-		bValueContainerNeedsZeroing = true;
 	}
 
-	if (DataType->CanCreateProperty())
-	{
-		const FShaderValueType::FValue ShaderValue = DataType->MakeShaderValue();
-		
-		if (bValueContainerNeedsZeroing || ValueData.Num() != ShaderValue.ShaderValue.Num())
-		{
-			ValueData.SetNumZeroed(ShaderValue.ShaderValue.Num());
-			bValueContainerChanged = true;
-		}
-	}
-	else if (!ValueData.IsEmpty())
-	{
-		ValueData.Reset();
-		bValueContainerChanged = true;
-	}
-
-	return bValueContainerChanged;
+	CachedShaderValue = DefaultValue->GetShaderValue();
 }
 
 
@@ -73,10 +52,15 @@ void UOptimusVariableDescription::PostLoad()
 		DataType = DoubleDataType;
 	}
 
-	if (EnsureValueContainer())
+	EnsureValueContainer();
+
+PRAGMA_DISABLE_DEPRECATION_WARNINGS
+	if (!ValueData_DEPRECATED.IsEmpty())
 	{
-		(void)MarkPackageDirty();
+		CachedShaderValue.ShaderValue = ValueData_DEPRECATED;
+		ValueData_DEPRECATED.Reset();
 	}
+PRAGMA_ENABLE_DEPRECATION_WARNINGS
 }
 
 
@@ -111,8 +95,7 @@ void UOptimusVariableDescription::PostEditChangeProperty(FPropertyChangedEvent&
 	else if (PropertyName == GET_MEMBER_NAME_CHECKED(UOptimusVariableDescription, DefaultValue))
 	{
 		// Store the default shader value.
-		FShaderValueType::FValue Value = DefaultValue->GetShaderValue();
-		ValueData = MoveTemp(Value.ShaderValue);
+		CachedShaderValue = DefaultValue->GetShaderValue();
 	}
 }
 
diff --git a/Engine/Plugins/Animation/DeformerGraph/Source/OptimusCore/Public/IOptimusValueProvider.h b/Engine/Plugins/Animation/DeformerGraph/Source/OptimusCore/Public/IOptimusValueProvider.h
--- a/Engine/Plugins/Animation/DeformerGraph/Source/OptimusCore/Public/IOptimusValueProvider.h
+++ b/Engine/Plugins/Animation/DeformerGraph/Source/OptimusCore/Public/IOptimusValueProvider.h
@@ -28,5 +28,5 @@ public:
 	// Returns the value type.
 	virtual FOptimusDataTypeRef GetValueType() const = 0;
 	// Returns the stored value as a shader-compatible value.
-	virtual FShaderValueType::FValue GetShaderValue() const = 0;
+	virtual FShaderValueContainer GetShaderValue() const = 0;
 };
diff --git a/Engine/Plugins/Animation/DeformerGraph/Source/OptimusCore/Public/OptimusDataType.h b/Engine/Plugins/Animation/DeformerGraph/Source/OptimusCore/Public/OptimusDataType.h
--- a/Engine/Plugins/Animation/DeformerGraph/Source/OptimusCore/Public/OptimusDataType.h
+++ b/Engine/Plugins/Animation/DeformerGraph/Source/OptimusCore/Public/OptimusDataType.h
@@ -61,11 +61,11 @@ struct OPTIMUSCORE_API FOptimusDataType
 	// accommodate the newly converted value.
 	bool ConvertPropertyValueToShader(
 		TArrayView<const uint8> InValue,
-		FShaderValueType::FValueView OutConvertedValue
+		FShaderValueContainer& OutConvertedValue
 		) const;
 
 	// Return a value struct that can hold raw shader value of this type
-	FShaderValueType::FValue MakeShaderValue() const;
+	FShaderValueContainer MakeShaderValue() const;
 	
 	// Returns true if the data type can create a FProperty object to represent it.
 	bool CanCreateProperty() const;
diff --git a/Engine/Plugins/Animation/DeformerGraph/Source/OptimusCore/Public/OptimusDataTypeRegistry.h b/Engine/Plugins/Animation/DeformerGraph/Source/OptimusCore/Public/OptimusDataTypeRegistry.h
--- a/Engine/Plugins/Animation/DeformerGraph/Source/OptimusCore/Public/OptimusDataTypeRegistry.h
+++ b/Engine/Plugins/Animation/DeformerGraph/Source/OptimusCore/Public/OptimusDataTypeRegistry.h
@@ -22,14 +22,14 @@ class FOptimusDataTypeRegistry :
 public:
 	DECLARE_EVENT_OneParam(FOptimusDataTypeRegistry, FOnDataTypeChanged, FName /* InTypeName */);
 	
-	using PropertyCreateFuncT = TFunction<FProperty *(UStruct *InScope, FName InName)>;
+	using PropertyCreateFuncT = TFunction<FProperty *(FFieldVariant InScope, FName InName)>;
 
 	/** Defines a function that takes an array view on UE property data and converts it to
 	 *  a matching HLSL shader value data. It's expected that the input and output array views
 	 *  hold enough data to both read all property data for this type and write out shader values
 	 *  matching these data.
 	  */
-	using PropertyValueConvertFuncT = TFunction<bool(TArrayView<const uint8> InRawValue, FShaderValueType::FValueView OutShaderValue)>;
+	using PropertyValueConvertFuncT = TFunction<bool(TArrayView<const uint8> InRawValue, FShaderValueContainerView OutShaderValue)>;
 
 	struct FArrayMetadata
 	{
@@ -124,6 +124,11 @@ public:
 	OPTIMUSCORE_API bool RegisterStructType(
 		UScriptStruct *InStructType
 		);
+	
+	// Register a array type that has corresponding types on both the UE and HLSL side.
+	OPTIMUSCORE_API bool RegisterArrayTypeIfApplicable(
+		FOptimusDataTypeHandle InElementDataType
+		);
 		
 	// Refresh struct type info in case a user defined struct is changed.
 	OPTIMUSCORE_API void RefreshStructType(
diff --git a/Engine/Plugins/Animation/DeformerGraph/Source/OptimusCore/Public/OptimusDeformer.h b/Engine/Plugins/Animation/DeformerGraph/Source/OptimusCore/Public/OptimusDeformer.h
--- a/Engine/Plugins/Animation/DeformerGraph/Source/OptimusCore/Public/OptimusDeformer.h
+++ b/Engine/Plugins/Animation/DeformerGraph/Source/OptimusCore/Public/OptimusDeformer.h
@@ -37,7 +37,7 @@ struct FOptimusCompoundAction;
 DECLARE_MULTICAST_DELEGATE_OneParam(FOptimusCompileBegin, UOptimusDeformer *);
 DECLARE_MULTICAST_DELEGATE_OneParam(FOptimusCompileEnd, UOptimusDeformer *);
 DECLARE_MULTICAST_DELEGATE_OneParam(FOptimusGraphCompileMessageDelegate, FOptimusCompilerDiagnostic const&);
-DECLARE_MULTICAST_DELEGATE_TwoParams(FOptimusConstantValueUpdate, TSoftObjectPtr<UObject>, TArray<uint8> const&);
+DECLARE_MULTICAST_DELEGATE_TwoParams(FOptimusConstantValueUpdate, TSoftObjectPtr<UObject>, FShaderValueContainer const&);
 DECLARE_MULTICAST_DELEGATE_OneParam(FOptimusSetAllInstancesCanbeActive, bool);
 
 UENUM()
diff --git a/Engine/Plugins/Animation/DeformerGraph/Source/OptimusCore/Public/OptimusDeformerInstance.h b/Engine/Plugins/Animation/DeformerGraph/Source/OptimusCore/Public/OptimusDeformerInstance.h
--- a/Engine/Plugins/Animation/DeformerGraph/Source/OptimusCore/Public/OptimusDeformerInstance.h
+++ b/Engine/Plugins/Animation/DeformerGraph/Source/OptimusCore/Public/OptimusDeformerInstance.h
@@ -19,6 +19,7 @@ class UOptimusDeformer;
 class UOptimusVariableContainer;
 class UOptimusVariableDescription;
 class UOptimusComponentSourceBinding;
+struct FShaderValueContainer;
 
 
 struct FOptimusDeformerInstanceComponentLodContext
@@ -246,7 +247,7 @@ public:
 	
 	
 	/** Directly set a graph constant value. */
-	void SetConstantValueDirect(TSoftObjectPtr<UObject> InSourceObject, TArray<uint8> const& InValue);
+	void SetConstantValueDirect(TSoftObjectPtr<UObject> InSourceObject, FShaderValueContainer const& InValue);
 
 	FOptimusPersistentBufferPoolPtr GetBufferPool() const { return BufferPool; }
 
diff --git a/Engine/Plugins/Animation/DeformerGraph/Source/OptimusCore/Public/OptimusValueContainer.h b/Engine/Plugins/Animation/DeformerGraph/Source/OptimusCore/Public/OptimusValueContainer.h
--- a/Engine/Plugins/Animation/DeformerGraph/Source/OptimusCore/Public/OptimusValueContainer.h
+++ b/Engine/Plugins/Animation/DeformerGraph/Source/OptimusCore/Public/OptimusValueContainer.h
@@ -51,5 +51,5 @@ public:
 	static UOptimusValueContainer* MakeValueContainer(UObject* InOwner, FOptimusDataTypeRef InDataTypeRef);
 
 	FOptimusDataTypeRef GetValueType() const;
-	FShaderValueType::FValue GetShaderValue() const;
+	FShaderValueContainer GetShaderValue() const;
 };
\ No newline at end of file
diff --git a/Engine/Plugins/Animation/DeformerGraph/Source/OptimusCore/Public/OptimusVariableDescription.h b/Engine/Plugins/Animation/DeformerGraph/Source/OptimusCore/Public/OptimusVariableDescription.h
--- a/Engine/Plugins/Animation/DeformerGraph/Source/OptimusCore/Public/OptimusVariableDescription.h
+++ b/Engine/Plugins/Animation/DeformerGraph/Source/OptimusCore/Public/OptimusVariableDescription.h
@@ -46,7 +46,7 @@ public:
 	 * Set the value data storage to match the size required by the DataType. 
 	 * If a reallocation is required then the value data will be zeroed.
 	 */
-	bool EnsureValueContainer();
+	void EnsureValueContainer();
 
 	/** Returns the owning deformer to operate on this variable */
 	// FIXME: Move to interface-based system.
@@ -73,7 +73,7 @@ public:
 
 	/** Cached shader value binary data. */
 	UPROPERTY()
-	TArray<uint8> ValueData;
+	FShaderValueContainer CachedShaderValue;
 
 	void PostLoad();
 #if WITH_EDITOR
@@ -83,6 +83,11 @@ public:
 #endif
 
 private:
+
+	// Deprecated 
+    UPROPERTY(meta = (DeprecatedProperty, DeprecationMessage = "Use CachedShaderValue instead"))
+    TArray<uint8> ValueData_DEPRECATED;
+	
 #if WITH_EDITORONLY_DATA
 	FName VariableNameForUndo;
 #endif
diff --git a/Engine/Plugins/Runtime/ComputeFramework/Source/ComputeFramework/Public/ComputeFramework/ShaderParamTypeDefinition.h b/Engine/Plugins/Runtime/ComputeFramework/Source/ComputeFramework/Public/ComputeFramework/ShaderParamTypeDefinition.h
--- a/Engine/Plugins/Runtime/ComputeFramework/Source/ComputeFramework/Public/ComputeFramework/ShaderParamTypeDefinition.h
+++ b/Engine/Plugins/Runtime/ComputeFramework/Source/ComputeFramework/Public/ComputeFramework/ShaderParamTypeDefinition.h
@@ -107,43 +107,76 @@ template<> struct TStructOpsTypeTraits<FShaderValueTypeHandle> : TStructOpsTypeT
 	enum { WithSerializer = true, WithIdenticalViaEquality = true, WithCopy = false };
 };
 
-/*  */
+
+
 USTRUCT()
-struct COMPUTEFRAMEWORK_API FShaderValueType
+struct FArrayShaderValue
+{
+	GENERATED_BODY()
+
+	UPROPERTY()
+	TArray<uint8> ArrayOfValues;
+};
+
+USTRUCT()
+struct FShaderValueContainer
 {
 	GENERATED_BODY()
 
-	struct FValue
+	UPROPERTY()
+	TArray<uint8> ShaderValue;
+
+	UPROPERTY()
+	TArray<FArrayShaderValue> ArrayList;
+
+	FShaderValueContainer() = default;
+	
+	explicit FShaderValueContainer(int32 InShaderValueSize, int32 InNumArrays)
 	{
-		FValue() = default;
-		FValue(int32 InShaderValueSize, int32 InNumArrays)
-		{
-			ShaderValue.SetNumUninitialized(InShaderValueSize);
-			ArrayList.AddDefaulted(InNumArrays);
-		}
-		
-		TArray<uint8> ShaderValue;
-		TArray<TArray<uint8>> ArrayList;
+		ShaderValue.SetNumUninitialized(InShaderValueSize);
+		ArrayList.AddDefaulted(InNumArrays);
+	}
+
+	static bool IsSameType(const FShaderValueContainer& A, const FShaderValueContainer& B)
+	{
+		return A.ShaderValue.Num() == B.ShaderValue.Num() && A.ArrayList.Num() == B.ArrayList.Num();
 	};
-	
-	struct FValueView
+
+	bool IsValid() const
+	{
+		return ShaderValue.Num() > 0 || ArrayList.Num() > 0;
+	}
+	void Reset()
 	{
-		FValueView(
+		ShaderValue.Reset();
+		ArrayList.Reset();
+	};
+};
+
+struct FShaderValueContainerView
+{
+	TArrayView<uint8> ShaderValue;
+	TArrayView<FArrayShaderValue> ArrayList;
+
+	FShaderValueContainerView(
 			TArrayView<uint8> InShaderValue,
-			TArrayView<TArray<uint8>> InArrayList) :
+			TArrayView<FArrayShaderValue> InArrayList) :
 			ShaderValue(InShaderValue),
 			ArrayList(InArrayList) {}
 		
-		FValueView(TArrayView<uint8> InShaderValue) :
-			ShaderValue(InShaderValue) {}
-		
-		FValueView(FValue& InValue) :
-			ShaderValue(InValue.ShaderValue),
-			ArrayList(InValue.ArrayList) {}
+	FShaderValueContainerView(TArrayView<uint8> InShaderValue) :
+		ShaderValue(InShaderValue) {}
 		
-		TArrayView<uint8> ShaderValue;
-		TArrayView<TArray<uint8>> ArrayList;
-	};
+	FShaderValueContainerView(FShaderValueContainer& InValue) :
+		ShaderValue(InValue.ShaderValue),
+		ArrayList(InValue.ArrayList) {}
+};
+
+/*  */
+USTRUCT()
+struct COMPUTEFRAMEWORK_API FShaderValueType
+{
+	GENERATED_BODY()
 	
 	// A simple container representing a single, named element in a shader value struct.
 	struct FStructElement
