[Deformer Graph] fixed minor macarm compilation error

diff --git a/Engine/Plugins/Animation/DeformerGraph/Source/OptimusCore/Private/OptimusDataTypeRegistry.cpp b/Engine/Plugins/Animation/DeformerGraph/Source/OptimusCore/Private/OptimusDataTypeRegistry.cpp
--- a/Engine/Plugins/Animation/DeformerGraph/Source/OptimusCore/Private/OptimusDataTypeRegistry.cpp
+++ b/Engine/Plugins/Animation/DeformerGraph/Source/OptimusCore/Private/OptimusDataTypeRegistry.cpp
@@ -921,7 +921,7 @@ bool FOptimusDataTypeRegistry::RegisterArrayTypeIfApplicable(FOptimusDataTypeHan
 		ArrayDataType,
 		ArrayPropertyCreateFunc,
 		ArrayPropertyValueConvertFunc,
-		{FArrayMetadata(ElementShaderValueSize, 0 )}
+		{{ElementShaderValueSize, 0}}
 	};
 
 	{
