[Deformer Graph] avoid showing array types up in anim attribute type selector

diff --git a/Engine/Plugins/Animation/DeformerGraph/Source/OptimusCore/Private/OptimusDataTypeRegistry.cpp b/Engine/Plugins/Animation/DeformerGraph/Source/OptimusCore/Private/OptimusDataTypeRegistry.cpp
--- a/Engine/Plugins/Animation/DeformerGraph/Source/OptimusCore/Private/OptimusDataTypeRegistry.cpp
+++ b/Engine/Plugins/Animation/DeformerGraph/Source/OptimusCore/Private/OptimusDataTypeRegistry.cpp
@@ -930,6 +930,8 @@ bool FOptimusDataTypeRegistry::RegisterArrayTypeIfApplicable(FOptimusDataTypeHan
 		ArrayDataType->DisplayName = FText::FromString(InElementDataType->DisplayName.ToString() + TEXT(" Array"));
 		ArrayDataType->ShaderValueType = FShaderValueType::MakeDynamicArrayType(InnerDataTypeForStructuredBuffer->ShaderValueType);
 		
+		EnumRemoveFlags(ArrayDataType->UsageFlags, EOptimusDataTypeUsageFlags::Resource | EOptimusDataTypeUsageFlags::AnimAttributes);
+		
 		// Unused field
 		ArrayDataType->ShaderValueSize = 0;	
 	}
