[Deformer Graph] deprecating uclass based value container in favor of propertybag

diff --git a/Engine/Plugins/Animation/DeformerGraph/DeformerGraph.uplugin b/Engine/Plugins/Animation/DeformerGraph/DeformerGraph.uplugin
--- a/Engine/Plugins/Animation/DeformerGraph/DeformerGraph.uplugin
+++ b/Engine/Plugins/Animation/DeformerGraph/DeformerGraph.uplugin
@@ -52,6 +52,10 @@
 		{
 			"Name": "ComputeFramework",
 			"Enabled": true
+		},
+		{
+			"Name": "StructUtils",
+			"Enabled": true
 		}
 	]
 }
diff --git a/Engine/Plugins/Animation/DeformerGraph/Source/OptimusCore/OptimusCore.Build.cs b/Engine/Plugins/Animation/DeformerGraph/Source/OptimusCore/OptimusCore.Build.cs
--- a/Engine/Plugins/Animation/DeformerGraph/Source/OptimusCore/OptimusCore.Build.cs
+++ b/Engine/Plugins/Animation/DeformerGraph/Source/OptimusCore/OptimusCore.Build.cs
@@ -26,6 +26,7 @@ namespace UnrealBuildTool.Rules
 					"RenderCore",
 					"Renderer",
 					"RHI",
+					"StructUtils"
 				}
 			);
 
diff --git a/Engine/Plugins/Animation/DeformerGraph/Source/OptimusCore/Private/Actions/OptimusVariableActions.cpp b/Engine/Plugins/Animation/DeformerGraph/Source/OptimusCore/Private/Actions/OptimusVariableActions.cpp
--- a/Engine/Plugins/Animation/DeformerGraph/Source/OptimusCore/Private/Actions/OptimusVariableActions.cpp
+++ b/Engine/Plugins/Animation/DeformerGraph/Source/OptimusCore/Private/Actions/OptimusVariableActions.cpp
@@ -46,9 +46,7 @@ bool FOptimusVariableAction_AddVariable::Do(IOptimusPathResolver* InRoot)
 
 
 	Variable->VariableName = Variable->GetFName();
-	Variable->DataType = DataType;
-
-	Variable->EnsureValueContainer();
+	Variable->SetDataType(DataType);
 
 	if (!Deformer->AddVariableDirect(Variable, INDEX_NONE))
 	{
@@ -102,11 +100,6 @@ bool FOptimusVariableAction_RemoveVariable::Do(IOptimusPathResolver* InRoot)
 	{
 		Optimus::FBinaryObjectWriter VarArchive(Variable, VariableData);
 	}
-	
-	if (ensure(Variable->DefaultValue))
-	{
-		Optimus::FBinaryObjectWriter ValueArchive(Variable->DefaultValue, DefaultValueData);
-	}
 
 	return Deformer->RemoveVariableDirect(Variable);
 }
@@ -129,13 +122,6 @@ bool FOptimusVariableAction_RemoveVariable::Undo(IOptimusPathResolver* InRoot)
 	{
 		Optimus::FBinaryObjectReader VarArchive(Variable, VariableData);
 	}
-
-	Variable->EnsureValueContainer();
-
-	if (!DefaultValueData.IsEmpty() && Variable->DefaultValue)
-	{
-		Optimus::FBinaryObjectReader VarArchive(Variable->DefaultValue, DefaultValueData);
-	}
 	
 	if (!Deformer->AddVariableDirect(Variable, VariableIndex))
 	{
diff --git a/Engine/Plugins/Animation/DeformerGraph/Source/OptimusCore/Private/DataInterfaces/OptimusDataInterfaceAnimAttribute.cpp b/Engine/Plugins/Animation/DeformerGraph/Source/OptimusCore/Private/DataInterfaces/OptimusDataInterfaceAnimAttribute.cpp
--- a/Engine/Plugins/Animation/DeformerGraph/Source/OptimusCore/Private/DataInterfaces/OptimusDataInterfaceAnimAttribute.cpp
+++ b/Engine/Plugins/Animation/DeformerGraph/Source/OptimusCore/Private/DataInterfaces/OptimusDataInterfaceAnimAttribute.cpp
@@ -5,6 +5,7 @@
 #include "Engine/SkeletalMesh.h"
 #include "OptimusDataDomain.h"
 #include "OptimusDataTypeRegistry.h"
+#include "OptimusObjectVersion.h"
 #include "OptimusValueContainer.h"
 #include "RenderGraphBuilder.h"
 #include "RenderGraphResources.h"
@@ -28,14 +29,14 @@ static const FString HlslIdDelimiter = TEXT("_");
 
 
 
-FOptimusAnimAttributeDescription& FOptimusAnimAttributeDescription::Init(UOptimusAnimAttributeDataInterface* InOwner,const FString& InName, FName InBoneName,
+FOptimusAnimAttributeDescription& FOptimusAnimAttributeDescription::Init(const FString& InName, FName InBoneName,
 	const FOptimusDataTypeRef& InDataType)
 {
 	Name = InName;
 	BoneName = InBoneName;
 	DataType = InDataType;
-	DefaultValue = UOptimusValueContainer::MakeValueContainer(InOwner, InDataType);
 
+	DefaultValueStruct.SetType(DataType.Resolve());
 	// Caller should ensure that the name is unique
 	HlslId = InName;
 	PinName = *InName;
@@ -133,7 +134,7 @@ void UOptimusAnimAttributeDataInterface::PostEditChangeChainProperty(FPropertyCh
 			FOptimusAnimAttributeDescription& ChangedAttribute = AttributeArray[ChangedIndex];
 
 			// Update the default value container accordingly
-			ChangedAttribute.DefaultValue = UOptimusValueContainer::MakeValueContainer(this, ChangedAttribute.DataType);
+			ChangedAttribute.DefaultValueStruct.SetType(ChangedAttribute.DataType);
 		}
 	}
 	else if (PropertyChangedEvent.ChangeType == EPropertyChangeType::ArrayAdd)
@@ -144,7 +145,7 @@ void UOptimusAnimAttributeDataInterface::PostEditChangeChainProperty(FPropertyCh
 			FOptimusAnimAttributeDescription& Attribute = AttributeArray[ChangedIndex];
 			
 			// Default to a float attribute
-			Attribute.Init(this, GetUnusedAttributeName(TEXT("EmptyName")), NAME_None,
+			Attribute.Init(GetUnusedAttributeName(TEXT("EmptyName")), NAME_None,
 				FOptimusDataTypeRegistry::Get().FindType(*FFloatProperty::StaticClass()));
 		}
 	}
@@ -157,9 +158,36 @@ void UOptimusAnimAttributeDataInterface::PostEditChangeChainProperty(FPropertyCh
 			
 			Attribute.Name = GetUnusedAttributeName(Attribute.Name);
 			Attribute.UpdatePinNameAndHlslId();
+		}
+	}
+}
+
+void UOptimusAnimAttributeDataInterface::PostLoad()
+{
+	Super::PostLoad();
+	
+	if (GetLinkerCustomVersion(FOptimusObjectVersion::GUID) < FOptimusObjectVersion::PropertyBagValueContainer)
+	{
+PRAGMA_DISABLE_DEPRECATION_WARNINGS
+		
+		for (FOptimusAnimAttributeDescription& Description : AttributeArray)
+		{
+			if (Description.DefaultValue_DEPRECATED)
+			{
+				Description.DefaultValue_DEPRECATED->ConditionalPostLoad();
+			}
+		}
 
-			Attribute.DefaultValue = DuplicateObject(Attribute.DefaultValue, this);
+		for (FOptimusAnimAttributeDescription& Description : AttributeArray)
+		{
+			if (Description.DefaultValue_DEPRECATED)
+			{
+				Description.DefaultValueStruct = Description.DefaultValue_DEPRECATED->MakeValueContainerStruct();
+				Description.DefaultValue_DEPRECATED = nullptr;
+			}
 		}
+		
+PRAGMA_ENABLE_DEPRECATION_WARNINGS	
 	}
 }
 #endif
@@ -389,45 +417,9 @@ const FOptimusAnimAttributeDescription& UOptimusAnimAttributeDataInterface::AddA
 	const FOptimusDataTypeRef& InDataType)
 {
 	return AttributeArray.InnerArray.AddDefaulted_GetRef()
-		.Init(this, GetUnusedAttributeName(InName), InBoneName, InDataType);
+		.Init(GetUnusedAttributeName(InName), InBoneName, InDataType);
 }
 
-void UOptimusAnimAttributeDataInterface::RecreateValueContainers()
-{
-	for (int32 Index = 0; Index < AttributeArray.Num(); Index++)
-	{
-		FOptimusAnimAttributeDescription& Attribute = AttributeArray[Index];
-
-		if (!Attribute.DefaultValue)
-		{
-			continue;
-		}
-		
-		if (Attribute.DefaultValue->GetClass()->GetPackage() != GetPackage())
-		{
-			// Save container data
-			TArray<uint8> ContainerData;
-			{
-				FMemoryWriter ContainerArchive(ContainerData);
-				FObjectAndNameAsStringProxyArchive ContainerProxyArchive(
-						ContainerArchive, /* bInLoadIfFindFails=*/ false);
-				Attribute.DefaultValue->SerializeScriptProperties(ContainerProxyArchive);
-			}
-			
-			UOptimusValueContainer* NewContainer = UOptimusValueContainer::MakeValueContainer(this,Attribute.DefaultValue->GetValueType());
-
-			// Load container data into the new container
-			{
-				FMemoryReader ContainerArchive(ContainerData);
-				FObjectAndNameAsStringProxyArchive ContainerProxyArchive(
-						ContainerArchive, /* bInLoadIfFindFails=*/ true);
-				NewContainer->SerializeScriptProperties(ContainerProxyArchive);
-			}
-
-			Attribute.DefaultValue = NewContainer;
-		}
-	}
-}
 
 void UOptimusAnimAttributeDataInterface::OnDataTypeChanged(FName InDataType)
 {
@@ -435,8 +427,7 @@ void UOptimusAnimAttributeDataInterface::OnDataTypeChanged(FName InDataType)
 	{
 		if (AttributeDescription.DataType.TypeName == InDataType)
 		{
-			AttributeDescription.DefaultValue =
-				UOptimusValueContainer::MakeValueContainer(this, AttributeDescription.DataType);
+			AttributeDescription.DefaultValueStruct.SetType(AttributeDescription.DataType);
 		}
 	}
 }
@@ -540,10 +531,7 @@ FOptimusAnimAttributeRuntimeData::FOptimusAnimAttributeRuntimeData(
 
 	AttributeType = Registry.FindAttributeType(InDescription.DataType.TypeName);
 
-	if (ensure(InDescription.DefaultValue) && ensure(InDescription.DefaultValue->GetValueType() == InDescription.DataType))
-	{
-		CachedDefaultValue = InDescription.DefaultValue->GetShaderValue();
-	}
+	CachedDefaultValue = InDescription.DefaultValueStruct.GetShaderValue(InDescription.DataType);
 }
 
 void UOptimusAnimAttributeDataProvider::Init(
diff --git a/Engine/Plugins/Animation/DeformerGraph/Source/OptimusCore/Private/DataInterfaces/OptimusDataInterfaceAnimAttribute.h b/Engine/Plugins/Animation/DeformerGraph/Source/OptimusCore/Private/DataInterfaces/OptimusDataInterfaceAnimAttribute.h
--- a/Engine/Plugins/Animation/DeformerGraph/Source/OptimusCore/Private/DataInterfaces/OptimusDataInterfaceAnimAttribute.h
+++ b/Engine/Plugins/Animation/DeformerGraph/Source/OptimusCore/Private/DataInterfaces/OptimusDataInterfaceAnimAttribute.h
@@ -5,6 +5,7 @@
 #include "OptimusComputeDataInterface.h"
 #include "OptimusDataType.h"
 #include "OptimusDataTypeRegistry.h"
+#include "OptimusValueContainerStruct.h"
 #include "ComputeFramework/ComputeDataProvider.h"
 
 #include "OptimusDataInterfaceAnimAttribute.generated.h"
@@ -30,9 +31,10 @@ struct FOptimusAnimAttributeDescription
 	UPROPERTY(EditAnywhere, Category = "Data Interface", meta=(UseInAnimAttribute))
 	FOptimusDataTypeRef DataType;
 
-	UPROPERTY(EditAnywhere, Category = "Data Interface", meta=(EditInLine))
-	TObjectPtr<UOptimusValueContainer> DefaultValue = nullptr;
-	
+	// Default value if the animation attribute is not found
+	UPROPERTY(EditAnywhere, Category = "Data Interface")
+	FOptimusValueContainerStruct DefaultValueStruct;
+
 	UPROPERTY()
 	FString HlslId;
 
@@ -42,14 +44,20 @@ struct FOptimusAnimAttributeDescription
 	void UpdatePinNameAndHlslId(bool bInIncludeBoneName = true, bool bInIncludeTypeName = true);
 	
 	// Helpers
-	FOptimusAnimAttributeDescription& Init(class UOptimusAnimAttributeDataInterface* InOwner,const FString& InName, FName InBoneName,
+	FOptimusAnimAttributeDescription& Init(const FString& InName, FName InBoneName,
 	const FOptimusDataTypeRef& InDataType);
 	
 private:
+	friend class UOptimusAnimAttributeDataInterface;
+	
 	FString GetFormattedId(
 		const FString& InDelimiter,
 		bool bInIncludeBoneName,
 		bool bInIncludeTypeName) const;
+
+	// Deprecated
+	UPROPERTY(meta = (DeprecatedProperty, DeprecationMessage = "use DefaultValueStruct instead"))
+	TObjectPtr<UOptimusValueContainer> DefaultValue_DEPRECATED = nullptr;
 };
 
 USTRUCT()
@@ -111,6 +119,7 @@ public:
 #if WITH_EDITOR
 	void PostEditChangeChainProperty(FPropertyChangedChainEvent& PropertyChangedEvent) override;
 #endif
+	void PostLoad() override;
 	
 	//~ Begin UOptimusComputeDataInterface Interface
 	FString GetDisplayName() const override;
@@ -131,10 +140,6 @@ public:
 
 	const FOptimusAnimAttributeDescription& AddAnimAttribute(const FString& InName, FName InBoneName, const FOptimusDataTypeRef& InDataType);
 
-	// Value containers use generated classes that not duplicated when the asset is duplicated
-	// so they have to be recreated with classes in the current asset
-	void RecreateValueContainers();
-
 	void OnDataTypeChanged(FName InDataType);
 	
 	UPROPERTY(EditAnywhere, Category = "Animation Attribute", meta = (ShowOnlyInnerProperties))
diff --git a/Engine/Plugins/Animation/DeformerGraph/Source/OptimusCore/Private/Nodes/OptimusNode_AnimAttributeDataInterface.cpp b/Engine/Plugins/Animation/DeformerGraph/Source/OptimusCore/Private/Nodes/OptimusNode_AnimAttributeDataInterface.cpp
--- a/Engine/Plugins/Animation/DeformerGraph/Source/OptimusCore/Private/Nodes/OptimusNode_AnimAttributeDataInterface.cpp
+++ b/Engine/Plugins/Animation/DeformerGraph/Source/OptimusCore/Private/Nodes/OptimusNode_AnimAttributeDataInterface.cpp
@@ -79,14 +79,6 @@ void UOptimusNode_AnimAttributeDataInterface::PostEditChangeChainProperty(FPrope
 
 #endif
 
-void UOptimusNode_AnimAttributeDataInterface::RecreateValueContainers()
-{
-	if (UOptimusAnimAttributeDataInterface* Interface = Cast<UOptimusAnimAttributeDataInterface>(DataInterfaceData))
-	{
-		Interface->RecreateValueContainers();
-	}
-}
-
 void UOptimusNode_AnimAttributeDataInterface::OnDataTypeChanged(FName InTypeName)
 {
 	Super::OnDataTypeChanged(InTypeName);
@@ -98,10 +90,6 @@ void UOptimusNode_AnimAttributeDataInterface::OnDataTypeChanged(FName InTypeName
 void UOptimusNode_AnimAttributeDataInterface::PostDuplicate(EDuplicateMode::Type DuplicateMode)
 {
 	Super::PostDuplicate(DuplicateMode);
-
-	// Recreate the value containers in case the node was pasted from a different asset
-	// Otherwise the value container can still reference the value container generator class in the asset it was copied from
-	RecreateValueContainers();
 }
 
 
diff --git a/Engine/Plugins/Animation/DeformerGraph/Source/OptimusCore/Private/Nodes/OptimusNode_AnimAttributeDataInterface.h b/Engine/Plugins/Animation/DeformerGraph/Source/OptimusCore/Private/Nodes/OptimusNode_AnimAttributeDataInterface.h
--- a/Engine/Plugins/Animation/DeformerGraph/Source/OptimusCore/Private/Nodes/OptimusNode_AnimAttributeDataInterface.h
+++ b/Engine/Plugins/Animation/DeformerGraph/Source/OptimusCore/Private/Nodes/OptimusNode_AnimAttributeDataInterface.h
@@ -22,7 +22,6 @@ public:
 	void PostEditChangeChainProperty(FPropertyChangedChainEvent& PropertyChangedEvent) override;
 #endif
 
-	void RecreateValueContainers();
 
 	virtual void OnDataTypeChanged(FName InTypeName) override;
 
diff --git a/Engine/Plugins/Animation/DeformerGraph/Source/OptimusCore/Private/Nodes/OptimusNode_GetVariable.cpp b/Engine/Plugins/Animation/DeformerGraph/Source/OptimusCore/Private/Nodes/OptimusNode_GetVariable.cpp
--- a/Engine/Plugins/Animation/DeformerGraph/Source/OptimusCore/Private/Nodes/OptimusNode_GetVariable.cpp
+++ b/Engine/Plugins/Animation/DeformerGraph/Source/OptimusCore/Private/Nodes/OptimusNode_GetVariable.cpp
@@ -115,10 +115,10 @@ void UOptimusNode_GetVariable::ExportCustomProperties(FOutputDevice& Out, uint32
 		Out.Logf(TEXT("%sCustomProperties VariableDefinition Name=\"%s\" Type=%s"),
 			FCString::Spc(Indent), *Var->VariableName.ToString(), *Var->DataType->TypeName.ToString());
 
-		if (const FProperty* Property = Var->DefaultValue->GetClass()->FindPropertyByName(UOptimusValueContainerGeneratorClass::ValuePropertyName))
+		if (const FProperty* Property = Var->DefaultValueStruct.GetValueProperty())
 		{
 			FString ValueStr;
-			Property->ExportTextItem_InContainer(ValueStr, Var->DefaultValue.Get(), nullptr, nullptr, PPF_None);
+			Property->ExportTextItem_InContainer(ValueStr, Var->DefaultValueStruct.GetValueMemory(), nullptr, nullptr, PPF_None);
 			Out.Logf(TEXT(" DefaultValue=\"%s\""), *ValueStr.ReplaceCharWithEscapedChar());
 		}
 		Out.Logf(TEXT("\n"));
diff --git a/Engine/Plugins/Animation/DeformerGraph/Source/OptimusCore/Private/OptimusDeformer.cpp b/Engine/Plugins/Animation/DeformerGraph/Source/OptimusCore/Private/OptimusDeformer.cpp
--- a/Engine/Plugins/Animation/DeformerGraph/Source/OptimusCore/Private/OptimusDeformer.cpp
+++ b/Engine/Plugins/Animation/DeformerGraph/Source/OptimusCore/Private/OptimusDeformer.cpp
@@ -421,8 +421,6 @@ bool UOptimusDeformer::SetVariableDataType(
 		return false;
 	}
 	
-	// Make sure the value data container matches the property value.
-	InVariableDesc->EnsureValueContainer();
 	return true;
 }
 
@@ -575,7 +573,7 @@ bool UOptimusDeformer::SetVariableDataTypeDirect(
 	
 	if (InVariableDesc->DataType != InDataType)
 	{
-		InVariableDesc->DataType = InDataType;
+		InVariableDesc->SetDataType(InDataType);
 		Notify(EOptimusGlobalNotifyType::VariableTypeChanged, InVariableDesc);
 		(void)MarkPackageDirty();
 	}
@@ -2872,9 +2870,6 @@ TArray<FOptimusComputeGraphInfo> UOptimusDeformer::CompileNodeGraphToComputeGrap
 
 void UOptimusDeformer::OnDataTypeChanged(FName InTypeName)
 {
-	// Currently only value containers depends on the UDSs,
-	UOptimusValueContainerGeneratorClass::RefreshClassForType(GetPackage(), FOptimusDataTypeRegistry::Get().FindType(InTypeName));
-	
 	for (UOptimusNodeGraph* Graph : Graphs)
 	{
 		for (UOptimusNode* Node: Graph->Nodes)
@@ -3126,6 +3121,11 @@ void UOptimusDeformer::PostLoad()
 {
 	Super::PostLoad();
 
+	for (UOptimusVariableDescription* VariableDescription : Variables->Descriptions)
+	{
+		VariableDescription->ConditionalPostLoad();
+	}
+	
 	// PostLoad everything first before changing anything for back compat
 	// Each graph postloads everything it owns
 	for (UOptimusNodeGraph* Graph: GetGraphs())
@@ -3221,6 +3221,11 @@ void UOptimusDeformer::PostLoad()
 		PostLoadRemoveDeprecatedExecutionNodes();
 	}
 
+	if (GetLinkerCustomVersion(FOptimusObjectVersion::GUID) < FOptimusObjectVersion::PropertyBagValueContainer)
+	{
+		PostLoadRemoveDeprecatedValueContainerGeneratorClass();
+	}
+
 	// If the graph was saved at any previous version, and was clean, mark the status now as modified.
 	if (Status == EOptimusDeformerStatus::Compiled &&
 		GetLinkerCustomVersion(FOptimusObjectVersion::GUID) < FOptimusObjectVersion::LatestVersion)
@@ -3428,6 +3433,22 @@ void UOptimusDeformer::PostLoadRemoveDeprecatedExecutionNodes()
 	(void)MarkPackageDirty();
 }
 
+void UOptimusDeformer::PostLoadRemoveDeprecatedValueContainerGeneratorClass()
+{
+	// Remove deprecated uclass based value container generator class
+	TArray<UObject*> ObjectsInPackage;
+	GetObjectsWithOuter(GetPackage(), ObjectsInPackage, false);
+
+	for (UObject* Object : ObjectsInPackage)
+	{
+		if (UOptimusValueContainerGeneratorClass* GeneratorClass = Cast<UOptimusValueContainerGeneratorClass>(Object))
+		{
+			Optimus::RemoveObject(GeneratorClass);
+			Optimus::RemoveObject(GeneratorClass->GetDefaultObject());
+		}
+	}	
+}
+
 
 UOptimusComponentSourceBinding* UOptimusDeformer::FindCompatibleBindingWithInterface(
 	const UOptimusComputeDataInterface* InDataInterface) const
diff --git a/Engine/Plugins/Animation/DeformerGraph/Source/OptimusCore/Private/OptimusDeformerInstance.cpp b/Engine/Plugins/Animation/DeformerGraph/Source/OptimusCore/Private/OptimusDeformerInstance.cpp
--- a/Engine/Plugins/Animation/DeformerGraph/Source/OptimusCore/Private/OptimusDeformerInstance.cpp
+++ b/Engine/Plugins/Animation/DeformerGraph/Source/OptimusCore/Private/OptimusDeformerInstance.cpp
@@ -456,7 +456,7 @@ void UOptimusDeformerInstance::SetupFromDeformer(UOptimusDeformer* InDeformer)
 		VariableDescriptionCopy->Guid = VariableDescription->Guid;
 		VariableDescriptionCopy->VariableName = VariableDescription->VariableName;
 		VariableDescriptionCopy->DataType = VariableDescription->DataType;
-		VariableDescriptionCopy->DefaultValue = nullptr; // No need to copy the default value, we directly copy the shader value of the default value below
+		 // No need to copy the default value, we directly copy the shader value of the default value below
 		VariableDescriptionCopy->CachedShaderValue = VariableDescription->CachedShaderValue;
 		Variables->Descriptions.Add(VariableDescriptionCopy);
 	}
diff --git a/Engine/Plugins/Animation/DeformerGraph/Source/OptimusCore/Private/OptimusNodeGraph.cpp b/Engine/Plugins/Animation/DeformerGraph/Source/OptimusCore/Private/OptimusNodeGraph.cpp
--- a/Engine/Plugins/Animation/DeformerGraph/Source/OptimusCore/Private/OptimusNodeGraph.cpp
+++ b/Engine/Plugins/Animation/DeformerGraph/Source/OptimusCore/Private/OptimusNodeGraph.cpp
@@ -152,18 +152,6 @@ void UOptimusNodeGraph::PostDuplicate(EDuplicateMode::Type DuplicateMode)
 			}
 		}	
 	}
-
-	// Similarly, Attribute Nodes also use objects with generated class that need to be recreated
-	// if the duplication took place at the asset level
-	{
-		for (UOptimusNode* Node : Nodes)
-		{
-			if (UOptimusNode_AnimAttributeDataInterface* AttributeNode = Cast<UOptimusNode_AnimAttributeDataInterface>(Node))
-			{
-				AttributeNode->RecreateValueContainers();
-			}
-		}	
-	}
 }
 
 void UOptimusNodeGraph::PostLoad()
diff --git a/Engine/Plugins/Animation/DeformerGraph/Source/OptimusCore/Private/OptimusValueContainer.cpp b/Engine/Plugins/Animation/DeformerGraph/Source/OptimusCore/Private/OptimusValueContainer.cpp
--- a/Engine/Plugins/Animation/DeformerGraph/Source/OptimusCore/Private/OptimusValueContainer.cpp
+++ b/Engine/Plugins/Animation/DeformerGraph/Source/OptimusCore/Private/OptimusValueContainer.cpp
@@ -2,6 +2,7 @@
 
 #include "OptimusValueContainer.h"
 #include "OptimusHelpers.h"
+#include "OptimusValueContainerStruct.h"
 
 #include UE_INLINE_GENERATED_CPP_BY_NAME(OptimusValueContainer)
 
@@ -96,6 +97,28 @@ UOptimusValueContainer* UOptimusValueContainer::MakeValueContainer(UObject* InOw
 	return NewObject<UOptimusValueContainer>(InOwner, Class);
 }
 
+FOptimusValueContainerStruct UOptimusValueContainer::MakeValueContainerStruct()
+{
+	FOptimusValueContainerStruct OtherContainer;
+	
+	UOptimusValueContainerGeneratorClass* Class = Cast<UOptimusValueContainerGeneratorClass>(GetClass());
+	const FProperty* ValueProperty = Class->PropertyLink;
+	
+	FOptimusDataTypeRef DataType = GetValueType();
+	if (ensure(ValueProperty) && ensure(DataType.IsValid()))
+	{
+		OtherContainer.SetType(DataType.Resolve());
+		
+		TArrayView<const uint8> ValueData(ValueProperty->ContainerPtrToValuePtr<uint8>(this), ValueProperty->GetSize());
+
+		uint8* OtherContainerMemory = OtherContainer.GetMutableValueMemory();
+
+		FMemory::Memcpy(OtherContainerMemory, ValueData.GetData(), ValueData.Num());
+	}
+
+	return OtherContainer;
+}
+
 FOptimusDataTypeRef UOptimusValueContainer::GetValueType() const
 {
 	UOptimusValueContainerGeneratorClass* Class = Cast<UOptimusValueContainerGeneratorClass>(GetClass());
diff --git a/Engine/Plugins/Animation/DeformerGraph/Source/OptimusCore/Private/OptimusValueContainerStruct.cpp b/Engine/Plugins/Animation/DeformerGraph/Source/OptimusCore/Private/OptimusValueContainerStruct.cpp
new file mode 100644
--- /dev/null
+++ b/Engine/Plugins/Animation/DeformerGraph/Source/OptimusCore/Private/OptimusValueContainerStruct.cpp
@@ -0,0 +1,62 @@
+﻿// Copyright Epic Games, Inc. All Rights Reserved.
+
+#include "OptimusValueContainerStruct.h"
+
+#include UE_INLINE_GENERATED_CPP_BY_NAME(OptimusValueContainerStruct)
+
+const TCHAR* FOptimusValueContainerStruct::ValuePropertyName = TEXT("Value");
+
+bool FOptimusValueContainerStruct::IsInitialized() const
+{
+	return Value.GetNumPropertiesInBag() == 1;
+}
+
+void FOptimusValueContainerStruct::SetType(FOptimusDataTypeRef InDataType)
+{
+	Value.Reset();
+	Value.AddProperty(ValuePropertyName, InDataType->CreateProperty(nullptr, ValuePropertyName));		
+}
+
+FShaderValueContainer FOptimusValueContainerStruct::GetShaderValue(FOptimusDataTypeRef InDataType) const
+{
+	check(Value.GetNumPropertiesInBag() == 1);
+
+	const UPropertyBag* BagStruct = Value.GetPropertyBagStruct();
+	TConstArrayView<FPropertyBagPropertyDesc> Descs = BagStruct->GetPropertyDescs();
+
+	const FPropertyBagPropertyDesc& ValueDesc = Descs[0];
+
+	if (ensure(InDataType.IsValid()) && ensure(ValueDesc.CachedProperty))
+	{
+		TArrayView<const uint8> ValueData(ValueDesc.CachedProperty->ContainerPtrToValuePtr<uint8>(Value.GetValue().GetMemory()), ValueDesc.CachedProperty->GetSize());
+		FShaderValueContainer ValueResult = InDataType->MakeShaderValue();
+		if (InDataType->ConvertPropertyValueToShader(ValueData, ValueResult))
+		{
+			return ValueResult;
+		}
+	}
+	
+	return {};
+}
+
+const FProperty* FOptimusValueContainerStruct::GetValueProperty() const
+{
+	check(Value.GetNumPropertiesInBag() == 1);
+	
+	const UPropertyBag* BagStruct = Value.GetPropertyBagStruct();
+	TConstArrayView<FPropertyBagPropertyDesc> Descs = BagStruct->GetPropertyDescs();
+
+	const FPropertyBagPropertyDesc& ValueDesc = Descs[0];
+
+	return ValueDesc.CachedProperty;
+}
+
+const uint8* FOptimusValueContainerStruct::GetValueMemory() const
+{
+	return Value.GetValue().GetMemory();
+}
+
+uint8* FOptimusValueContainerStruct::GetMutableValueMemory()
+{
+	return Value.GetMutableValue().GetMemory();
+}
diff --git a/Engine/Plugins/Animation/DeformerGraph/Source/OptimusCore/Private/OptimusVariableDescription.cpp b/Engine/Plugins/Animation/DeformerGraph/Source/OptimusCore/Private/OptimusVariableDescription.cpp
--- a/Engine/Plugins/Animation/DeformerGraph/Source/OptimusCore/Private/OptimusVariableDescription.cpp
+++ b/Engine/Plugins/Animation/DeformerGraph/Source/OptimusCore/Private/OptimusVariableDescription.cpp
@@ -5,23 +5,27 @@
 #include "OptimusDataTypeRegistry.h"
 #include "OptimusDeformer.h"
 #include "OptimusHelpers.h"
+#include "OptimusObjectVersion.h"
 #include "OptimusValueContainer.h"
 
 #include UE_INLINE_GENERATED_CPP_BY_NAME(OptimusVariableDescription)
 
 
-void UOptimusVariableDescription::EnsureValueContainer()
+void UOptimusVariableDescription::SetDataType(FOptimusDataTypeRef InDataType)
 {
-	// Check if the current default value storage matches, otherwise create a matching default value storage, otherwise
-	// if the variable type changes, we end up with mismatch in storage vs type.
-	const UClass* RequiredClass = UOptimusValueContainerGeneratorClass::GetClassForType(GetPackage(), DataType);
-
-	if (!DefaultValue || DefaultValue->GetClass() != RequiredClass)
+	if (InDataType != DataType)
 	{
-		DefaultValue = UOptimusValueContainer::MakeValueContainer(this, DataType);
+		DataType = InDataType;
+		DefaultValueStruct.SetType(InDataType);
+	}
+	else
+	{
+		DataType = InDataType;
+		if (!DefaultValueStruct.IsInitialized())
+		{
+			DefaultValueStruct.SetType(InDataType);
+		}
 	}
-
-	CachedShaderValue = DefaultValue->GetShaderValue();
 }
 
 
@@ -49,18 +53,38 @@ void UOptimusVariableDescription::PostLoad()
 	
 	if (DataType == FloatDataType)
 	{
-		DataType = DoubleDataType;
+		SetDataType(DoubleDataType);
 	}
-
-	EnsureValueContainer();
-
+	
 PRAGMA_DISABLE_DEPRECATION_WARNINGS
 	if (!ValueData_DEPRECATED.IsEmpty())
 	{
 		CachedShaderValue.ShaderValue = ValueData_DEPRECATED;
 		ValueData_DEPRECATED.Reset();
 	}
+
+	if (GetLinkerCustomVersion(FOptimusObjectVersion::GUID) < FOptimusObjectVersion::PropertyBagValueContainer)
+	{
+		if (DefaultValue_DEPRECATED)
+		{
+			DefaultValue_DEPRECATED->ConditionalPostLoad();
+			DefaultValueStruct = DefaultValue_DEPRECATED->MakeValueContainerStruct();
+			DefaultValue_DEPRECATED = nullptr;
+		}
+	}
 PRAGMA_ENABLE_DEPRECATION_WARNINGS
+
+	if (!DefaultValueStruct.IsInitialized())
+	{
+		SetDataType(DataType);
+	}
+
+	// The cached shader value might have incorrect size due to LWC changes,
+	// so we refresh the cache here. Example: FVector shader value was saved as 3*8 bytes at some point
+	if (DataType->ShaderValueSize != CachedShaderValue.ShaderValue.Num())
+	{
+		CachedShaderValue = DefaultValueStruct.GetShaderValue(DataType);
+	}
 }
 
 
@@ -87,15 +111,18 @@ void UOptimusVariableDescription::PostEditChangeProperty(FPropertyChangedEvent&
 		UOptimusDeformer* Deformer = GetOwningDeformer();
 		if (ensure(Deformer))
 		{
+			// Keep the default value in sync
+			DefaultValueStruct.SetType(DataType);
+			
 			// Set the variable type again, so that we can remove any links that are now type-incompatible.
 			constexpr bool bForceChange = true;
 			Deformer->SetVariableDataType(this, DataType, bForceChange);
 		}
 	}
-	else if (PropertyName == GET_MEMBER_NAME_CHECKED(UOptimusVariableDescription, DefaultValue))
+	else if (PropertyName == GET_MEMBER_NAME_CHECKED(UOptimusVariableDescription, DefaultValueStruct))
 	{
 		// Store the default shader value.
-		CachedShaderValue = DefaultValue->GetShaderValue();
+		CachedShaderValue = DefaultValueStruct.GetShaderValue(DataType);
 	}
 }
 
diff --git a/Engine/Plugins/Animation/DeformerGraph/Source/OptimusCore/Public/OptimusDeformer.h b/Engine/Plugins/Animation/DeformerGraph/Source/OptimusCore/Public/OptimusDeformer.h
--- a/Engine/Plugins/Animation/DeformerGraph/Source/OptimusCore/Public/OptimusDeformer.h
+++ b/Engine/Plugins/Animation/DeformerGraph/Source/OptimusCore/Public/OptimusDeformer.h
@@ -493,6 +493,7 @@ private:
 	void PostLoadFixupMissingComponentBindingsCompat();
 	void PostLoadFixupMismatchedResourceDataDomains();
 	void PostLoadRemoveDeprecatedExecutionNodes();
+	void PostLoadRemoveDeprecatedValueContainerGeneratorClass();
 
 	/** Find a compatible binding with the given data interface. Returns nullptr if no such binding exists */
 	UOptimusComponentSourceBinding* FindCompatibleBindingWithInterface(
diff --git a/Engine/Plugins/Animation/DeformerGraph/Source/OptimusCore/Public/OptimusObjectVersion.h b/Engine/Plugins/Animation/DeformerGraph/Source/OptimusCore/Public/OptimusObjectVersion.h
--- a/Engine/Plugins/Animation/DeformerGraph/Source/OptimusCore/Public/OptimusObjectVersion.h
+++ b/Engine/Plugins/Animation/DeformerGraph/Source/OptimusCore/Public/OptimusObjectVersion.h
@@ -21,6 +21,7 @@ struct OPTIMUSCORE_API FOptimusObjectVersion
 		DataDomainExpansion,
 		KernelDataInterface,
 		KernelParameterBindingToggleAtomic,
+		PropertyBagValueContainer,
 
 		// -----<new versions can be added above this line>-------------------------------------------------
 		VersionPlusOne,
diff --git a/Engine/Plugins/Animation/DeformerGraph/Source/OptimusCore/Public/OptimusValueContainer.h b/Engine/Plugins/Animation/DeformerGraph/Source/OptimusCore/Public/OptimusValueContainer.h
--- a/Engine/Plugins/Animation/DeformerGraph/Source/OptimusCore/Public/OptimusValueContainer.h
+++ b/Engine/Plugins/Animation/DeformerGraph/Source/OptimusCore/Public/OptimusValueContainer.h
@@ -6,6 +6,9 @@
 #include "OptimusDataType.h"
 #include "OptimusValueContainer.generated.h"
 
+struct FOptimusValueContainerStruct;
+
+// Deprecated
 UCLASS()
 class OPTIMUSCORE_API UOptimusValueContainerGeneratorClass : public UClass
 {
@@ -20,7 +23,9 @@ public:
 	// this flag has to be kept such that we can load the old asset in the first place and
 	// re-parent it back to the package in post load
 	DECLARE_WITHIN(UObject)
-
+private:
+	friend class UOptimusValueContainer;
+	
 	static FName ValuePropertyName;
 	
 	// UClass overrides
@@ -40,12 +45,17 @@ public:
 	FOptimusDataTypeRef DataType;
 };
 
+// Deprecated
 UCLASS()
 class OPTIMUSCORE_API UOptimusValueContainer : public UObject
 {
 public:
 	GENERATED_BODY()
-
+	
+	// Convert to the newer container type
+	FOptimusValueContainerStruct MakeValueContainerStruct();
+	
+private:
 	void PostLoad() override;
 	
 	static UOptimusValueContainer* MakeValueContainer(UObject* InOwner, FOptimusDataTypeRef InDataTypeRef);
diff --git a/Engine/Plugins/Animation/DeformerGraph/Source/OptimusCore/Public/OptimusValueContainerStruct.h b/Engine/Plugins/Animation/DeformerGraph/Source/OptimusCore/Public/OptimusValueContainerStruct.h
new file mode 100644
--- /dev/null
+++ b/Engine/Plugins/Animation/DeformerGraph/Source/OptimusCore/Public/OptimusValueContainerStruct.h
@@ -0,0 +1,31 @@
+﻿// Copyright Epic Games, Inc. All Rights Reserved.
+
+#pragma once
+
+#include "CoreMinimal.h"
+#include "OptimusDataType.h"
+#include "PropertyBag.h"
+#include "OptimusValueContainerStruct.generated.h"
+
+USTRUCT()
+struct OPTIMUSCORE_API FOptimusValueContainerStruct
+{
+	GENERATED_BODY()
+
+	static const TCHAR* ValuePropertyName;
+	
+	UPROPERTY(EditAnywhere, Category="Value", meta=(FixedLayout, ShowOnlyInnerProperties))
+	FInstancedPropertyBag Value;
+
+	bool IsInitialized() const;
+	void SetType(FOptimusDataTypeRef InDataType);
+	FShaderValueContainer GetShaderValue(FOptimusDataTypeRef InDataType) const;
+	const FProperty* GetValueProperty() const;
+	const uint8* GetValueMemory() const;
+	
+protected:
+	friend class UOptimusValueContainer;
+	
+	uint8* GetMutableValueMemory();
+	
+};
diff --git a/Engine/Plugins/Animation/DeformerGraph/Source/OptimusCore/Public/OptimusVariableDescription.h b/Engine/Plugins/Animation/DeformerGraph/Source/OptimusCore/Public/OptimusVariableDescription.h
--- a/Engine/Plugins/Animation/DeformerGraph/Source/OptimusCore/Public/OptimusVariableDescription.h
+++ b/Engine/Plugins/Animation/DeformerGraph/Source/OptimusCore/Public/OptimusVariableDescription.h
@@ -5,6 +5,7 @@
 #include "OptimusDataType.h"
 
 #include "CoreMinimal.h"
+#include "OptimusValueContainerStruct.h"
 #include "UObject/NameTypes.h"
 #include "UObject/Object.h"
 #include "UObject/ObjectMacros.h"
@@ -43,10 +44,9 @@ class OPTIMUSCORE_API UOptimusVariableDescription :
 	GENERATED_BODY()
 public:
 	/** 
-	 * Set the value data storage to match the size required by the DataType. 
-	 * If a reallocation is required then the value data will be zeroed.
+	 * Set the data type, and recreate the backing data storage as well
 	 */
-	void EnsureValueContainer();
+	void SetDataType(FOptimusDataTypeRef InDataType);
 
 	/** Returns the owning deformer to operate on this variable */
 	// FIXME: Move to interface-based system.
@@ -68,9 +68,9 @@ public:
 	FOptimusDataTypeRef DataType;
 
 	/** The default value for the variable. */
-	UPROPERTY(EditAnywhere, Category = VariableDefinition, meta = (EditInLine))
-	TObjectPtr<UOptimusValueContainer> DefaultValue = nullptr;
-
+	UPROPERTY(EditAnywhere, Category = VariableDefinition)
+	FOptimusValueContainerStruct DefaultValueStruct;
+	
 	/** Cached shader value binary data. */
 	UPROPERTY()
 	FShaderValueContainer CachedShaderValue;
@@ -87,6 +87,10 @@ private:
 	// Deprecated 
     UPROPERTY(meta = (DeprecatedProperty, DeprecationMessage = "Use CachedShaderValue instead"))
     TArray<uint8> ValueData_DEPRECATED;
+
+	// Deprecated 
+	UPROPERTY(meta = (DeprecatedProperty, DeprecationMessage = "use DefaultValueStruct instead"))
+	TObjectPtr<UOptimusValueContainer> DefaultValue_DEPRECATED = nullptr;
 	
 #if WITH_EDITORONLY_DATA
 	FName VariableNameForUndo;
diff --git a/Engine/Plugins/Animation/DeformerGraph/Source/OptimusEditor/OptimusEditor.Build.cs b/Engine/Plugins/Animation/DeformerGraph/Source/OptimusEditor/OptimusEditor.Build.cs
--- a/Engine/Plugins/Animation/DeformerGraph/Source/OptimusEditor/OptimusEditor.Build.cs
+++ b/Engine/Plugins/Animation/DeformerGraph/Source/OptimusEditor/OptimusEditor.Build.cs
@@ -32,8 +32,9 @@ namespace UnrealBuildTool.Rules
 					"MessageLog",
 					"PropertyEditor",
 					"KismetCompiler",
-					"SourceCodeAccess"
-				}
+					"SourceCodeAccess", 
+					"StructUtilsEditor"
+                }
 			);
 
         }
diff --git a/Engine/Plugins/Animation/DeformerGraph/Source/OptimusEditor/Private/OptimusDetailsCustomization.cpp b/Engine/Plugins/Animation/DeformerGraph/Source/OptimusEditor/Private/OptimusDetailsCustomization.cpp
--- a/Engine/Plugins/Animation/DeformerGraph/Source/OptimusEditor/Private/OptimusDetailsCustomization.cpp
+++ b/Engine/Plugins/Animation/DeformerGraph/Source/OptimusEditor/Private/OptimusDetailsCustomization.cpp
@@ -28,6 +28,8 @@
 #include "OptimusValidatedName.h"
 #include "OptimusValueContainer.h"
 #include "OptimusExecutionDomain.h"
+#include "OptimusValueContainerStruct.h"
+#include "PropertyBagDetails.h"
 #include "ScopedTransaction.h"
 #include "Styling/AppStyle.h"
 #include "Styling/SlateIconFinder.h"
@@ -1545,11 +1547,12 @@ void FOptimusParameterBindingArrayCustomization::CustomizeChildren(TSharedRef<IP
 	InChildBuilder.AddCustomBuilder(ArrayBuilder.ToSharedRef());
 }
 
-FOptimusValueContainerCustomization::FOptimusValueContainerCustomization()
+FOptimusValueContainerStructCustomization::FOptimusValueContainerStructCustomization()
 {
+	PropertyBagCustomization = FPropertyBagDetails::MakeInstance();
 }
 
-void FOptimusValueContainerCustomization::CustomizeHeader(TSharedRef<IPropertyHandle> InPropertyHandle,
+void FOptimusValueContainerStructCustomization::CustomizeHeader(TSharedRef<IPropertyHandle> InPropertyHandle,
 	FDetailWidgetRow& InHeaderRow, IPropertyTypeCustomizationUtils& InCustomizationUtils)
 {
 	uint32 NumChildren = 0;
@@ -1558,33 +1561,18 @@ void FOptimusValueContainerCustomization::CustomizeHeader(TSharedRef<IPropertyHa
 	// During reordering, we may have zero children temporarily
 	if (NumChildren > 0)
 	{
-		InnerPropertyHandle = InPropertyHandle->GetChildHandle(UOptimusValueContainerGeneratorClass::ValuePropertyName, true);
+		InnerPropertyHandle = InPropertyHandle->GetChildHandle(GET_MEMBER_NAME_CHECKED(FOptimusValueContainerStruct, Value));
 
-		if (ensure(InnerPropertyHandle.IsValid()))
-		{
-			InHeaderRow.NameContent()
-			[
-				InPropertyHandle->CreatePropertyNameWidget()
-			]
-			.ValueContent()
-			[
-				InnerPropertyHandle->CreatePropertyValueWidget()
-			];
-		}
+		PropertyBagCustomization->CustomizeHeader(InnerPropertyHandle.ToSharedRef(), InHeaderRow, InCustomizationUtils);
 	}
 }
 
-void FOptimusValueContainerCustomization::CustomizeChildren(TSharedRef<IPropertyHandle> InPropertyHandle,
+void FOptimusValueContainerStructCustomization::CustomizeChildren(TSharedRef<IPropertyHandle> InPropertyHandle,
 	IDetailChildrenBuilder& InChildBuilder, IPropertyTypeCustomizationUtils& InCustomizationUtils)
 {
 	if (InnerPropertyHandle)
 	{
-		uint32 NumChildren = 0;
-		InnerPropertyHandle->GetNumChildren(NumChildren)	;
-		for (uint32 Index = 0; Index < NumChildren; Index++)
-		{
-			InChildBuilder.AddProperty(InnerPropertyHandle->GetChildHandle(Index).ToSharedRef());
-		}
+		PropertyBagCustomization->CustomizeChildren(InnerPropertyHandle.ToSharedRef(), InChildBuilder, InCustomizationUtils);
 	}
 }
 
diff --git a/Engine/Plugins/Animation/DeformerGraph/Source/OptimusEditor/Private/OptimusDetailsCustomization.h b/Engine/Plugins/Animation/DeformerGraph/Source/OptimusEditor/Private/OptimusDetailsCustomization.h
--- a/Engine/Plugins/Animation/DeformerGraph/Source/OptimusEditor/Private/OptimusDetailsCustomization.h
+++ b/Engine/Plugins/Animation/DeformerGraph/Source/OptimusEditor/Private/OptimusDetailsCustomization.h
@@ -322,15 +322,15 @@ private:
 };
 
 
-class FOptimusValueContainerCustomization : public IPropertyTypeCustomization
+class FOptimusValueContainerStructCustomization : public IPropertyTypeCustomization
 {
 public:
 	static TSharedRef<IPropertyTypeCustomization> MakeInstance()
 	{
-		return MakeShared<FOptimusValueContainerCustomization>();
+		return MakeShared<FOptimusValueContainerStructCustomization>();
 	}
 
-	FOptimusValueContainerCustomization();
+	FOptimusValueContainerStructCustomization();
 	// IPropertyTypeCustomization overrides
 	virtual void CustomizeHeader(
 		TSharedRef<IPropertyHandle> InPropertyHandle,
@@ -343,6 +343,8 @@ public:
 		IPropertyTypeCustomizationUtils& InCustomizationUtils) override;
 private:
 	TSharedPtr<IPropertyHandle> InnerPropertyHandle;
+
+	TSharedPtr<IPropertyTypeCustomization> PropertyBagCustomization;
 };
 
 
diff --git a/Engine/Plugins/Animation/DeformerGraph/Source/OptimusEditor/Private/OptimusEditorModule.cpp b/Engine/Plugins/Animation/DeformerGraph/Source/OptimusEditor/Private/OptimusEditorModule.cpp
--- a/Engine/Plugins/Animation/DeformerGraph/Source/OptimusEditor/Private/OptimusEditorModule.cpp
+++ b/Engine/Plugins/Animation/DeformerGraph/Source/OptimusEditor/Private/OptimusEditorModule.cpp
@@ -29,6 +29,7 @@
 #include "OptimusSourceAssetActions.h"
 #include "OptimusValidatedName.h"
 #include "OptimusValueContainer.h"
+#include "OptimusValueContainerStruct.h"
 #include "PropertyEditorModule.h"
 #include "UObject/FieldIterator.h"
 #include "UserDefinedStructure/UserDefinedStructEditorData.h"
@@ -206,7 +207,7 @@ void FOptimusEditorModule::RegisterPropertyCustomizations()
 	RegisterPropertyCustomization(FOptimusShaderText::StaticStruct()->GetFName(), &FOptimusShaderTextCustomization::MakeInstance);
 	RegisterPropertyCustomization(FOptimusParameterBinding::StaticStruct()->GetFName(), &FOptimusParameterBindingCustomization::MakeInstance);
 	RegisterPropertyCustomization(FOptimusParameterBindingArray::StaticStruct()->GetFName(), &FOptimusParameterBindingArrayCustomization::MakeInstance);
-	RegisterPropertyCustomization(UOptimusValueContainer::StaticClass()->GetFName(), &FOptimusValueContainerCustomization::MakeInstance);
+	RegisterPropertyCustomization(FOptimusValueContainerStruct::StaticStruct()->GetFName(), &FOptimusValueContainerStructCustomization::MakeInstance);
 	RegisterPropertyCustomization(FOptimusValidatedName::StaticStruct()->GetFName(), &FOptimusValidatedNameCustomization::MakeInstance);
 	RegisterPropertyCustomization(FOptimusDeformerInstanceComponentBinding::StaticStruct()->GetFName(), &FOptimusDeformerInstanceComponentBindingCustomization::MakeInstance);
 	
