[Deformer Graph Source Library] minor fix to scrollbar behavior. Still bad, but slightly better

diff --git a/Engine/Plugins/Animation/DeformerGraph/Source/OptimusEditor/Private/OptimusDetailsCustomization.cpp b/Engine/Plugins/Animation/DeformerGraph/Source/OptimusEditor/Private/OptimusDetailsCustomization.cpp
--- a/Engine/Plugins/Animation/DeformerGraph/Source/OptimusEditor/Private/OptimusDetailsCustomization.cpp
+++ b/Engine/Plugins/Animation/DeformerGraph/Source/OptimusEditor/Private/OptimusDetailsCustomization.cpp
@@ -1657,12 +1657,11 @@ void FOptimusSourceDetailsCustomization::CustomizeDetails(IDetailLayoutBuilder&
 	}
 
 	TSharedRef<IPropertyHandle> SourcePropertyHandle = DetailBuilder.GetProperty(TEXT("SourceText"));
+	DetailBuilder.GetDetailsView()->HideFilterArea(true);
 	DetailBuilder.EditDefaultProperty(SourcePropertyHandle)->CustomWidget()
-	.WholeRowContent()
 	[
-		SNew(SVerticalBox)
-		+ SVerticalBox::Slot()
-		.FillHeight(1.0f)
+		SNew(SBox)
+		.HeightOverride(500)
 		[
 			SAssignNew(SourceTextBox, SOptimusShaderTextDocumentTextBox)
 			.Text(this, &FOptimusSourceDetailsCustomization::GetText)
