[Deformer Graph] Added guardrails for invalid shared pointer. Hopefully avoids hard crash and provide useful info in case we hit the crash again

diff --git a/Engine/Plugins/Animation/DeformerGraph/Source/OptimusCore/Private/OptimusDataTypeRegistry.cpp b/Engine/Plugins/Animation/DeformerGraph/Source/OptimusCore/Private/OptimusDataTypeRegistry.cpp
--- a/Engine/Plugins/Animation/DeformerGraph/Source/OptimusCore/Private/OptimusDataTypeRegistry.cpp
+++ b/Engine/Plugins/Animation/DeformerGraph/Source/OptimusCore/Private/OptimusDataTypeRegistry.cpp
@@ -73,9 +73,14 @@ static bool ConvertPropertyValuePOD(
 static FOptimusDataTypeHandle GetArrayElementDataTypeForStructuredBuffer(FOptimusDataTypeHandle InDataType)
 {
 	FOptimusDataTypeHandle Result = InDataType;
-	if (InDataType->ShaderValueType->Type != EShaderFundamentalType::Struct && InDataType->ShaderValueType->VectorElemCount == 3)
+	if (InDataType->ShaderValueType->Type != EShaderFundamentalType::Struct && InDataType->ShaderValueType->DimensionType == EShaderFundamentalDimensionType::Vector && InDataType->ShaderValueType->VectorElemCount == 3)
 	{
 		Result = FOptimusDataTypeRegistry::Get().FindType(FShaderValueType::Get(InDataType->ShaderValueType->Type, 4));
+		
+		if (!ensureMsgf(Result.IsValid(), TEXT("Cannot find element type for structured buffer: %s"), *(InDataType->TypeName.ToString())))
+		{
+			return InDataType;
+		}
 	}
 
 	return Result;
@@ -478,7 +483,11 @@ void FOptimusDataTypeRegistry::RegisterBuiltinTypes()
 	
 	for (const TPair<FName, FTypeInfo>& Type : AlreadyRegisteredTypes)
 	{
-		Registry.RegisterArrayTypeIfApplicable(Type.Value.DataType);
+		// This should never happen theoretically, but somehow Linux complained about it once, hope this helps the next time it complains
+		if (ensureMsgf(Type.Value.DataType.IsValid(), TEXT("Invalid data type for array type registration: %s"), *Type.Key.ToString()))
+		{
+			Registry.RegisterArrayTypeIfApplicable(Type.Value.DataType);
+		}
 	}
 }
 
@@ -845,6 +854,7 @@ bool FOptimusDataTypeRegistry::RegisterStructType(UScriptStruct* InStructType)
 
 bool FOptimusDataTypeRegistry::RegisterArrayTypeIfApplicable(FOptimusDataTypeHandle InElementDataType)
 {
+	check(InElementDataType.IsValid());
 	// For now only allow array type for variables
 	if (!EnumHasAnyFlags(InElementDataType->UsageFlags, EOptimusDataTypeUsageFlags::Variable))
 	{
@@ -868,6 +878,11 @@ bool FOptimusDataTypeRegistry::RegisterArrayTypeIfApplicable(FOptimusDataTypeHan
 
 	// Making sure we are copying property value of a type into shader value of a equal or larger type, see comment for GetArrayElementDataTypeForStructuredBuffer
 	FOptimusDataTypeHandle InnerDataTypeForStructuredBuffer = GetArrayElementDataTypeForStructuredBuffer(InElementDataType);
+	// This should never happen theoretically, but somehow Linux complained about it, hope this helps the next time it complains
+	if (!ensureMsgf(InnerDataTypeForStructuredBuffer.IsValid(), TEXT("Cannot find matching element type for array type registration: %s"), *(InElementDataType->TypeName.ToString())))
+	{
+		return false;
+	}
 	check(InnerDataTypeForStructuredBuffer->ShaderValueSize >= InElementDataType->ShaderValueSize);
 	PropertyValueConvertFuncT ElementPropertyValueConvertFunc = FindPropertyValueConvertFunc(InnerDataTypeForStructuredBuffer->TypeName);
 	check(ElementPropertyValueConvertFunc);
