[Deformer Graph] Added support for dynamic pins toUOptimusNode_DataInterface
- deprecated UOptimusNode_AnimAttributeDataInterface in favor of UOptimusNode_DataInterface

diff --git a/Engine/Plugins/Animation/DeformerGraph/Source/OptimusCore/Private/DataInterfaces/OptimusDataInterfaceAnimAttribute.cpp b/Engine/Plugins/Animation/DeformerGraph/Source/OptimusCore/Private/DataInterfaces/OptimusDataInterfaceAnimAttribute.cpp
--- a/Engine/Plugins/Animation/DeformerGraph/Source/OptimusCore/Private/DataInterfaces/OptimusDataInterfaceAnimAttribute.cpp
+++ b/Engine/Plugins/Animation/DeformerGraph/Source/OptimusCore/Private/DataInterfaces/OptimusDataInterfaceAnimAttribute.cpp
@@ -17,6 +17,7 @@
 #include "ShaderParameterMetadataBuilder.h"
 #include "ComputeFramework/ComputeMetadataBuilder.h"
 #include "Engine/UserDefinedStruct.h"
+#include "Nodes/OptimusNode_DataInterface.h"
 
 #include "Serialization/MemoryReader.h"
 #include "Serialization/MemoryWriter.h"
@@ -104,6 +105,8 @@ void UOptimusAnimAttributeDataInterface::PostEditChangeChainProperty(FPropertyCh
 			if(ensure(AttributeArray.IsValidIndex(ChangedIndex)))
 			{
 				FOptimusAnimAttributeDescription& ChangedAttribute = AttributeArray[ChangedIndex];
+				
+				FName OldPinName = ChangedAttribute.PinName;
 
 				if (ChangedAttribute.Name.IsEmpty())
 				{
@@ -126,6 +129,7 @@ void UOptimusAnimAttributeDataInterface::PostEditChangeChainProperty(FPropertyCh
 				}
 
 				UpdateAttributePinNamesAndHlslIds();
+				OnPinDefinitionRenamedDelegate.Execute(OldPinName, ChangedAttribute.PinName);
 			}
 		}
 
@@ -137,28 +141,33 @@ void UOptimusAnimAttributeDataInterface::PostEditChangeChainProperty(FPropertyCh
 			ChangedAttribute.DefaultValueStruct.SetType(ChangedAttribute.DataType);
 		}
 	}
-	else if (PropertyChangedEvent.ChangeType == EPropertyChangeType::ArrayAdd)
+	else 
 	{
-		if (PropertyName == GET_MEMBER_NAME_STRING_CHECKED(FOptimusAnimAttributeArray, InnerArray))
+		if (PropertyChangedEvent.ChangeType == EPropertyChangeType::ArrayAdd)
 		{
-			const int32 ChangedIndex = PropertyChangedEvent.GetArrayIndex(GET_MEMBER_NAME_STRING_CHECKED(FOptimusAnimAttributeArray, InnerArray));
-			FOptimusAnimAttributeDescription& Attribute = AttributeArray[ChangedIndex];
+			if (PropertyName == GET_MEMBER_NAME_STRING_CHECKED(FOptimusAnimAttributeArray, InnerArray))
+			{
+				const int32 ChangedIndex = PropertyChangedEvent.GetArrayIndex(GET_MEMBER_NAME_STRING_CHECKED(FOptimusAnimAttributeArray, InnerArray));
+				FOptimusAnimAttributeDescription& Attribute = AttributeArray[ChangedIndex];
 			
-			// Default to a float attribute
-			Attribute.Init(GetUnusedAttributeName(TEXT("EmptyName")), NAME_None,
-				FOptimusDataTypeRegistry::Get().FindType(*FFloatProperty::StaticClass()));
+				// Default to a float attribute
+				Attribute.Init(GetUnusedAttributeName(TEXT("EmptyName")), NAME_None,
+					FOptimusDataTypeRegistry::Get().FindType(*FFloatProperty::StaticClass()));
+			}
 		}
-	}
-	else if (PropertyChangedEvent.ChangeType == EPropertyChangeType::Duplicate)
-	{
-		if (PropertyName == GET_MEMBER_NAME_STRING_CHECKED(FOptimusAnimAttributeArray, InnerArray))
-		{	
-			const int32 ChangedIndex = PropertyChangedEvent.GetArrayIndex(GET_MEMBER_NAME_STRING_CHECKED(FOptimusAnimAttributeArray, InnerArray));
-			FOptimusAnimAttributeDescription& Attribute = AttributeArray[ChangedIndex];
+		else if (PropertyChangedEvent.ChangeType == EPropertyChangeType::Duplicate)
+		{
+			if (PropertyName == GET_MEMBER_NAME_STRING_CHECKED(FOptimusAnimAttributeArray, InnerArray))
+			{	
+				const int32 ChangedIndex = PropertyChangedEvent.GetArrayIndex(GET_MEMBER_NAME_STRING_CHECKED(FOptimusAnimAttributeArray, InnerArray));
+				FOptimusAnimAttributeDescription& Attribute = AttributeArray[ChangedIndex];
 			
-			Attribute.Name = GetUnusedAttributeName(Attribute.Name);
-			Attribute.UpdatePinNameAndHlslId();
+				Attribute.Name = GetUnusedAttributeName(Attribute.Name);
+				Attribute.UpdatePinNameAndHlslId();
+			}
 		}
+		
+		OnPinDefinitionChangedDelegate.Execute();
 	}
 }
 #endif
@@ -217,6 +226,18 @@ TSubclassOf<UActorComponent> UOptimusAnimAttributeDataInterface::GetRequiredComp
 	return USkeletalMeshComponent::StaticClass();
 }
 
+void UOptimusAnimAttributeDataInterface::Initialize()
+{
+	AddAnimAttribute(TEXT("EmptyName"), NAME_None,
+			FOptimusDataTypeRegistry::Get().FindType(*FFloatProperty::StaticClass()) );
+}
+
+void UOptimusAnimAttributeDataInterface::RegisterDynamicPinDelegatesForOwningNode(UOptimusNode_DataInterface* InNode)
+{
+	OnPinDefinitionChangedDelegate.BindUObject(InNode, &UOptimusNode_DataInterface::RecreatePinsFromPinDefinitions);
+	OnPinDefinitionRenamedDelegate.BindUObject(InNode, &UOptimusNode_DataInterface::RenamePinFromPinDefinition);
+}
+
 
 void UOptimusAnimAttributeDataInterface::GetSupportedInputs(TArray<FShaderFunctionDefinition>& OutFunctions) const
 {
@@ -422,11 +443,11 @@ const FOptimusAnimAttributeDescription& UOptimusAnimAttributeDataInterface::AddA
 }
 
 
-void UOptimusAnimAttributeDataInterface::OnDataTypeChanged(FName InDataType)
+void UOptimusAnimAttributeDataInterface::OnDataTypeChanged(FName InTypeName)
 {
 	for (FOptimusAnimAttributeDescription& AttributeDescription : AttributeArray)
 	{
-		if (AttributeDescription.DataType.TypeName == InDataType)
+		if (AttributeDescription.DataType.TypeName == InTypeName)
 		{
 			AttributeDescription.DefaultValueStruct.SetType(AttributeDescription.DataType);
 		}
@@ -518,6 +539,7 @@ FOptimusAnimAttributeRuntimeData::FOptimusAnimAttributeRuntimeData(
 	const FOptimusAnimAttributeDescription& InDescription)
 {
 	Name = *InDescription.Name;
+	HlslId = *InDescription.HlslId;
 	BoneName = InDescription.BoneName;
 	CachedBoneIndex = 0;
 	
@@ -537,7 +559,7 @@ FOptimusAnimAttributeRuntimeData::FOptimusAnimAttributeRuntimeData(
 
 void UOptimusAnimAttributeDataProvider::Init(
 	USkeletalMeshComponent* InSkeletalMesh,
-	TArray<FOptimusAnimAttributeDescription> InAttributeArray
+	const TArray<FOptimusAnimAttributeDescription>& InAttributeArray
 	)
 {
 	SkeletalMesh = InSkeletalMesh;
@@ -571,9 +593,9 @@ void UOptimusAnimAttributeDataProvider::Init(
 
 	TArray<FShaderParametersMetadata*> AllocatedMetadatas;
 	TArray<FShaderParametersMetadata*> NestedStructs;
-	for (FOptimusAnimAttributeDescription& Attribute : InAttributeArray)
+	for (const FOptimusAnimAttributeDescription& Attribute : InAttributeArray)
 	{
-		ComputeFramework::AddParamForType(Builder, *Attribute.Name, Attribute.DataType->ShaderValueType, NestedStructs);
+		ComputeFramework::AddParamForType(Builder, *Attribute.HlslId, Attribute.DataType->ShaderValueType, NestedStructs);
 	}
 
 	FShaderParametersMetadata* ShaderParameterMetadata = Builder.Build(FShaderParametersMetadata::EUseCase::ShaderParameterStruct, TEXT("UAnimAttributeDataInterface"));
@@ -585,7 +607,7 @@ void UOptimusAnimAttributeDataProvider::Init(
 	for (int32 Index = 0; Index < AttributeRuntimeData.Num(); ++Index)
 	{
 		FOptimusAnimAttributeRuntimeData& RuntimeData = AttributeRuntimeData[Index];
-		check(RuntimeData.Name == Members[Index].GetName());
+		check(RuntimeData.HlslId == Members[Index].GetName());
 		
 		RuntimeData.Offset = Members[Index].GetOffset();
 
diff --git a/Engine/Plugins/Animation/DeformerGraph/Source/OptimusCore/Private/DataInterfaces/OptimusDataInterfaceAnimAttribute.h b/Engine/Plugins/Animation/DeformerGraph/Source/OptimusCore/Private/DataInterfaces/OptimusDataInterfaceAnimAttribute.h
--- a/Engine/Plugins/Animation/DeformerGraph/Source/OptimusCore/Private/DataInterfaces/OptimusDataInterfaceAnimAttribute.h
+++ b/Engine/Plugins/Animation/DeformerGraph/Source/OptimusCore/Private/DataInterfaces/OptimusDataInterfaceAnimAttribute.h
@@ -125,6 +125,11 @@ public:
 	FString GetDisplayName() const override;
 	virtual TArray<FOptimusCDIPinDefinition> GetPinDefinitions() const override;
 	TSubclassOf<UActorComponent> GetRequiredComponentClass() const override;
+	
+	void Initialize() override;
+	bool CanPinDefinitionChange() override {return true;};
+	void RegisterDynamicPinDelegatesForOwningNode(UOptimusNode_DataInterface* InNode) override;
+	
 	//~ End UOptimusComputeDataInterface Interface
 	
 	//~ Begin UComputeDataInterface Interface
@@ -140,7 +145,7 @@ public:
 
 	const FOptimusAnimAttributeDescription& AddAnimAttribute(const FString& InName, FName InBoneName, const FOptimusDataTypeRef& InDataType);
 
-	void OnDataTypeChanged(FName InDataType);
+	void OnDataTypeChanged(FName InTypeName) override;
 	
 	UPROPERTY(EditAnywhere, Category = "Animation Attribute", meta = (ShowOnlyInnerProperties))
 	FOptimusAnimAttributeArray AttributeArray;
@@ -148,6 +153,9 @@ public:
 private:
 	FString GetUnusedAttributeName(const FString& InName) const;
 	void UpdateAttributePinNamesAndHlslIds();
+	
+	FOnPinDefinitionChanged OnPinDefinitionChangedDelegate;
+	FOnPinDefinitionRenamed OnPinDefinitionRenamedDelegate;
 };
 
 // Runtime data with cached values baked out from AttributeDescription
@@ -158,6 +166,7 @@ struct FOptimusAnimAttributeRuntimeData
 	FOptimusAnimAttributeRuntimeData(const FOptimusAnimAttributeDescription& InDescription);
 	
 	FName Name;
+	FName HlslId;
 
 	FName BoneName;
 
@@ -189,7 +198,7 @@ public:
 	
 	void Init(
 		USkeletalMeshComponent* InSkeletalMesh,
-		TArray<FOptimusAnimAttributeDescription> InAttributeArray
+		const TArray<FOptimusAnimAttributeDescription>& InAttributeArray
 	);
 	
 	UPROPERTY(EditAnywhere, BlueprintReadWrite, Category = Binding)
diff --git a/Engine/Plugins/Animation/DeformerGraph/Source/OptimusCore/Private/Nodes/OptimusNode_AnimAttributeDataInterface.h b/Engine/Plugins/Animation/DeformerGraph/Source/OptimusCore/Private/Nodes/OptimusNode_AnimAttributeDataInterface.h
--- a/Engine/Plugins/Animation/DeformerGraph/Source/OptimusCore/Private/Nodes/OptimusNode_AnimAttributeDataInterface.h
+++ b/Engine/Plugins/Animation/DeformerGraph/Source/OptimusCore/Private/Nodes/OptimusNode_AnimAttributeDataInterface.h
@@ -6,6 +6,7 @@
 
 #include "OptimusNode_AnimAttributeDataInterface.generated.h"
 
+// Deprecated
 UCLASS(Hidden)
 class UOptimusNode_AnimAttributeDataInterface :
 	public UOptimusNode_DataInterface
diff --git a/Engine/Plugins/Animation/DeformerGraph/Source/OptimusCore/Private/Nodes/OptimusNode_ConstantValue.cpp b/Engine/Plugins/Animation/DeformerGraph/Source/OptimusCore/Private/Nodes/OptimusNode_ConstantValue.cpp
--- a/Engine/Plugins/Animation/DeformerGraph/Source/OptimusCore/Private/Nodes/OptimusNode_ConstantValue.cpp
+++ b/Engine/Plugins/Animation/DeformerGraph/Source/OptimusCore/Private/Nodes/OptimusNode_ConstantValue.cpp
@@ -71,9 +71,9 @@ UClass* UOptimusNode_ConstantValueGeneratorClass::GetClassForType(UPackage* InPa
 }
 
 
-void UOptimusNode_ConstantValue::PostLoad()
+void UOptimusNode_ConstantValue::PostLoadNodeSpecificData()
 {
-	Super::PostLoad();
+	Super::PostLoadNodeSpecificData();
 
 	if (!GetClass()->GetOuter()->IsA<UPackage>())
 	{
diff --git a/Engine/Plugins/Animation/DeformerGraph/Source/OptimusCore/Private/Nodes/OptimusNode_ConstantValue.h b/Engine/Plugins/Animation/DeformerGraph/Source/OptimusCore/Private/Nodes/OptimusNode_ConstantValue.h
--- a/Engine/Plugins/Animation/DeformerGraph/Source/OptimusCore/Private/Nodes/OptimusNode_ConstantValue.h
+++ b/Engine/Plugins/Animation/DeformerGraph/Source/OptimusCore/Private/Nodes/OptimusNode_ConstantValue.h
@@ -57,8 +57,7 @@ public:
 	}
 
 
-	// UObject overrides
-	void PostLoad() override;
+	void PostLoadNodeSpecificData() override;
 
 #if WITH_EDITOR
 	void PostEditChangeProperty(FPropertyChangedEvent& PropertyChangedEvent) override;
diff --git a/Engine/Plugins/Animation/DeformerGraph/Source/OptimusCore/Private/Nodes/OptimusNode_CustomComputeKernel.cpp b/Engine/Plugins/Animation/DeformerGraph/Source/OptimusCore/Private/Nodes/OptimusNode_CustomComputeKernel.cpp
--- a/Engine/Plugins/Animation/DeformerGraph/Source/OptimusCore/Private/Nodes/OptimusNode_CustomComputeKernel.cpp
+++ b/Engine/Plugins/Animation/DeformerGraph/Source/OptimusCore/Private/Nodes/OptimusNode_CustomComputeKernel.cpp
@@ -678,6 +678,52 @@
 	UpdatePreamble();
 }
 
+void UOptimusNode_CustomComputeKernel::PostLoadNodeSpecificData()
+{
+	Super::PostLoadNodeSpecificData();
 
+	if (GetLinkerCustomVersion(FOptimusObjectVersion::GUID) < FOptimusObjectVersion::SwitchToParameterBindingArrayStruct)
+	{
+		Modify();
+		InputBindingArray.InnerArray = InputBindings_DEPRECATED;
+		OutputBindingArray.InnerArray = OutputBindings_DEPRECATED;
+	}
+	
+	if (!Parameters_DEPRECATED.IsEmpty())
+	{
+PRAGMA_DISABLE_DEPRECATION_WARNINGS		
+		TArray<FOptimusParameterBinding> ParameterInputBindings;
+
+		for (const FOptimus_ShaderBinding& OldBinding: Parameters_DEPRECATED)
+		{
+			FOptimusParameterBinding NewBinding;
+			NewBinding.Name = OldBinding.Name;
+			NewBinding.DataType = OldBinding.DataType;
+			NewBinding.DataDomain = FOptimusDataDomain();
+			
+			ParameterInputBindings.Add(NewBinding);
+		}
+		
+		InputBindingArray.InnerArray.Insert(ParameterInputBindings, 0);
+		
+PRAGMA_ENABLE_DEPRECATION_WARNINGS
+	}
+	
+	
+	if (GetLinkerCustomVersion(FOptimusObjectVersion::GUID) < FOptimusObjectVersion::KernelDataInterface)
+	{
+		PostLoadExtractExecutionDomain();
+		PostLoadAddMissingPrimaryGroupPin();
+	}
+
+	if (GetLinkerCustomVersion(FOptimusObjectVersion::GUID) < FOptimusObjectVersion::KernelParameterBindingToggleAtomic)
+	{
+		PostLoadExtractAtomicModeFromConnectedResource();
+	}
+
+	SetDisplayName(FText::FromName(KernelName));	
+}
+
+
 void UOptimusNode_CustomComputeKernel::Serialize(FArchive& Ar)
 {
@@ -1494,51 +1540,6 @@
 
 #endif
 
-void UOptimusNode_CustomComputeKernel::PostLoad()
-{
-	if (GetLinkerCustomVersion(FOptimusObjectVersion::GUID) < FOptimusObjectVersion::SwitchToParameterBindingArrayStruct)
-	{
-		Modify();
-		InputBindingArray.InnerArray = InputBindings_DEPRECATED;
-		OutputBindingArray.InnerArray = OutputBindings_DEPRECATED;
-	}
-	
-	if (!Parameters_DEPRECATED.IsEmpty())
-	{
-PRAGMA_DISABLE_DEPRECATION_WARNINGS		
-		TArray<FOptimusParameterBinding> ParameterInputBindings;
 
-		for (const FOptimus_ShaderBinding& OldBinding: Parameters_DEPRECATED)
-		{
-			FOptimusParameterBinding NewBinding;
-			NewBinding.Name = OldBinding.Name;
-			NewBinding.DataType = OldBinding.DataType;
-			NewBinding.DataDomain = FOptimusDataDomain();
-			
-			ParameterInputBindings.Add(NewBinding);
-		}
-		
-		InputBindingArray.InnerArray.Insert(ParameterInputBindings, 0);
-		
-PRAGMA_ENABLE_DEPRECATION_WARNINGS
-	}
-	
-	Super::PostLoad();
-	
-	if (GetLinkerCustomVersion(FOptimusObjectVersion::GUID) < FOptimusObjectVersion::KernelDataInterface)
-	{
-		PostLoadExtractExecutionDomain();
-		PostLoadAddMissingPrimaryGroupPin();
-	}
-
-	if (GetLinkerCustomVersion(FOptimusObjectVersion::GUID) < FOptimusObjectVersion::KernelParameterBindingToggleAtomic)
-	{
-		PostLoadExtractAtomicModeFromConnectedResource();
-	}
-
-	SetDisplayName(FText::FromName(KernelName));
-}
-
-
 void UOptimusNode_CustomComputeKernel::ConstructNode()
 {
diff --git a/Engine/Plugins/Animation/DeformerGraph/Source/OptimusCore/Private/Nodes/OptimusNode_CustomComputeKernel.h b/Engine/Plugins/Animation/DeformerGraph/Source/OptimusCore/Private/Nodes/OptimusNode_CustomComputeKernel.h
--- a/Engine/Plugins/Animation/DeformerGraph/Source/OptimusCore/Private/Nodes/OptimusNode_CustomComputeKernel.h
+++ b/Engine/Plugins/Animation/DeformerGraph/Source/OptimusCore/Private/Nodes/OptimusNode_CustomComputeKernel.h
@@ -51,6 +51,7 @@ public:
 	}
 
 	void OnDataTypeChanged(FName InTypeName) override;
+	void PostLoadNodeSpecificData() override;
 
 	// UOptimusNode_ComputeKernelBase overrides
 	FString GetKernelName() const override { return KernelName.ToString(); }
@@ -168,8 +169,6 @@ PRAGMA_ENABLE_DEPRECATION_WARNINGS
 #endif
 
 	void Serialize(FArchive& Ar) override;
-	void PostLoad() override;
-
 	
 	
 protected:
diff --git a/Engine/Plugins/Animation/DeformerGraph/Source/OptimusCore/Private/Nodes/OptimusNode_DataInterface.cpp b/Engine/Plugins/Animation/DeformerGraph/Source/OptimusCore/Private/Nodes/OptimusNode_DataInterface.cpp
--- a/Engine/Plugins/Animation/DeformerGraph/Source/OptimusCore/Private/Nodes/OptimusNode_DataInterface.cpp
+++ b/Engine/Plugins/Animation/DeformerGraph/Source/OptimusCore/Private/Nodes/OptimusNode_DataInterface.cpp
@@ -95,6 +95,30 @@ TOptional<FText> UOptimusNode_DataInterface::ValidateForCompile(const FOptimusPi
 	return {};
 }
 
+void UOptimusNode_DataInterface::PostLoadNodeSpecificData()
+{
+	Super::PostLoadNodeSpecificData();
+	
+	// Previously DataInterfaceData wasn't always created.
+	if (DataInterfaceClass && !DataInterfaceData)
+	{
+		DataInterfaceData = NewObject<UOptimusComputeDataInterface>(this, DataInterfaceClass);
+	}
+
+	// Add in the component pin.
+	if (GetLinkerCustomVersion(FOptimusObjectVersion::GUID) < FOptimusObjectVersion::ComponentProviderSupport)
+	{
+		CreateComponentPin();
+	}
+}
+
+void UOptimusNode_DataInterface::OnDataTypeChanged(FName InTypeName)
+{
+	Super::OnDataTypeChanged(InTypeName);
+
+	DataInterfaceData->OnDataTypeChanged(InTypeName);
+}
+
 void UOptimusNode_DataInterface::SaveState(FArchive& Ar) const
 {
 	Super::SaveState(Ar);
@@ -121,6 +145,60 @@ bool UOptimusNode_DataInterface::IsComponentSourceCompatible(const UOptimusCompo
 	return InComponentSource && InComponentSource->GetComponentClass()->IsChildOf(DataInterfaceData->GetRequiredComponentClass());
 }
 
+void UOptimusNode_DataInterface::RecreatePinsFromPinDefinitions()
+{
+	// Recreate all the pins
+	// Save the links and readd them later when new pins are created
+	TMap<FName, TArray<UOptimusNodePin*>> ConnectedPinsMap;
+
+	TArray<UOptimusNodePin*> PinsToRemove;
+	
+	for (UOptimusNodePin* Pin : GetPins())
+	{
+		if (Pin != GetComponentPin())
+		{
+			ConnectedPinsMap.Add(Pin->GetFName()) = Pin->GetConnectedPins();
+			PinsToRemove.Add(Pin);
+		}
+	}	
+
+	for (UOptimusNodePin* Pin : PinsToRemove)
+	{
+		RemovePin(Pin);
+	}
+
+	CreatePinsFromDataInterface(DataInterfaceData, true);
+
+	for (UOptimusNodePin* Pin : GetPins())
+	{
+		if (TArray<UOptimusNodePin*>* ConnectedPins = ConnectedPinsMap.Find(Pin->GetFName()))
+		{
+			for (UOptimusNodePin* ConnectedPin : *ConnectedPins)
+			{
+				GetOwningGraph()->AddLink(Pin, ConnectedPin);
+			}
+		}	
+	}	
+}
+
+void UOptimusNode_DataInterface::RenamePinFromPinDefinition(FName InOld, FName InNew)
+{
+	UOptimusNodePin* Pin = FindPin(InOld.ToString());
+	SetPinName(Pin, InNew);
+}
+
+void UOptimusNode_DataInterface::InitializeTransientData()
+{
+	if (ensure(DataInterfaceData))
+	{
+		if (DataInterfaceData->CanPinDefinitionChange())
+		{
+			EnableDynamicPins();
+			DataInterfaceData->RegisterDynamicPinDelegatesForOwningNode(this);
+		}
+	}
+}
+
 void UOptimusNode_DataInterface::Serialize(FArchive& Ar)
 {
  	Super::Serialize(Ar);
@@ -198,6 +276,9 @@ void UOptimusNode_DataInterface::SetDataInterfaceClass(
 {
 	DataInterfaceClass = InDataInterfaceClass;
 	DataInterfaceData = NewObject<UOptimusComputeDataInterface>(this, DataInterfaceClass);
+	// Undo support
+	DataInterfaceData->SetFlags(RF_Transactional);
+	DataInterfaceData->Initialize();
 }
 
 UOptimusComponentSourceBinding* UOptimusNode_DataInterface::GetComponentBinding(const FOptimusPinTraversalContext& InContext) const
@@ -254,24 +335,6 @@ EOptimusPinMutability UOptimusNode_DataInterface::GetOutputPinMutability(const U
 }
 
 
-void UOptimusNode_DataInterface::PostLoad() 
-{
-	Super::PostLoad();
-
-	// Previously DataInterfaceData wasn't always created.
-	if (DataInterfaceClass && !DataInterfaceData)
-	{
-		DataInterfaceData = NewObject<UOptimusComputeDataInterface>(this, DataInterfaceClass);
-	}
-
-	// Add in the component pin.
-	if (GetLinkerCustomVersion(FOptimusObjectVersion::GUID) < FOptimusObjectVersion::ComponentProviderSupport)
-	{
-		CreateComponentPin();
-	}
-}
-
-
 void UOptimusNode_DataInterface::ConstructNode()
 {
 	// Create the component pin.
@@ -283,7 +346,7 @@ void UOptimusNode_DataInterface::ConstructNode()
 		}
 		SetDisplayName(FText::FromString(DataInterfaceData->GetDisplayName()));
 		CreateComponentPin();
-		CreatePinsFromDataInterface(DataInterfaceData);
+		CreatePinsFromDataInterface(DataInterfaceData, false);
 	}
 }
 
@@ -301,9 +364,7 @@ void UOptimusNode_DataInterface::PostDuplicate(EDuplicateMode::Type DuplicateMod
 }
 
 
-void UOptimusNode_DataInterface::CreatePinsFromDataInterface(
-	const UOptimusComputeDataInterface* InDataInterface
-	)
+void UOptimusNode_DataInterface::CreatePinsFromDataInterface(const UOptimusComputeDataInterface* InDataInterface, bool bSupportUndo)
 {
 	// A data interface provides read and write functions. A data interface node exposes
 	// the read functions as output pins to be fed into kernel nodes (or into other interface
@@ -333,16 +394,12 @@ void UOptimusNode_DataInterface::CreatePinsFromDataInterface(
 	{
 		if (ensure(!Def.PinName.IsNone()))
 		{
-			CreatePinFromDefinition(Def, ReadFunctionMap, WriteFunctionMap);
+			CreatePinFromDefinition(Def, ReadFunctionMap, WriteFunctionMap, bSupportUndo);
 		}
 	}
 }
 
-void UOptimusNode_DataInterface::CreatePinFromDefinition(
-	const FOptimusCDIPinDefinition& InDefinition,
-	const TMap<FString, const FShaderFunctionDefinition*>& InReadFunctionMap,
-	const TMap<FString, const FShaderFunctionDefinition*>& InWriteFunctionMap
-	)	
+void UOptimusNode_DataInterface::CreatePinFromDefinition(const FOptimusCDIPinDefinition& InDefinition, const TMap<FString, const FShaderFunctionDefinition*>& InReadFunctionMap, const TMap<FString, const FShaderFunctionDefinition*>& InWriteFunctionMap, bool bSupportUndo)
 {
 	const FOptimusDataTypeRegistry& TypeRegistry = FOptimusDataTypeRegistry::Get();
 
@@ -376,7 +433,14 @@ void UOptimusNode_DataInterface::CreatePinFromDefinition(
 			return;
 		}
 
-		AddPinDirect(InDefinition.PinName, EOptimusNodePinDirection::Output, {}, PinDataType);
+		if (bSupportUndo)
+		{
+			AddPin(InDefinition.PinName, EOptimusNodePinDirection::Output, {}, PinDataType);
+		}
+		else
+		{
+			AddPinDirect(InDefinition.PinName, EOptimusNodePinDirection::Output, {}, PinDataType);
+		}
 	}
 	else if (!InDefinition.DataFunctionName.IsEmpty())
 	{
@@ -450,7 +514,14 @@ void UOptimusNode_DataInterface::CreatePinFromDefinition(
 		}
 
 		const FOptimusDataDomain DataDomain{ContextNames, InDefinition.DomainMultiplier};
-		AddPinDirect(InDefinition.PinName, PinDirection, DataDomain, PinDataType);
+		if (bSupportUndo)
+		{
+			AddPin(InDefinition.PinName, PinDirection, DataDomain, PinDataType);	
+		}
+		else
+		{
+			AddPinDirect(InDefinition.PinName, PinDirection, DataDomain, PinDataType);
+		}
 	}
 	else
 	{
diff --git a/Engine/Plugins/Animation/DeformerGraph/Source/OptimusCore/Private/Nodes/OptimusNode_DataInterface.h b/Engine/Plugins/Animation/DeformerGraph/Source/OptimusCore/Private/Nodes/OptimusNode_DataInterface.h
--- a/Engine/Plugins/Animation/DeformerGraph/Source/OptimusCore/Private/Nodes/OptimusNode_DataInterface.h
+++ b/Engine/Plugins/Animation/DeformerGraph/Source/OptimusCore/Private/Nodes/OptimusNode_DataInterface.h
@@ -31,11 +31,15 @@ public:
 	virtual void SetDataInterfaceClass(TSubclassOf<UOptimusComputeDataInterface> InDataInterfaceClass);
 	bool IsComponentSourceCompatible(const UOptimusComponentSource* InComponentSource) const;
 
+	void RecreatePinsFromPinDefinitions();
+	void RenamePinFromPinDefinition(FName InOld, FName InNew);
+
 	// -- UOptimusNode overrides
 	FName GetNodeCategory() const override 
 	{
 		return CategoryName::DataInterfaces;
 	}
+	void InitializeTransientData() override;
 
 	// -- UObject overrides
 	void Serialize(FArchive& Ar) override;
@@ -55,29 +59,31 @@ protected:
 	void ConstructNode() override;
 	bool ValidateConnection(const UOptimusNodePin& InThisNodesPin, const UOptimusNodePin& InOtherNodesPin, FString* OutReason) const override;
 	TOptional<FText> ValidateForCompile(const FOptimusPinTraversalContext& InContext) const override;
+	void PostLoadNodeSpecificData() override;
+	void OnDataTypeChanged(FName InTypeName) override;
 
 	void SaveState(FArchive& Ar) const override;
 	void RestoreState(FArchive& Ar) override;
 
 	// -- UObject overrides
-	void PostLoad() override;
 	void PostDuplicate(EDuplicateMode::Type DuplicateMode) override;
 	
 private:
-	void CreatePinsFromDataInterface(
-		const UOptimusComputeDataInterface *InDataInterface
-		);
+	void CreatePinsFromDataInterface(const UOptimusComputeDataInterface *InDataInterface, bool bSupportUndo);
 	
+
 	void CreatePinFromDefinition(
 		const FOptimusCDIPinDefinition &InDefinition,
 		const TMap<FString, const FShaderFunctionDefinition *>& InReadFunctionMap,
-		const TMap<FString, const FShaderFunctionDefinition *>& InWriteFunctionMap
+		const TMap<FString, const FShaderFunctionDefinition *>& InWriteFunctionMap,
+		bool bSupportUndo
 		);
 
 	void CreateComponentPin();
 
 protected:
 	friend class UOptimusDeformer;
+	friend class UOptimusNodeGraph;
 
 	/** Accessor for the deformer object to connect this node automatically on backcomp and to unlink it as well
 	 *  if a component binding changes the source type.
diff --git a/Engine/Plugins/Animation/DeformerGraph/Source/OptimusCore/Private/OptimusNode.cpp b/Engine/Plugins/Animation/DeformerGraph/Source/OptimusCore/Private/OptimusNode.cpp
--- a/Engine/Plugins/Animation/DeformerGraph/Source/OptimusCore/Private/OptimusNode.cpp
+++ b/Engine/Plugins/Animation/DeformerGraph/Source/OptimusCore/Private/OptimusNode.cpp
@@ -424,6 +424,15 @@ void UOptimusNode::PostLoad()
 {
 	Super::PostLoad();
 
+	// PostLoad fix up pass
+	PostLoadNodeSpecificData();
+
+	// Now that all serialized data is ready, let's generate any transient data
+	InitializeTransientData();
+}
+
+void UOptimusNode::PostLoadNodeSpecificData()
+{
 	// Earlier iterations didn't set this flag. 
 	SetFlags(RF_Transactional);
 
@@ -431,8 +440,6 @@ void UOptimusNode::PostLoad()
 	{
 		Pin->ConditionalPostLoad();
 	}
-	
-	InitializeTransientData();
 }
 
 #if WITH_EDITOR
diff --git a/Engine/Plugins/Animation/DeformerGraph/Source/OptimusCore/Private/OptimusNodeGraph.cpp b/Engine/Plugins/Animation/DeformerGraph/Source/OptimusCore/Private/OptimusNodeGraph.cpp
--- a/Engine/Plugins/Animation/DeformerGraph/Source/OptimusCore/Private/OptimusNodeGraph.cpp
+++ b/Engine/Plugins/Animation/DeformerGraph/Source/OptimusCore/Private/OptimusNodeGraph.cpp
@@ -172,6 +172,8 @@ void UOptimusNodeGraph::PostLoad()
 	{
 		Link->ConditionalPostLoad();
 	}
+
+	PostLoadReplaceAnimAttributeDataInterfaceNodeWithGenericDataInterfaceNode();
 }
 
 UOptimusNodeGraph* UOptimusNodeGraph::GetParentGraph() const
@@ -537,15 +539,6 @@ UOptimusNode* UOptimusNodeGraph::AddDataInterfaceNode(
 	const FVector2D& InPosition
 	)
 {
-	if (InDataInterfaceClass == UOptimusAnimAttributeDataInterface::StaticClass())
-	{
-		return AddNodeInternal(UOptimusNode_AnimAttributeDataInterface::StaticClass(), InPosition,
-			[InDataInterfaceClass](UOptimusNode *InNode)
-			{
-				Cast<UOptimusNode_AnimAttributeDataInterface>(InNode)->SetDataInterfaceClass(InDataInterfaceClass);		
-			});
-	}
-
 	return AddNodeInternal(UOptimusNode_DataInterface::StaticClass(), InPosition,
 		[InDataInterfaceClass](UOptimusNode *InNode)
 		{
@@ -3071,6 +3064,73 @@ FString UOptimusNodeGraph::ConstructSubGraphPath(const FString& InGraphOwnerPath
 	return SubGraphPath;	
 }
 
+void UOptimusNodeGraph::PostLoadReplaceAnimAttributeDataInterfaceNodeWithGenericDataInterfaceNode()
+{
+	TArray<UOptimusNode_AnimAttributeDataInterface*> DeprecatedAnimAtttributeNodes;
+	
+	for (UOptimusNode* Node : Nodes)
+	{
+		if (UOptimusNode_AnimAttributeDataInterface* DeprecatedAnimAtttributeNode = Cast<UOptimusNode_AnimAttributeDataInterface>(Node))
+		{
+			DeprecatedAnimAtttributeNodes.Add(DeprecatedAnimAtttributeNode);
+		}
+	}
+	
+	for (UOptimusNode_AnimAttributeDataInterface* Node : DeprecatedAnimAtttributeNodes)
+	{
+		// Save pin connections
+		TArrayView<UOptimusNodePin* const> Pins = Node->GetPins();
+
+		TMap<FName, TArray<UOptimusNodePin*>> ConnectedPinsMap;
+		
+		for (UOptimusNodePin* Pin : Pins)
+		{
+			ConnectedPinsMap.Add(Pin->GetFName(), Pin->GetConnectedPins());
+		}
+
+		FVector2D Position = Node->GetGraphPosition();
+		// Assign to the graph temporarily
+		UOptimusComputeDataInterface* NewDataInterfaceData = Node->GetDataInterface(this);
+		
+		RemoveNodeDirect(Node, false /* remove all links as well */);
+
+		auto BootstrapNodeFunc = [NewDataInterfaceData, Position](UOptimusNode* InNode) -> bool
+		{
+			UOptimusNode_DataInterface* DataInterfaceNode = Cast<UOptimusNode_DataInterface>(InNode);
+			DataInterfaceNode->SetDataInterfaceClass(UOptimusAnimAttributeDataInterface::StaticClass());
+
+			// Replace the data interface
+			Optimus::RenameObject(NewDataInterfaceData, nullptr, DataInterfaceNode);
+			DataInterfaceNode->DataInterfaceData = NewDataInterfaceData;
+			
+			InNode->SetGraphPositionDirect(Position);	
+			return true;
+		};
+		
+		FName NodeName = Optimus::GetUniqueNameForScope(this, UOptimusNode_DataInterface::StaticClass()->GetFName());	
+		UOptimusNode* NewNode = CreateNodeDirect(UOptimusNode_DataInterface::StaticClass(), NodeName, BootstrapNodeFunc);
+
+		// Recover links
+		for (UOptimusNodePin* Pin : NewNode->GetPins())
+		{
+			if (TArray<UOptimusNodePin*>* ConnectedPins = ConnectedPinsMap.Find(Pin->GetFName()))
+			{
+				for (UOptimusNodePin* ConnectedPin : *ConnectedPins)
+				{
+					if (Pin->GetDirection() == EOptimusNodePinDirection::Input)
+					{
+						AddLinkDirect(ConnectedPin, Pin);
+					}
+					else if (Pin->GetDirection() == EOptimusNodePinDirection::Output)
+					{
+						AddLinkDirect(Pin, ConnectedPin);
+					}
+				}
+			}
+		}
+	}	
+}
+
 
 UOptimusActionStack* UOptimusNodeGraph::GetActionStack() const
 {
diff --git a/Engine/Plugins/Animation/DeformerGraph/Source/OptimusCore/Public/OptimusComputeDataInterface.h b/Engine/Plugins/Animation/DeformerGraph/Source/OptimusCore/Public/OptimusComputeDataInterface.h
--- a/Engine/Plugins/Animation/DeformerGraph/Source/OptimusCore/Public/OptimusComputeDataInterface.h
+++ b/Engine/Plugins/Animation/DeformerGraph/Source/OptimusCore/Public/OptimusComputeDataInterface.h
@@ -10,6 +10,7 @@
 
 
 class UActorComponent;
+class UOptimusNode_DataInterface;
 
 
 struct FOptimusCDIPinDefinition
@@ -129,6 +130,9 @@ class OPTIMUSCORE_API UOptimusComputeDataInterface : public UComputeDataInterfac
 	GENERATED_BODY()
 	
 public:
+	DECLARE_DELEGATE_TwoParams(FOnPinDefinitionRenamed, FName /* Old */ , FName /* New */);
+	DECLARE_DELEGATE(FOnPinDefinitionChanged);
+	
 	struct OPTIMUSCORE_API CategoryName
 	{
 		static const FName DataInterfaces;
@@ -145,11 +149,22 @@ public:
 	/// Returns the list of pins that will map to the shader functions provided by this data interface.
 	virtual TArray<FOptimusCDIPinDefinition> GetPinDefinitions() const PURE_VIRTUAL(UOptimusComputeDataInterface::GetDisplayName, return {};)
 
+	/// Data interface can use this to set default values/pins
+	virtual void Initialize() {};
+
+	/// Whether the data interface allow users to add / remove pins
+	virtual bool CanPinDefinitionChange() { return false; };
+
+	/// Register delegates for data interface node to update its pins when the data interface changes
+	virtual void RegisterDynamicPinDelegatesForOwningNode(UOptimusNode_DataInterface* InNode) {};
+
 	/**
 	 * @return Returns the component type that this data interface operates on.
 	 */
 	virtual TSubclassOf<UActorComponent> GetRequiredComponentClass() const PURE_VIRTUAL(UOptimusComputeDataInterface::GetRequiredComponent, return nullptr;)
 
+	virtual void OnDataTypeChanged(FName InTypeName) {};
+	
 	/**
 	 * Register any additional data types provided by this data interface. 
 	 */
diff --git a/Engine/Plugins/Animation/DeformerGraph/Source/OptimusCore/Public/OptimusNode.h b/Engine/Plugins/Animation/DeformerGraph/Source/OptimusCore/Public/OptimusNode.h
--- a/Engine/Plugins/Animation/DeformerGraph/Source/OptimusCore/Public/OptimusNode.h
+++ b/Engine/Plugins/Animation/DeformerGraph/Source/OptimusCore/Public/OptimusNode.h
@@ -182,7 +182,13 @@ public:
 
 	//== UObject overrides
 	void Serialize(FArchive& Ar) override;
-	void PostLoad() override;
+	
+	// Using "final" here to make sure all derive nodes have InitializeTransientData() automatically called on them during PostLoad(),
+	// so please use PostLoadNodeSpecificData instead for any PostLoad fix-ups
+	void PostLoad() override final;
+
+	// Derived nodes should override this function for any PostLoad fix-ups
+	virtual void PostLoadNodeSpecificData();
 #if WITH_EDITOR
 	void PostEditChangeProperty(FPropertyChangedEvent& PropertyChangedEvent) override;
 #endif
diff --git a/Engine/Plugins/Animation/DeformerGraph/Source/OptimusCore/Public/OptimusNodeGraph.h b/Engine/Plugins/Animation/DeformerGraph/Source/OptimusCore/Public/OptimusNodeGraph.h
--- a/Engine/Plugins/Animation/DeformerGraph/Source/OptimusCore/Public/OptimusNodeGraph.h
+++ b/Engine/Plugins/Animation/DeformerGraph/Source/OptimusCore/Public/OptimusNodeGraph.h
@@ -570,6 +570,8 @@ private:
 
 	FString ConstructSubGraphPath(const FString& InSubGraphName) const;
 	static FString ConstructSubGraphPath(const FString& InGraphOwnerPath, const FString& InSubGraphName);
+
+	void PostLoadReplaceAnimAttributeDataInterfaceNodeWithGenericDataInterfaceNode();
 	
 	UPROPERTY(NonTransactional)
 	TArray<TObjectPtr<UOptimusNode>> Nodes;
