[Deformer] fixed pin type not changing when changing attribute data type change

diff --git a/Engine/Plugins/Animation/DeformerGraph/Source/OptimusCore/Private/DataInterfaces/OptimusDataInterfaceAnimAttribute.cpp b/Engine/Plugins/Animation/DeformerGraph/Source/OptimusCore/Private/DataInterfaces/OptimusDataInterfaceAnimAttribute.cpp
--- a/Engine/Plugins/Animation/DeformerGraph/Source/OptimusCore/Private/DataInterfaces/OptimusDataInterfaceAnimAttribute.cpp
+++ b/Engine/Plugins/Animation/DeformerGraph/Source/OptimusCore/Private/DataInterfaces/OptimusDataInterfaceAnimAttribute.cpp
@@ -139,6 +139,7 @@ void UOptimusAnimAttributeDataInterface::PostEditChangeChainProperty(FPropertyCh
 
 			// Update the default value container accordingly
 			ChangedAttribute.DefaultValueStruct.SetType(ChangedAttribute.DataType);
+			OnPinDefinitionChangedDelegate.Execute();
 		}
 	}
 	else 
