[Deformer Graph] Fixed two bugs for Mac

-Changed format of the UAV used for color buffer from BGRA to RGBA, as some mac does not support BGRA and defaults to RGBA. as a result we always do manual fetch swizzle in the shader as well.
-Adding node name as suffix to Kernel hlsl name to avoid anim attribute with the same name not getting its value on mac

#rnx

diff --git a/Engine/Plugins/Animation/DeformerGraph/Shaders/Private/DataInterfaceSkinnedMeshWrite.ush b/Engine/Plugins/Animation/DeformerGraph/Shaders/Private/DataInterfaceSkinnedMeshWrite.ush
--- a/Engine/Plugins/Animation/DeformerGraph/Shaders/Private/DataInterfaceSkinnedMeshWrite.ush
+++ b/Engine/Plugins/Animation/DeformerGraph/Shaders/Private/DataInterfaceSkinnedMeshWrite.ush
@@ -37,5 +37,5 @@ void WriteTangentZ_{DataInterfaceName}(uint VertexIndex, float4 TangentZ)
 void WriteColor_{DataInterfaceName}(uint VertexIndex, float4 Color)
 {
 	uint BufferIndex = {DataInterfaceName}_OutputStreamStart + VertexIndex;
-	{DataInterfaceName}_ColorBufferUAV[BufferIndex] = Color;
+	{DataInterfaceName}_ColorBufferUAV[BufferIndex] = Color FMANUALFETCH_COLOR_COMPONENT_SWIZZLE;
 }
diff --git a/Engine/Plugins/Animation/DeformerGraph/Source/OptimusCore/Private/DataInterfaces/OptimusDataInterfaceSkinnedMeshWrite.cpp b/Engine/Plugins/Animation/DeformerGraph/Source/OptimusCore/Private/DataInterfaces/OptimusDataInterfaceSkinnedMeshWrite.cpp
--- a/Engine/Plugins/Animation/DeformerGraph/Source/OptimusCore/Private/DataInterfaces/OptimusDataInterfaceSkinnedMeshWrite.cpp
+++ b/Engine/Plugins/Animation/DeformerGraph/Source/OptimusCore/Private/DataInterfaces/OptimusDataInterfaceSkinnedMeshWrite.cpp
@@ -209,11 +209,12 @@ void FOptimusSkinnedMeshWriteDataProviderProxy::AllocateResources(FRDGBuilder& G
 	if (OutputMask & 8)
 	{
 		ColorBuffer = FSkeletalMeshDeformerHelpers::AllocateVertexFactoryColorBuffer(GraphBuilder, SkeletalMeshObject, LodIndex, TEXT("OptimusSkinnedMeshColor"));
-		ColorBufferUAV = GraphBuilder.CreateUAV(ColorBuffer, PF_B8G8R8A8, ERDGUnorderedAccessViewFlags::SkipBarrier);
+		// using RGBA here and do a manual fetch swizzle in shader instead of BGRA directly because some Mac does not support it. See GMetalBufferFormats[PF_B8G8R8A8] 
+		ColorBufferUAV = GraphBuilder.CreateUAV(ColorBuffer, PF_R8G8B8A8, ERDGUnorderedAccessViewFlags::SkipBarrier);
 	}
 	else
 	{
-		ColorBufferUAV = GraphBuilder.CreateUAV(GraphBuilder.RegisterExternalBuffer(GWhiteVertexBufferWithRDG->Buffer), PF_B8G8R8A8);
+		ColorBufferUAV = GraphBuilder.CreateUAV(GraphBuilder.RegisterExternalBuffer(GWhiteVertexBufferWithRDG->Buffer), PF_R8G8B8A8);
 	}
 
 	FSkeletalMeshDeformerHelpers::UpdateVertexFactoryBufferOverrides(GraphBuilder.RHICmdList, SkeletalMeshObject, LodIndex);
diff --git a/Engine/Plugins/Animation/DeformerGraph/Source/OptimusCore/Private/Nodes/OptimusNode_ComputeKernelBase.cpp b/Engine/Plugins/Animation/DeformerGraph/Source/OptimusCore/Private/Nodes/OptimusNode_ComputeKernelBase.cpp
--- a/Engine/Plugins/Animation/DeformerGraph/Source/OptimusCore/Private/Nodes/OptimusNode_ComputeKernelBase.cpp
+++ b/Engine/Plugins/Animation/DeformerGraph/Source/OptimusCore/Private/Nodes/OptimusNode_ComputeKernelBase.cpp
@@ -197,7 +197,7 @@ FOptimus_ComputeKernelResult UOptimusNode_ComputeKernelBase::CreateComputeKernel
 	CookedSource += GetKernelSourceText();
 	
 	KernelSource->SetSource(CookedSource);
-	KernelSource->EntryPoint = GetKernelName();
+	KernelSource->EntryPoint = GetKernelHlslName();
 	KernelSource->GroupSize = GetGroupSize();
 	KernelSource->AdditionalSources = GetAdditionalSources();
 
diff --git a/Engine/Plugins/Animation/DeformerGraph/Source/OptimusCore/Private/Nodes/OptimusNode_ComputeKernelBase.h b/Engine/Plugins/Animation/DeformerGraph/Source/OptimusCore/Private/Nodes/OptimusNode_ComputeKernelBase.h
--- a/Engine/Plugins/Animation/DeformerGraph/Source/OptimusCore/Private/Nodes/OptimusNode_ComputeKernelBase.h
+++ b/Engine/Plugins/Animation/DeformerGraph/Source/OptimusCore/Private/Nodes/OptimusNode_ComputeKernelBase.h
@@ -24,7 +24,7 @@ class UOptimusNode_ComputeKernelBase :
 	
 public:
 	/** Implement this to return the HLSL kernel's function name */
-	virtual FString GetKernelName() const PURE_VIRTUAL(UOptimusNode_ComputeKernelBase::GetKernelName, return FString();)
+	virtual FString GetKernelHlslName() const PURE_VIRTUAL(UOptimusNode_ComputeKernelBase::GetKernelHlslName, return FString();)
 
 	/** Implement this to return the HLSL kernel's function name */
 	virtual FIntVector GetGroupSize() const PURE_VIRTUAL(UOptimusNode_ComputeKernelBase::GetGroupSize, return FIntVector();)
diff --git a/Engine/Plugins/Animation/DeformerGraph/Source/OptimusCore/Private/Nodes/OptimusNode_ComputeKernelFunction.cpp b/Engine/Plugins/Animation/DeformerGraph/Source/OptimusCore/Private/Nodes/OptimusNode_ComputeKernelFunction.cpp
--- a/Engine/Plugins/Animation/DeformerGraph/Source/OptimusCore/Private/Nodes/OptimusNode_ComputeKernelFunction.cpp
+++ b/Engine/Plugins/Animation/DeformerGraph/Source/OptimusCore/Private/Nodes/OptimusNode_ComputeKernelFunction.cpp
@@ -171,7 +171,7 @@ FName UOptimusNode_ComputeKernelFunction::GetNodeCategory() const
 }
 
 
-FString UOptimusNode_ComputeKernelFunction::GetKernelName() const
+FString UOptimusNode_ComputeKernelFunction::GetKernelHlslName() const
 {
 	return GetGeneratorClass()->KernelName.ToString();
 }
@@ -185,7 +185,7 @@ FIntVector UOptimusNode_ComputeKernelFunction::GetGroupSize() const
 
 FString UOptimusNode_ComputeKernelFunction::GetKernelSourceText() const
 {
-	return Optimus::GetCookedKernelSource(GetPathName(), GetGeneratorClass()->ShaderSource, GetKernelName(), GetGroupSize());
+	return Optimus::GetCookedKernelSource(GetPathName(), GetGeneratorClass()->ShaderSource, GetKernelHlslName(), GetGroupSize());
 }
 
 
diff --git a/Engine/Plugins/Animation/DeformerGraph/Source/OptimusCore/Private/Nodes/OptimusNode_ComputeKernelFunction.h b/Engine/Plugins/Animation/DeformerGraph/Source/OptimusCore/Private/Nodes/OptimusNode_ComputeKernelFunction.h
--- a/Engine/Plugins/Animation/DeformerGraph/Source/OptimusCore/Private/Nodes/OptimusNode_ComputeKernelFunction.h
+++ b/Engine/Plugins/Animation/DeformerGraph/Source/OptimusCore/Private/Nodes/OptimusNode_ComputeKernelFunction.h
@@ -69,7 +69,7 @@ public:
 	FName GetNodeCategory() const override; 
 
 	// UOptimusNode_ComputeKernelBase overrides
-	FString GetKernelName() const override;
+	FString GetKernelHlslName() const override;
 	FIntVector GetGroupSize() const override;
 	FString GetKernelSourceText() const override;
 
diff --git a/Engine/Plugins/Animation/DeformerGraph/Source/OptimusCore/Private/Nodes/OptimusNode_CustomComputeKernel.cpp b/Engine/Plugins/Animation/DeformerGraph/Source/OptimusCore/Private/Nodes/OptimusNode_CustomComputeKernel.cpp
--- a/Engine/Plugins/Animation/DeformerGraph/Source/OptimusCore/Private/Nodes/OptimusNode_CustomComputeKernel.cpp
+++ b/Engine/Plugins/Animation/DeformerGraph/Source/OptimusCore/Private/Nodes/OptimusNode_CustomComputeKernel.cpp
@@ -102,7 +102,7 @@ UOptimusNode_CustomComputeKernel::UOptimusNode_CustomComputeKernel()
 
 FString UOptimusNode_CustomComputeKernel::GetKernelSourceText() const
 {
-	return Optimus::GetCookedKernelSource(GetPathName(), ShaderSource.ShaderText, KernelName.ToString(), GroupSize);
+	return Optimus::GetCookedKernelSource(GetPathName(), ShaderSource.ShaderText, GetKernelHlslName(), GroupSize);
 }
 
 FOptimusExecutionDomain UOptimusNode_CustomComputeKernel::GetExecutionDomain() const
@@ -724,6 +724,16 @@ PRAGMA_ENABLE_DEPRECATION_WARNINGS
 	SetDisplayName(FText::FromName(KernelName));	
 }
 
+FString UOptimusNode_CustomComputeKernel::GetKernelHlslName() const
+{
+	// We need to append the node name here because of a weird issue on Mac where
+	// if the last word of the name of a shader parameter matches the kernel name, the value
+	// for the parameter is always zero for unknown reasons. (e.g. AnimAttribute named "XXXBend" & Kernel also named "Bend")
+	// Adding this bit of "random" string after the user provided kernel name should
+	// greatly reduce the chance of things like that happening
+	return KernelName.ToString() + TEXT("_") + GetName();
+}
+
 
 void UOptimusNode_CustomComputeKernel::Serialize(FArchive& Ar)
 {
diff --git a/Engine/Plugins/Animation/DeformerGraph/Source/OptimusCore/Private/Nodes/OptimusNode_CustomComputeKernel.h b/Engine/Plugins/Animation/DeformerGraph/Source/OptimusCore/Private/Nodes/OptimusNode_CustomComputeKernel.h
--- a/Engine/Plugins/Animation/DeformerGraph/Source/OptimusCore/Private/Nodes/OptimusNode_CustomComputeKernel.h
+++ b/Engine/Plugins/Animation/DeformerGraph/Source/OptimusCore/Private/Nodes/OptimusNode_CustomComputeKernel.h
@@ -54,7 +54,7 @@ public:
 	void PostLoadNodeSpecificData() override;
 
 	// UOptimusNode_ComputeKernelBase overrides
-	FString GetKernelName() const override { return KernelName.ToString(); }
+	FString GetKernelHlslName() const override;
 	FIntVector GetGroupSize() const override { return GroupSize; }
 	FString GetKernelSourceText() const override;
 	TArray<TObjectPtr<UComputeSource>> GetAdditionalSources() const override { return AdditionalSources; }
