[Deformer Graph] Added a new data interface "Advanced Skeleton" that supports additional features like accessing different skin weight profiles and access per-bone animation attribute

diff --git a/Engine/Plugins/Animation/DeformerGraph/Shaders/Private/DataInterfaceAdvancedSkeleton.ush b/Engine/Plugins/Animation/DeformerGraph/Shaders/Private/DataInterfaceAdvancedSkeleton.ush
new file mode 100644
--- /dev/null
+++ b/Engine/Plugins/Animation/DeformerGraph/Shaders/Private/DataInterfaceAdvancedSkeleton.ush
@@ -0,0 +1,199 @@
+// Copyright Epic Games, Inc. All Rights Reserved.
+
+#include "/Engine/Private/GpuSkinCommon.ush"
+
+uint {DataInterfaceName}_NumVertices;
+uint {DataInterfaceName}_InputStreamStart;
+uint {DataInterfaceName}_NumBoneInfluences;
+uint {DataInterfaceName}_InputWeightStart;
+uint {DataInterfaceName}_InputWeightStride;
+uint {DataInterfaceName}_InputWeightIndexSize;
+Buffer<float4> {DataInterfaceName}_BoneMatrices;
+Buffer<uint> {DataInterfaceName}_InputWeightStream;
+Buffer<uint> {DataInterfaceName}_InputWeightLookupStream;
+
+StructuredBuffer<float3x4> {DataInterfaceName}_LayeredBoneMatrices; 
+
+uint ReadNumVertices_{DataInterfaceName}()
+{
+	return {DataInterfaceName}_NumVertices;
+}
+
+#if !GPUSKIN_UNLIMITED_BONE_INFLUENCE
+
+int4 GetBlendIndices_{DataInterfaceName}(uint StreamOffset, uint ExtraInfluenceIndex)
+{
+	return GetBlendIndices({DataInterfaceName}_InputWeightStream, StreamOffset, ExtraInfluenceIndex);
+}
+
+float4 GetBlendWeights_{DataInterfaceName}(uint StreamOffset, uint ExtraInfluenceIndex)
+{
+	return GetBlendWeights( { DataInterfaceName }_InputWeightStream, StreamOffset, ExtraInfluenceIndex, { DataInterfaceName }_NumBoneInfluences);
+}
+
+#endif // GPUSKIN_UNLIMITED_BONE_INFLUENCE
+
+float3x4 GetBoneMatrix_{DataInterfaceName}(uint BoneIndex)
+{
+	uint BufferIndex = BoneIndex * 3;
+	return float3x4({DataInterfaceName}_BoneMatrices[BufferIndex], {DataInterfaceName}_BoneMatrices[BufferIndex + 1], {DataInterfaceName}_BoneMatrices[BufferIndex + 2]);
+}
+
+uint ReadNumBones_{DataInterfaceName}(uint VertexIndex)
+{
+#if !ENABLE_DEFORMER_BONES
+	return 1;
+#elif GPUSKIN_UNLIMITED_BONE_INFLUENCE
+	const uint BufferIndex 		= {DataInterfaceName}_InputStreamStart + VertexIndex;
+	uint BlendOffsetCount = {DataInterfaceName}_InputWeightLookupStream[BufferIndex];
+	return BlendOffsetCount & 0xff;
+#else // !GPUSKIN_UNLIMITED_BONE_INFLUENCE
+	return {DataInterfaceName}_NumBoneInfluences;
+#endif
+}
+
+float3x4 ReadBoneMatrix_{DataInterfaceName}(uint VertexIndex, uint BoneIndex)
+{
+#if !ENABLE_DEFORMER_BONES
+
+	return float3x4(1, 0, 0, 0, 0, 1, 0, 0, 0, 0, 1, 0);
+
+#elif GPUSKIN_UNLIMITED_BONE_INFLUENCE
+	const uint BufferIndex 		= {DataInterfaceName}_InputStreamStart + VertexIndex;
+	uint BlendOffsetCount = {DataInterfaceName}_InputWeightLookupStream[BufferIndex];
+	int StreamOffset = BlendOffsetCount >> 8;
+	int BoneIndexSize = {DataInterfaceName}_InputWeightIndexSize & 0xff;
+	int BoneIndexOffset = StreamOffset + (BoneIndexSize * BoneIndex);
+	int BoneGlobalIndex = {DataInterfaceName}_InputWeightStream[BoneIndexOffset];
+	if (BoneIndexSize > 1)
+	{
+		BoneGlobalIndex = {DataInterfaceName}_InputWeightStream[BoneIndexOffset + 1] << 8 | BoneGlobalIndex;
+	}
+	
+	return GetBoneMatrix_{DataInterfaceName}(BoneGlobalIndex);
+
+#else // !GPUSKIN_UNLIMITED_BONE_INFLUENCE
+	
+	uint StreamOffset = {DataInterfaceName}_InputWeightStart + (VertexIndex * ({DataInterfaceName}_InputWeightStride / 4));
+	int BlendIndex = GetBlendIndices_{DataInterfaceName}(StreamOffset, BoneIndex / 4)[BoneIndex & 0x3];
+	return GetBoneMatrix_{DataInterfaceName}(BlendIndex);
+
+#endif
+}
+
+float ReadBoneWeight_{DataInterfaceName}(uint VertexIndex, uint BoneIndex)
+{
+#if !ENABLE_DEFORMER_BONES
+
+	return 1;
+
+#elif GPUSKIN_UNLIMITED_BONE_INFLUENCE
+	const uint BufferIndex 		= {DataInterfaceName}_InputStreamStart + VertexIndex;	
+	uint BlendOffsetCount = {DataInterfaceName}_InputWeightLookupStream[BufferIndex];
+	int NumBoneInfluencesLocal = BlendOffsetCount & 0xff;
+	int StreamOffset = BlendOffsetCount >> 8;
+	int BoneIndexSize = {DataInterfaceName}_InputWeightIndexSize & 0xff;
+	int BoneWeightSize = {DataInterfaceName}_InputWeightIndexSize >> 8;
+	int WeightsOffset = StreamOffset + (BoneIndexSize * NumBoneInfluencesLocal);
+	int BoneWeightOffset = WeightsOffset + BoneIndex * BoneWeightSize;
+	
+	if (BoneWeightSize > 1)
+	{
+		return float({DataInterfaceName}_InputWeightStream[BoneWeightOffset + 1] << 8 | {DataInterfaceName}_InputWeightStream[BoneWeightOffset]) / 65535.0;
+	}
+	else
+	{
+		return float({DataInterfaceName}_InputWeightStream[BoneWeightOffset]) / 255.0f;
+	}
+
+#else // !GPUSKIN_UNLIMITED_BONE_INFLUENCE
+	
+	uint StreamOffset = {DataInterfaceName}_InputWeightStart + (VertexIndex * ({DataInterfaceName}_InputWeightStride / 4));
+	float BlendWeight = GetBlendWeights_{DataInterfaceName}(StreamOffset, BoneIndex / 4)[BoneIndex & 0x3];
+	return BlendWeight;
+
+#endif
+}
+
+float3x4 ReadWeightedBoneMatrix_{DataInterfaceName}(uint VertexIndex)
+{
+#if !ENABLE_DEFORMER_BONES
+
+	return float3x4(1, 0, 0, 0, 0, 1, 0, 0, 0, 0, 1, 0);
+
+#elif GPUSKIN_UNLIMITED_BONE_INFLUENCE
+
+	const uint BufferIndex 		= {DataInterfaceName}_InputStreamStart + VertexIndex;
+	const uint BlendOffsetCount = {DataInterfaceName}_InputWeightLookupStream[BufferIndex];
+	return ComputeBoneMatrixWithUnlimitedInfluences({DataInterfaceName}_BoneMatrices, {DataInterfaceName}_InputWeightStream, {DataInterfaceName}_InputWeightIndexSize, BlendOffsetCount);
+
+#else
+
+	const uint StreamOffset = {DataInterfaceName}_InputWeightStart + (VertexIndex * ({DataInterfaceName}_InputWeightStride / 4));
+	const uint NumInfluences = {DataInterfaceName}_NumBoneInfluences;
+
+	FGPUSkinIndexAndWeight IndicesAndWeights = (FGPUSkinIndexAndWeight)0;
+	if (NumInfluences > 0)
+	{
+		IndicesAndWeights.BlendIndices  = GetBlendIndices_{DataInterfaceName}(StreamOffset, 0);
+		IndicesAndWeights.BlendWeights  = GetBlendWeights_{DataInterfaceName}(StreamOffset, 0);
+	}
+	if (NumInfluences > 4)
+	{
+		IndicesAndWeights.BlendIndices2 = GetBlendIndices_{DataInterfaceName}(StreamOffset, 1);
+		IndicesAndWeights.BlendWeights2 = GetBlendWeights_{DataInterfaceName}(StreamOffset, 1);
+	}
+	return ComputeBoneMatrixWithLimitedInfluences({DataInterfaceName}_BoneMatrices, IndicesAndWeights, NumInfluences > 4);
+
+#endif
+}
+
+float3x4 ReadLayeredBoneMatrix_{DataInterfaceName}(uint VertexIndex, uint BoneIndex)
+{
+#if !ENABLE_DEFORMER_BONES
+
+	return float3x4(1, 0, 0, 0, 0, 1, 0, 0, 0, 0, 1, 0);
+
+#elif GPUSKIN_UNLIMITED_BONE_INFLUENCE
+	const uint BufferIndex 		= {DataInterfaceName}_InputStreamStart + VertexIndex;
+	uint BlendOffsetCount = {DataInterfaceName}_InputWeightLookupStream[BufferIndex];
+	int StreamOffset = BlendOffsetCount >> 8;
+	int BoneIndexSize = {DataInterfaceName}_InputWeightIndexSize & 0xff;
+	int BoneIndexOffset = StreamOffset + (BoneIndexSize * BoneIndex);
+	int BoneGlobalIndex = {DataInterfaceName}_InputWeightStream[BoneIndexOffset];
+	if (BoneIndexSize > 1)
+	{
+		BoneGlobalIndex = {DataInterfaceName}_InputWeightStream[BoneIndexOffset + 1] << 8 | BoneGlobalIndex;
+	}
+	
+	return {DataInterfaceName}_LayeredBoneMatrices[BoneGlobalIndex];
+
+#else // !GPUSKIN_UNLIMITED_BONE_INFLUENCE
+	
+	uint StreamOffset = {DataInterfaceName}_InputWeightStart + (VertexIndex * ({DataInterfaceName}_InputWeightStride / 4));
+	int BlendIndex = GetBlendIndices_{DataInterfaceName}(StreamOffset, BoneIndex / 4)[BoneIndex & 0x3];
+	return {DataInterfaceName}_LayeredBoneMatrices[BlendIndex];
+
+#endif
+}
+
+float3x4 ReadWeightedLayeredBoneMatrix_{DataInterfaceName}(uint VertexIndex)
+{
+#if !ENABLE_DEFORMER_BONES
+	
+	return float3x4(1, 0, 0, 0, 0, 1, 0, 0, 0, 0, 1, 0);
+	
+#else
+	uint NumBoneInfluences = ReadNumBones_{DataInterfaceName}(VertexIndex);
+	float3x4 Result = 0;
+	
+	for (uint InfluenceIdx = 0; InfluenceIdx < NumBoneInfluences; InfluenceIdx++)
+	{
+		float Weight = ReadBoneWeight_{DataInterfaceName}(VertexIndex, InfluenceIdx);
+		float3x4 BoneMatrix = ReadLayeredBoneMatrix_{DataInterfaceName}(VertexIndex, InfluenceIdx);
+		Result += Weight * BoneMatrix;
+	}
+
+	return Result;
+#endif
+}
\ No newline at end of file
diff --git a/Engine/Plugins/Animation/DeformerGraph/Shaders/Private/DataInterfaceAdvancedSkeletonAnimAttribute.ush b/Engine/Plugins/Animation/DeformerGraph/Shaders/Private/DataInterfaceAdvancedSkeletonAnimAttribute.ush
new file mode 100644
--- /dev/null
+++ b/Engine/Plugins/Animation/DeformerGraph/Shaders/Private/DataInterfaceAdvancedSkeletonAnimAttribute.ush
@@ -0,0 +1,33 @@
+// Copyright Epic Games, Inc. All Rights Reserved.
+
+StructuredBuffer<{TypeName}> {DataInterfaceName}_{AttributeName};
+
+
+{TypeName} Read{AttributeName}_{DataInterfaceName}(uint VertexIndex, uint BoneIndex)
+{
+#if !ENABLE_DEFORMER_BONES
+
+	return 0;
+
+#elif GPUSKIN_UNLIMITED_BONE_INFLUENCE
+	const uint BufferIndex 		= {DataInterfaceName}_InputStreamStart + VertexIndex;
+	uint BlendOffsetCount = {DataInterfaceName}_InputWeightLookupStream[BufferIndex];
+	int StreamOffset = BlendOffsetCount >> 8;
+	int BoneIndexSize = {DataInterfaceName}_InputWeightIndexSize & 0xff;
+	int BoneIndexOffset = StreamOffset + (BoneIndexSize * BoneIndex);
+	int BoneGlobalIndex = {DataInterfaceName}_InputWeightStream[BoneIndexOffset];
+	if (BoneIndexSize > 1)
+	{
+		BoneGlobalIndex = {DataInterfaceName}_InputWeightStream[BoneIndexOffset + 1] << 8 | BoneGlobalIndex;
+	}
+	
+	return {DataInterfaceName}_{AttributeName}[BoneGlobalIndex];
+
+#else // !GPUSKIN_UNLIMITED_BONE_INFLUENCE
+	
+	uint StreamOffset = {DataInterfaceName}_InputWeightStart + (VertexIndex * ({DataInterfaceName}_InputWeightStride / 4));
+	int BlendIndex = GetBlendIndices_{DataInterfaceName}(StreamOffset, BoneIndex / 4)[BoneIndex & 0x3];
+	return {DataInterfaceName}_{AttributeName}[BlendIndex];
+
+#endif
+}
\ No newline at end of file
diff --git a/Engine/Plugins/Animation/DeformerGraph/Source/OptimusCore/Private/DataInterfaces/OptimusDataInterfaceAdvancedSkeleton.cpp b/Engine/Plugins/Animation/DeformerGraph/Source/OptimusCore/Private/DataInterfaces/OptimusDataInterfaceAdvancedSkeleton.cpp
new file mode 100644
--- /dev/null
+++ b/Engine/Plugins/Animation/DeformerGraph/Source/OptimusCore/Private/DataInterfaces/OptimusDataInterfaceAdvancedSkeleton.cpp
@@ -0,0 +1,1027 @@
+// Copyright Epic Games, Inc. All Rights Reserved.
+
+#include "OptimusDataInterfaceAdvancedSkeleton.h"
+
+#include "ComputeFramework/ComputeKernelPermutationVector.h"
+#include "ComputeFramework/ShaderParamTypeDefinition.h"
+#include "OptimusHelpers.h"
+#include "OptimusDataDomain.h"
+#include "OptimusDataTypeRegistry.h"
+#include "RenderGraphBuilder.h"
+#include "Rendering/SkeletalMeshLODRenderData.h"
+#include "Rendering/SkeletalMeshRenderData.h"
+#include "ShaderCompilerCore.h"
+#include "ShaderParameterMetadataBuilder.h"
+#include "SkeletalMeshDeformerHelpers.h"
+#include "SkeletalRenderPublic.h"
+#include "SystemTextures.h"
+#include "Animation/SkinWeightProfileManager.h"
+#include "Nodes/OptimusNode_DataInterface.h"
+#include "ComputeFramework/ComputeMetadataBuilder.h"
+#include "ComputeFramework/ShaderParameterMetadataAllocation.h"
+#include "Components/SkeletalMeshComponent.h"
+
+#include UE_INLINE_GENERATED_CPP_BY_NAME(OptimusDataInterfaceAdvancedSkeleton)
+
+const TCHAR* FOptimusAnimAttributeBufferDescription::PinNameDelimiter = TEXT(" - ");
+const TCHAR* FOptimusAnimAttributeBufferDescription::HlslIdDelimiter = TEXT("_");
+
+
+FOptimusAnimAttributeBufferDescription& FOptimusAnimAttributeBufferDescription::Init(const FString& InName, const FOptimusDataTypeRef& InDataType)
+{
+	Name = InName;
+	DataType = InDataType;
+
+	DefaultValueStruct.SetType(DataType.Resolve());
+	// Caller should ensure that the name is unique
+	HlslId = InName;
+	PinName = *InName;
+	
+	return *this;
+}
+
+
+void FOptimusAnimAttributeBufferDescription::UpdatePinNameAndHlslId(bool bInIncludeTypeName)
+{
+	PinName = *GetFormattedId(PinNameDelimiter, bInIncludeTypeName);
+	HlslId = GetFormattedId(HlslIdDelimiter, bInIncludeTypeName);
+}
+
+FString FOptimusAnimAttributeBufferDescription::GetFormattedId(
+	const FString& InDelimiter, bool bInIncludeTypeName) const
+{
+	FString UniqueId;
+			
+	if (bInIncludeTypeName)
+	{
+		UniqueId += DataType.Resolve()->DisplayName.ToString();
+		UniqueId += InDelimiter;
+	}
+
+	UniqueId += Name;	
+
+	return  UniqueId;
+}
+
+FString UOptimusAdvancedSkeletonDataInterface::GetUnusedAttributeName(int32 CurrentAttributeIndex, const FString& InName) const
+{
+	TSet<FString> UsedNames;
+
+	TArray<FOptimusCDIPinDefinition> PinDefinitions = GetPinDefinitions_Internal(true , CurrentAttributeIndex);
+	for (const FOptimusCDIPinDefinition& Definition : PinDefinitions)
+	{
+		UsedNames.Add(Definition.PinName.ToString());
+	}
+
+	int32 Suffix = 0;
+	FString NewName = InName;
+	while (UsedNames.Contains(NewName))
+	{
+		NewName = FString::Printf(TEXT("%s_%d"), *InName, Suffix);
+		Suffix++;
+	}
+
+	return NewName;
+}
+
+void UOptimusAdvancedSkeletonDataInterface::UpdateAttributePinNamesAndHlslIds()
+{
+	const int32 NumAttributes = AttributeBufferArray.Num(); 
+
+	TMap<FString, TArray<int32>> AttributesByName;
+
+	for (int32 Index = 0; Index < NumAttributes; Index++)
+	{
+		const FOptimusAnimAttributeBufferDescription& Attribute = AttributeBufferArray[Index];
+		AttributesByName.FindOrAdd(Attribute.Name).Add(Index);
+	}
+
+	for (const TTuple<FString, TArray<int32>>& AttributeGroup : AttributesByName)
+	{
+		// For attributes that share the same name, prepend type name to make sure pin names are unique
+		bool bMoreThanOneType = false;
+
+		TOptional<FOptimusDataTypeRef> LastType;
+		
+		for (int32 Index : AttributeGroup.Value)
+		{
+			const FOptimusAnimAttributeBufferDescription& Attribute = AttributeBufferArray[Index];
+
+			if (!LastType.IsSet())
+			{
+				LastType = Attribute.DataType;
+			}
+			else if (Attribute.DataType != LastType.GetValue())
+			{
+				bMoreThanOneType = true;
+			}
+
+			if (bMoreThanOneType)
+			{
+				break;
+			}
+		}
+		
+		for (int32 Index : AttributeGroup.Value)
+		{
+			FOptimusAnimAttributeBufferDescription& Attribute = AttributeBufferArray[Index];
+
+			Attribute.UpdatePinNameAndHlslId(bMoreThanOneType);
+		}
+	}
+}
+
+TArray<FOptimusCDIPinDefinition> UOptimusAdvancedSkeletonDataInterface::GetPinDefinitions_Internal(bool bGetAllPossiblePins, int32 AttributeIndexToExclude) const
+{
+	TArray<FOptimusCDIPinDefinition> Defs;
+	Defs.Add({"NumBones", "ReadNumBones", Optimus::DomainName::Vertex, "ReadNumVertices"});
+	Defs.Add({"BoneMatrix", "ReadBoneMatrix", {{Optimus::DomainName::Vertex, "ReadNumVertices"}, {Optimus::DomainName::Bone, "ReadNumBones"}}});
+	Defs.Add({"BoneWeight", "ReadBoneWeight", {{Optimus::DomainName::Vertex, "ReadNumVertices"}, {Optimus::DomainName::Bone, "ReadNumBones"}}});
+	Defs.Add({"WeightedBoneMatrix", "ReadWeightedBoneMatrix", Optimus::DomainName::Vertex, "ReadNumVertices" });
+	
+	if (bGetAllPossiblePins || bEnableLayeredSkinning)
+	{
+		Defs.Add({"LayeredBoneMatrix", "ReadLayeredBoneMatrix", {{Optimus::DomainName::Vertex, "ReadNumVertices"}, {Optimus::DomainName::Bone, "ReadNumBones"}}});
+		Defs.Add({"WeightedLayeredBoneMatrix", "ReadWeightedLayeredBoneMatrix", Optimus::DomainName::Vertex, "ReadNumVertices" });
+	}
+
+	for (int32 Index = 0; Index < AttributeBufferArray.Num(); Index++)
+	{
+		if (Index != AttributeIndexToExclude)
+		{
+			const FOptimusAnimAttributeBufferDescription& Attribute = AttributeBufferArray[Index];
+			Defs.Add({Attribute.PinName, FString::Printf(TEXT("Read%s"), *Attribute.HlslId), {{Optimus::DomainName::Vertex, "ReadNumVertices"}, {Optimus::DomainName::Bone, "ReadNumBones"}}});
+		}
+	}
+
+	return Defs;	
+}
+
+
+const FOptimusAnimAttributeBufferDescription& UOptimusAdvancedSkeletonDataInterface::AddAnimAttribute(const FString& InName,
+                                                                                                      const FOptimusDataTypeRef& InDataType)
+{
+	return AttributeBufferArray.InnerArray.AddDefaulted_GetRef()
+		.Init(GetUnusedAttributeName(AttributeBufferArray.InnerArray.Num() - 1, InName), InDataType);
+}
+
+#if WITH_EDITOR
+void UOptimusAdvancedSkeletonDataInterface::PostEditChangeChainProperty(FPropertyChangedChainEvent& PropertyChangedEvent)
+{
+	Super::PostEditChangeChainProperty(PropertyChangedEvent);
+	
+	const FName BasePropertyName = (PropertyChangedEvent.MemberProperty ? PropertyChangedEvent.MemberProperty->GetFName() : NAME_None);
+	const FName PropertyName = (PropertyChangedEvent.Property ? PropertyChangedEvent.Property->GetFName() : NAME_None);
+	
+	if (PropertyChangedEvent.ChangeType == EPropertyChangeType::ValueSet)
+	{
+		const int32 ChangedIndex = PropertyChangedEvent.GetArrayIndex(GET_MEMBER_NAME_STRING_CHECKED(FOptimusAnimAttributeBufferArray, InnerArray));
+		
+		bool bHasAttributeIdChanged =
+			(PropertyName == GET_MEMBER_NAME_CHECKED(FOptimusAnimAttributeBufferDescription, Name)) ||
+			(PropertyName == GET_MEMBER_NAME_CHECKED(FOptimusDataTypeRef, TypeName));
+				
+		if (bHasAttributeIdChanged)
+		{
+			if(ensure(AttributeBufferArray.IsValidIndex(ChangedIndex)))
+			{
+				FOptimusAnimAttributeBufferDescription& ChangedAttribute = AttributeBufferArray[ChangedIndex];
+				
+				FName OldPinName = ChangedAttribute.PinName;
+
+				if (ChangedAttribute.Name.IsEmpty())
+				{
+					ChangedAttribute.Name = TEXT("EmptyName");
+				}
+
+				ChangedAttribute.Name = GetUnusedAttributeName(ChangedIndex, ChangedAttribute.Name);
+				
+				UpdateAttributePinNamesAndHlslIds();
+				OnPinDefinitionRenamedDelegate.Execute(OldPinName, ChangedAttribute.PinName);
+			}
+		}
+
+		if (PropertyName == GET_MEMBER_NAME_CHECKED(FOptimusDataTypeRef, TypeName))
+		{
+			FOptimusAnimAttributeBufferDescription& ChangedAttribute = AttributeBufferArray[ChangedIndex];
+
+			// Update the default value container accordingly
+			ChangedAttribute.DefaultValueStruct.SetType(ChangedAttribute.DataType);
+			OnPinDefinitionChangedDelegate.Execute();
+		}
+
+		if (PropertyName == GET_MEMBER_NAME_CHECKED(UOptimusAdvancedSkeletonDataInterface, bEnableLayeredSkinning))
+		{
+			OnPinDefinitionChangedDelegate.Execute();
+		}
+
+		if (PropertyName == GET_MEMBER_NAME_CHECKED(UOptimusAdvancedSkeletonDataInterface, SkinWeightProfile))
+		{
+			OnDisplayNameChangedDelegate.Execute();
+		}
+	}
+	else 
+	{
+		if (PropertyChangedEvent.ChangeType == EPropertyChangeType::ArrayAdd)
+		{
+			if (PropertyName == GET_MEMBER_NAME_STRING_CHECKED(FOptimusAnimAttributeBufferArray, InnerArray))
+			{
+				const int32 ChangedIndex = PropertyChangedEvent.GetArrayIndex(GET_MEMBER_NAME_STRING_CHECKED(FOptimusAnimAttributeBufferArray, InnerArray));
+				FOptimusAnimAttributeBufferDescription& Attribute = AttributeBufferArray[ChangedIndex];
+			
+				// Default to a float attribute
+				Attribute.Init(GetUnusedAttributeName(ChangedIndex, TEXT("EmptyName")), 
+					FOptimusDataTypeRegistry::Get().FindType(*FFloatProperty::StaticClass()));
+			}
+		}
+		else if (PropertyChangedEvent.ChangeType == EPropertyChangeType::Duplicate)
+		{
+			if (PropertyName == GET_MEMBER_NAME_STRING_CHECKED(FOptimusAnimAttributeBufferArray, InnerArray))
+			{	
+				const int32 ChangedIndex = PropertyChangedEvent.GetArrayIndex(GET_MEMBER_NAME_STRING_CHECKED(FOptimusAnimAttributeBufferArray, InnerArray));
+				FOptimusAnimAttributeBufferDescription& Attribute = AttributeBufferArray[ChangedIndex];
+			
+				Attribute.Name = GetUnusedAttributeName(ChangedIndex, Attribute.Name);
+				Attribute.UpdatePinNameAndHlslId();
+			}
+		}
+		
+		OnPinDefinitionChangedDelegate.Execute();
+	}
+}
+#endif
+
+FString UOptimusAdvancedSkeletonDataInterface::GetDisplayName() const
+{
+	if (SkinWeightProfile != NAME_None)
+	{
+		return TEXT("Advanced Skeleton (") + SkinWeightProfile.ToString() + TEXT(")");
+	}
+	
+	return TEXT("Advanced Skeleton");
+}
+
+TArray<FOptimusCDIPinDefinition> UOptimusAdvancedSkeletonDataInterface::GetPinDefinitions() const
+{
+	return GetPinDefinitions_Internal();
+}
+
+
+TSubclassOf<UActorComponent> UOptimusAdvancedSkeletonDataInterface::GetRequiredComponentClass() const
+{
+	return USkeletalMeshComponent::StaticClass();
+}
+
+void UOptimusAdvancedSkeletonDataInterface::Initialize()
+{
+}
+
+
+void UOptimusAdvancedSkeletonDataInterface::RegisterPropertyChangeDelegatesForOwningNode(UOptimusNode_DataInterface* InNode)
+{
+	OnPinDefinitionChangedDelegate.BindUObject(InNode, &UOptimusNode_DataInterface::RecreatePinsFromPinDefinitions);
+	OnPinDefinitionRenamedDelegate.BindUObject(InNode, &UOptimusNode_DataInterface::RenamePinFromPinDefinition);
+	OnDisplayNameChangedDelegate.BindUObject(InNode, &UOptimusNode_DataInterface::UpdateDisplayNameFromDataInterface);
+}
+
+void UOptimusAdvancedSkeletonDataInterface::OnDataTypeChanged(FName InTypeName)
+{
+	Super::OnDataTypeChanged(InTypeName);
+
+	for (FOptimusAnimAttributeBufferDescription& AttributeDescription : AttributeBufferArray)
+	{
+		if (AttributeDescription.DataType.TypeName == InTypeName)
+		{
+			AttributeDescription.DefaultValueStruct.SetType(AttributeDescription.DataType);
+		}
+	}
+}
+
+void UOptimusAdvancedSkeletonDataInterface::GetSupportedInputs(TArray<FShaderFunctionDefinition>& OutFunctions) const
+{
+	OutFunctions.AddDefaulted_GetRef()
+		.SetName(TEXT("ReadNumVertices"))
+		.AddReturnType(EShaderFundamentalType::Uint);
+
+	OutFunctions.AddDefaulted_GetRef()
+		.SetName(TEXT("ReadNumBones"))
+		.AddReturnType(EShaderFundamentalType::Uint)
+		.AddParam(EShaderFundamentalType::Uint);
+
+	OutFunctions.AddDefaulted_GetRef()
+		.SetName(TEXT("ReadBoneMatrix"))
+		.AddReturnType(EShaderFundamentalType::Float, 3, 4)
+		.AddParam(EShaderFundamentalType::Uint)
+		.AddParam(EShaderFundamentalType::Uint);
+
+	OutFunctions.AddDefaulted_GetRef()
+		.SetName(TEXT("ReadBoneWeight"))
+		.AddReturnType(EShaderFundamentalType::Float)
+		.AddParam(EShaderFundamentalType::Uint)
+		.AddParam(EShaderFundamentalType::Uint);
+
+	OutFunctions.AddDefaulted_GetRef()
+		.SetName(TEXT("ReadWeightedBoneMatrix"))
+		.AddReturnType(EShaderFundamentalType::Float, 3, 4)
+		.AddParam(EShaderFundamentalType::Uint);
+
+	if (bEnableLayeredSkinning)
+	{
+		OutFunctions.AddDefaulted_GetRef()
+			.SetName(TEXT("ReadLayeredBoneMatrix"))
+			.AddReturnType(EShaderFundamentalType::Float, 3, 4)
+			.AddParam(EShaderFundamentalType::Uint)
+			.AddParam(EShaderFundamentalType::Uint);
+
+		OutFunctions.AddDefaulted_GetRef()
+			.SetName(TEXT("ReadWeightedLayeredBoneMatrix"))
+			.AddReturnType(EShaderFundamentalType::Float, 3, 4)
+			.AddParam(EShaderFundamentalType::Uint);
+	}
+	
+	for (int32 Index = 0; Index < AttributeBufferArray.Num(); Index++)
+	{
+		const FOptimusAnimAttributeBufferDescription& Attribute = AttributeBufferArray[Index];
+
+		OutFunctions.AddDefaulted_GetRef()
+		.SetName(FString::Printf(TEXT("Read%s"), *Attribute.HlslId))
+		.AddReturnType(Attribute.DataType->ShaderValueType)
+		.AddParam(EShaderFundamentalType::Uint)
+		.AddParam(EShaderFundamentalType::Uint);
+	}
+
+}
+
+BEGIN_SHADER_PARAMETER_STRUCT(FAdvancedSkeletonDataInterfaceDefaultParameters, )
+	SHADER_PARAMETER(uint32, NumVertices)
+	SHADER_PARAMETER(uint32, InputStreamStart)
+	SHADER_PARAMETER(uint32, NumBoneInfluences)
+	SHADER_PARAMETER(uint32, InputWeightStart)
+	SHADER_PARAMETER(uint32, InputWeightStride)
+	SHADER_PARAMETER(uint32, InputWeightIndexSize)
+	SHADER_PARAMETER_SRV(Buffer<float4>, BoneMatrices)
+	SHADER_PARAMETER_SRV(Buffer<uint>, InputWeightStream)
+	SHADER_PARAMETER_SRV(Buffer<uint>, InputWeightLookupStream)
+
+	SHADER_PARAMETER_RDG_BUFFER_SRV(StructuredBuffer<float3x4>, LayeredBoneMatrices)
+END_SHADER_PARAMETER_STRUCT()
+
+void UOptimusAdvancedSkeletonDataInterface::GetShaderParameters(TCHAR const* UID, FShaderParametersMetadataBuilder& InOutBuilder, FShaderParametersMetadataAllocations& InOutAllocations) const
+{
+	FShaderParametersMetadataBuilder Builder;
+	Builder.AddIncludedStruct<FAdvancedSkeletonDataInterfaceDefaultParameters>();
+
+	TArray<FShaderParametersMetadata*> NestedStructs;
+
+	for (int32 Index = 0; Index < AttributeBufferArray.Num(); Index++)
+	{
+		const FOptimusAnimAttributeBufferDescription& Attribute = AttributeBufferArray[Index];
+		FShaderValueTypeHandle ArrayShaderType = FShaderValueType::MakeDynamicArrayType(Attribute.DataType->ShaderValueType);
+		ComputeFramework::AddParamForType(Builder, *Attribute.HlslId, ArrayShaderType, NestedStructs);	
+	}
+
+	FShaderParametersMetadata* ShaderParameterMetadata = Builder.Build(FShaderParametersMetadata::EUseCase::ShaderParameterStruct, TEXT("AnimAttributeBuffers"));
+
+	InOutAllocations.ShaderParameterMetadatas.Add(ShaderParameterMetadata);
+	InOutAllocations.ShaderParameterMetadatas.Append(NestedStructs);
+
+	// Add the generated nested struct to our builder.
+	InOutBuilder.AddNestedStruct(UID, ShaderParameterMetadata);
+}
+
+void UOptimusAdvancedSkeletonDataInterface::GetPermutations(FComputeKernelPermutationVector& OutPermutationVector) const
+{
+	// Need to be able to support these permutations according to the skeletal mesh settings.
+	// todo[CF]: I think GPUSKIN_UNLIMITED_BONE_INFLUENCE and GPUSKIN_BONE_INDEX_UINT16/GPUSKIN_BONE_WEIGHTS_UINT16 are mutually exclusive. So we could save permutations here.
+	OutPermutationVector.AddPermutation(TEXT("ENABLE_DEFORMER_BONES"), 2);
+	OutPermutationVector.AddPermutation(TEXT("GPUSKIN_UNLIMITED_BONE_INFLUENCE"), 2);
+	OutPermutationVector.AddPermutation(TEXT("GPUSKIN_BONE_INDEX_UINT16"), 2);
+	OutPermutationVector.AddPermutation(TEXT("GPUSKIN_BONE_WEIGHTS_UINT16"), 2);
+}
+
+TCHAR const* UOptimusAdvancedSkeletonDataInterface::SkeletonTemplateFilePath = TEXT("/Plugin/Optimus/Private/DataInterfaceAdvancedSkeleton.ush");
+TCHAR const* UOptimusAdvancedSkeletonDataInterface::AttributeTemplateFilePath = TEXT("/Plugin/Optimus/Private/DataInterfaceAdvancedSkeletonAnimAttribute.ush");
+
+
+void UOptimusAdvancedSkeletonDataInterface::GetShaderHash(FString& InOutKey) const
+{
+}
+
+void UOptimusAdvancedSkeletonDataInterface::GetHLSL(FString& OutHLSL, FString const& InDataInterfaceName) const
+{
+	TMap<FString, FStringFormatArg> SkeletonTemplateArgs =
+	{
+		{ TEXT("DataInterfaceName"), InDataInterfaceName },
+	};
+
+	FString SkeletonTemplateFile;
+	LoadShaderSourceFile(SkeletonTemplateFilePath, EShaderPlatform::SP_PCD3D_SM5, &SkeletonTemplateFile, nullptr);
+	OutHLSL += FString::Format(*SkeletonTemplateFile, SkeletonTemplateArgs);
+	
+	FString AttributeTemplateFile;
+	LoadShaderSourceFile(AttributeTemplateFilePath, EShaderPlatform::SP_PCD3D_SM5, &AttributeTemplateFile, nullptr);
+	
+	for (int32 Index = 0; Index < AttributeBufferArray.Num(); Index++)
+	{
+		FOptimusAnimAttributeBufferDescription AttributeBufferDescription = AttributeBufferArray[Index];
+		TMap<FString, FStringFormatArg> AttributeTemplateArgs =
+		{
+			{ TEXT("DataInterfaceName"), InDataInterfaceName },
+			{ TEXT("TypeName"), AttributeBufferDescription.DataType->ShaderValueType->ToString()},
+			{ TEXT("AttributeName"), AttributeBufferDescription.HlslId},
+		};
+		OutHLSL += FString::Format(*AttributeTemplateFile, AttributeTemplateArgs);	
+	}
+}
+
+FOptimusAnimAttributeBufferRuntimeData::FOptimusAnimAttributeBufferRuntimeData(const FOptimusAnimAttributeBufferDescription& InDescription)
+{
+	Name = *InDescription.Name;
+	HlslId = *InDescription.HlslId;
+	
+	Offset = 0;
+	Size = InDescription.DataType->ShaderValueSize;
+
+	const FOptimusDataTypeRegistry& Registry = FOptimusDataTypeRegistry::Get();
+
+	ConvertFunc = Registry.FindPropertyValueConvertFunc(InDescription.DataType.TypeName);
+
+	AttributeType = Registry.FindAttributeType(InDescription.DataType.TypeName);
+
+	CachedDefaultValue = InDescription.DefaultValueStruct.GetShaderValue(InDescription.DataType);	
+}
+
+UComputeDataProvider* UOptimusAdvancedSkeletonDataInterface::CreateDataProvider(TObjectPtr<UObject> InBinding, uint64 InInputMask, uint64 InOutputMask) const
+{
+	UOptimusAdvancedSkeletonDataProvider* Provider = NewObject<UOptimusAdvancedSkeletonDataProvider>();
+	Provider->Init(this, Cast<USkeletalMeshComponent>(InBinding));
+	return Provider;
+}
+void FOptimusBoneTransformBuffer::SetData(FSkeletalMeshLODRenderData const* InLodRenderData, const TArray<FTransform>& InBoneTransforms)
+{
+	const FName Matrix34TypeName = FOptimusDataTypeRegistry::Matrix34TypeName;
+	const FOptimusDataTypeHandle Matrix34TypeHandle = FOptimusDataTypeRegistry::Get().FindType(Matrix34TypeName);
+	int32 Matrix34ShaderSize = Matrix34TypeHandle->ShaderValueSize;
+
+	BufferData.AddDefaulted(InLodRenderData->RenderSections.Num());
+	NumBones.AddDefaulted(InLodRenderData->RenderSections.Num());
+				
+	for (int32 SectionIndex = 0; SectionIndex < InLodRenderData->RenderSections.Num(); ++SectionIndex)
+	{
+		FSkelMeshRenderSection const& RenderSection = InLodRenderData->RenderSections[SectionIndex];
+		NumBones[SectionIndex] = RenderSection.BoneMap.Num();
+
+		BufferData[SectionIndex].AddDefaulted(NumBones[SectionIndex] * Matrix34ShaderSize);
+		for (int32 BoneIndex = 0; BoneIndex < NumBones[SectionIndex]; BoneIndex++)
+		{
+			int32 FinalBoneIndex = RenderSection.BoneMap[BoneIndex];
+			uint8* BoneDataPtr = BufferData[SectionIndex].GetData() + BoneIndex * Matrix34ShaderSize;
+
+			const uint8* ValuePtr = reinterpret_cast<const uint8*>(&InBoneTransforms[FinalBoneIndex]);
+						
+			Optimus::ConvertFTransformToFMatrix3x4(InBoneTransforms[FinalBoneIndex],
+				{ {BoneDataPtr , Matrix34ShaderSize} ,{}} );
+		}
+	}		
+}
+
+void FOptimusBoneTransformBuffer::AllocateResources(FRDGBuilder& GraphBuilder)
+{
+	if (BufferRefPerSection.Num() == 0)
+	{
+		BufferRefPerSection.AddDefaulted(BufferData.Num());
+		BufferSRVPerSection.AddDefaulted(BufferData.Num());	
+	
+		// If we are using a raw type alias for the buffer then we need to adjust stride and count.
+		for (int32 InvocationIndex = 0 ; InvocationIndex < BufferData.Num(); InvocationIndex++)
+		{
+			int32 Stride = BufferData[InvocationIndex].Num() / NumBones[InvocationIndex];
+			FRDGBufferRef Buffer = GraphBuilder.CreateBuffer(FRDGBufferDesc::CreateStructuredDesc(Stride,  NumBones[InvocationIndex]), TEXT("BoneTransformBuffer"), ERDGBufferFlags::None);
+			FRDGBufferSRVRef BufferSRV = GraphBuilder.CreateSRV(Buffer);
+
+			BufferRefPerSection[InvocationIndex] = Buffer;
+			BufferSRVPerSection[InvocationIndex] = BufferSRV;
+
+			GraphBuilder.QueueBufferUpload(Buffer, BufferData[InvocationIndex].GetData(), BufferData[InvocationIndex].Num(), ERDGInitialDataFlags::None);	
+		}	
+	}
+}
+
+void UOptimusAdvancedSkeletonDataProvider::Init(const UOptimusAdvancedSkeletonDataInterface* InDataInterface, USkeletalMeshComponent* InSkeletalMesh)
+{
+	bEnableLayeredSkinning = InDataInterface->bEnableLayeredSkinning;
+	SkeletalMesh = InSkeletalMesh;
+	SkinWeightProfile = InDataInterface->SkinWeightProfile;
+
+	// Convert description to runtime data
+	for (const FOptimusAnimAttributeBufferDescription& Attribute : InDataInterface->AttributeBufferArray)
+	{
+		AttributeBufferRuntimeData.Add(Attribute);
+	}
+
+	// Compute offset within the shader parameter buffer for each attribute
+	FShaderParametersMetadataBuilder Builder;
+	FShaderParametersMetadataAllocations Allocations;
+	InDataInterface->GetShaderParameters(TEXT("Dummy"), Builder, Allocations);
+
+	{
+		struct FParametersMetadataScope
+		{
+			FParametersMetadataScope(FShaderParametersMetadata* InMetadata) : ShaderParameterMetadata(InMetadata){};
+			~FParametersMetadataScope() {delete ShaderParameterMetadata;};
+			
+			FShaderParametersMetadata* ShaderParameterMetadata = nullptr;
+		};
+
+		
+		FParametersMetadataScope MetadataScope(Builder.Build(FShaderParametersMetadata::EUseCase::ShaderParameterStruct, TEXT("UAnimAttributeDataInterface")));
+
+		TArray<FShaderParametersMetadata::FMember> const& DummyMember = MetadataScope.ShaderParameterMetadata->GetMembers();
+		check(DummyMember.Num()==1);
+		FShaderParametersMetadata::FMember const& Dummy = DummyMember.Last();
+		TArray<FShaderParametersMetadata::FMember> const& DataInterfaceParameterMembers = Dummy.GetStructMetadata()->GetMembers();
+
+		check(DataInterfaceParameterMembers.Num() == AttributeBufferRuntimeData.Num() + 1);
+
+		int32 AttributeMemberStart = 1;
+		
+		for (int32 Index = 0; Index < AttributeBufferRuntimeData.Num(); ++Index)
+		{
+			FOptimusAnimAttributeBufferRuntimeData& RuntimeData = AttributeBufferRuntimeData[Index];
+			int32 AttributeIndexInParameter = AttributeMemberStart + Index;
+			check(RuntimeData.HlslId == DataInterfaceParameterMembers[AttributeIndexInParameter].GetName());
+		
+			RuntimeData.Offset = DataInterfaceParameterMembers[AttributeIndexInParameter].GetOffset();
+		}
+
+		ParameterBufferSize = MetadataScope.ShaderParameterMetadata->GetSize();
+	}
+
+
+}
+
+void UOptimusAdvancedSkeletonDataProvider::ComputeBoneTransformsForLayeredSkinning(FOptimusBoneTransformBuffer& OutBoneBuffer, FSkeletalMeshLODRenderData const* InLodRenderData, const FReferenceSkeleton& InRefSkeleton)
+{
+	if (!bIsLayeredSkinInitialized)
+	{
+		bIsLayeredSkinInitialized = true;
+		
+		const TArray<FTransform>& InitialBoneSpaceTransforms = InRefSkeleton.GetRefBonePose();
+		
+		// 1. Look for all bones with non-zero weights in this skin weight profile
+		CachedWeightedBoneIndices.Reserve(InRefSkeleton.GetNum());
+		
+		FSkinWeightVertexBuffer const* WeightBuffer = InLodRenderData->GetSkinWeightVertexBuffer();
+		if (InLodRenderData->SkinWeightProfilesData.ContainsProfile(SkinWeightProfile))
+		{
+			WeightBuffer = InLodRenderData->SkinWeightProfilesData.GetOverrideBuffer(SkinWeightProfile);
+		}
+
+		for (uint32 VertexIndex = 0; VertexIndex < WeightBuffer->GetNumVertices(); VertexIndex++)
+		{
+			// Find the render section, which we need to find the final bone index.
+			int32 SectionIndex = INDEX_NONE;
+			int32 SectionVertexIndex = INDEX_NONE;
+			InLodRenderData->GetSectionFromVertexIndex(VertexIndex, SectionIndex, SectionVertexIndex);
+
+			uint32 VertexWeightOffset = 0;
+			uint32 VertexInfluenceCount = 0;
+			WeightBuffer->GetVertexInfluenceOffsetCount(VertexIndex, VertexWeightOffset, VertexInfluenceCount);
+			for (uint32 InfluenceIndex = 0 ; InfluenceIndex < VertexInfluenceCount; InfluenceIndex++)
+			{
+				float Weight = WeightBuffer->GetBoneWeight(VertexIndex, InfluenceIndex);
+				if (Weight > 0)
+				{
+					int32 SectionBoneIndex = WeightBuffer->GetBoneIndex(VertexIndex, InfluenceIndex);
+					int32 FinalBoneIndex = InLodRenderData->RenderSections[SectionIndex].BoneMap[SectionBoneIndex];
+					CachedWeightedBoneIndices.Add(FinalBoneIndex);
+				}
+			}
+		}
+
+		// 2. For each weighted bone, find its boundary bone
+		// Boundary bone is defined as any bone that is weighted and none of its parents(except root) are weighted,
+		// "sitting at the boundary of weighted and unweighted"
+		// E.g. If a bone is weighted, its parent is not weighted, but its grand parent is weighted, the grand parent is a boundary bone candidate
+		CachedBoundaryBoneIndex.Init(INDEX_NONE, InRefSkeleton.GetNum());
+		for (int32 BoneIndex = 1; BoneIndex < InRefSkeleton.GetNum(); BoneIndex++)
+		{
+			
+			int32 ParentIndex = InRefSkeleton.GetParentIndex(BoneIndex);
+
+			if (CachedBoundaryBoneIndex[ParentIndex] == INDEX_NONE)
+			{
+				if (CachedWeightedBoneIndices.Contains(BoneIndex))
+				{
+					// This must be the first weighted bone we have encountered that has no weighted parents
+					CachedBoundaryBoneIndex[BoneIndex] = BoneIndex;
+				}
+			}
+			else
+			{
+				// Child always inherit boundary bone from parent
+				CachedBoundaryBoneIndex[BoneIndex] = CachedBoundaryBoneIndex[ParentIndex];
+			}
+		}
+
+		// 3. For all children of the boundary bone, compute their initial layer space transform (rooted at the parent of the boundary)
+		CachedLayerSpaceInitialBoneTransform.AddDefaulted(InRefSkeleton.GetNum());
+		for (int32 BoneIndex = 1; BoneIndex < InRefSkeleton.GetNum(); BoneIndex++)
+		{
+			if (CachedBoundaryBoneIndex[BoneIndex] != INDEX_NONE)
+			{
+				int32 ParentIndex = InRefSkeleton.GetParentIndex(BoneIndex);
+
+				if (CachedBoundaryBoneIndex[ParentIndex] == INDEX_NONE)
+				{
+					CachedLayerSpaceInitialBoneTransform[BoneIndex] = InitialBoneSpaceTransforms[BoneIndex];		
+				}
+				else
+				{
+					// Accumulate Ref bone space transforms
+					CachedLayerSpaceInitialBoneTransform[BoneIndex] = InitialBoneSpaceTransforms[BoneIndex] * CachedLayerSpaceInitialBoneTransform[ParentIndex];
+				}
+				
+			}
+		}
+	}
+	
+	const TArray<FTransform>& CurrentComponentSpaceBoneMatrix = SkeletalMesh->GetComponentSpaceTransforms();
+
+	// 4. InverseBindMatrix for each weighted bone is:
+	// the inverse of (initial layer space transform * current global transform of parent of the layer boundary bone)
+	// and because we are looking at the current global of the boundary bone parent, this InverseBindMatrix can change every frame
+	TArray<FTransform> LayeredBoneMatrix;
+	LayeredBoneMatrix.AddDefaulted(InRefSkeleton.GetNum());
+	for (int32 BoneIndex : CachedWeightedBoneIndices)
+	{
+		if (BoneIndex == 0)
+		{
+			// Root bone is where we dump all the weights, so ignore any root transform if we are doing layered skinning
+			LayeredBoneMatrix[BoneIndex] = FTransform::Identity;
+			continue;
+		}
+
+		int32 BoundaryBoneParentIndex = InRefSkeleton.GetParentIndex(CachedBoundaryBoneIndex[BoneIndex]);
+		LayeredBoneMatrix[BoneIndex] =
+			(CachedLayerSpaceInitialBoneTransform[BoneIndex] * CurrentComponentSpaceBoneMatrix[BoundaryBoneParentIndex]).Inverse() * CurrentComponentSpaceBoneMatrix[BoneIndex];
+	}
+	
+	
+	OutBoneBuffer.SetData(InLodRenderData, LayeredBoneMatrix);
+}
+
+FComputeDataProviderRenderProxy* UOptimusAdvancedSkeletonDataProvider::GetRenderProxy()
+{
+	FOptimusAdvancedSkeletonDataProviderProxy* Proxy = new FOptimusAdvancedSkeletonDataProviderProxy();
+
+	if (SkeletalMesh && SkeletalMesh->MeshObject)
+	{
+		FSkeletalMeshObject* SkeletalMeshObject = SkeletalMesh->MeshObject;	
+		if (!bSkinWeightBufferReady)
+		{
+			const int32 LodIndex = SkeletalMeshObject->GetLOD();
+			FSkeletalMeshRenderData const& SkeletalMeshRenderData = SkeletalMeshObject->GetSkeletalMeshRenderData();
+			FSkeletalMeshLODRenderData const& LodRenderData = SkeletalMeshRenderData.LODRenderData[LodIndex];
+
+			if (SkinWeightProfile.IsNone())
+			{
+				bSkinWeightBufferReady = true;
+			}
+			else
+			{
+				if (LodRenderData.SkinWeightProfilesData.ContainsProfile(SkinWeightProfile))
+				{
+					// Retrieve this profile's skin weight buffer
+					FSkinWeightVertexBuffer* Buffer = LodRenderData.SkinWeightProfilesData.GetOverrideBuffer(SkinWeightProfile);
+						
+					if (Buffer)
+					{
+						bSkinWeightBufferReady = true;
+					}
+					else
+					{
+						TWeakObjectPtr<USkinnedMeshComponent> WeakComponent = SkeletalMesh;
+						TWeakObjectPtr<UOptimusAdvancedSkeletonDataProvider> WeakThis = this;
+						FRequestFinished Callback = [WeakComponent, WeakThis](TWeakObjectPtr<USkeletalMesh> WeakMesh, FName ProfileName)
+						{
+							// Ensure that the request objects are still valid
+							if (WeakMesh.IsValid() && WeakComponent.IsValid() && WeakThis.IsValid())
+							{
+								if (FSkeletalMeshRenderData * RenderData = WeakMesh->GetResourceForRendering())
+								{
+									const int32 NumLODs = RenderData->LODRenderData.Num();
+									for (int32 Index = 0; Index < NumLODs; ++Index)
+									{
+										FSkeletalMeshLODRenderData& LODRenderData = RenderData->LODRenderData[Index];
+										FSkinWeightProfilesData& SkinweightData = LODRenderData.SkinWeightProfilesData;
+
+										// Retrieve this profile's skin weight buffer
+										FSkinWeightVertexBuffer* Buffer = SkinweightData.GetOverrideBuffer(ProfileName);
+										if (ensure(Buffer))
+										{
+											WeakThis->bSkinWeightBufferReady = true;
+										}
+									}
+								}
+							}
+						};
+
+						// Put in a skin weight profile request
+						if (FSkinWeightProfileManager* Manager = FSkinWeightProfileManager::Get(SkeletalMesh->GetWorld()))
+						{
+							Manager->RequestSkinWeightProfile(SkinWeightProfile, SkeletalMesh->GetSkinnedAsset(), this, Callback);
+						}
+					}
+				}	
+			}
+		}
+	
+		if (bSkinWeightBufferReady)
+		{
+			const int32 LodIndex = SkeletalMeshObject->GetLOD();
+			FSkeletalMeshRenderData const& SkeletalMeshRenderData = SkeletalMeshObject->GetSkeletalMeshRenderData();
+			FSkeletalMeshLODRenderData const* LodRenderData = &SkeletalMeshRenderData.LODRenderData[LodIndex];
+			const FReferenceSkeleton& RefSkeleton = SkeletalMesh->GetSkinnedAsset()->GetRefSkeleton();
+			
+
+			FOptimusBoneTransformBuffer LayeredBoneMatrixBuffer;
+			
+			if (bEnableLayeredSkinning)
+			{
+				ComputeBoneTransformsForLayeredSkinning(LayeredBoneMatrixBuffer, LodRenderData, RefSkeleton);
+			}
+			
+			// Per-Bone Animation Attributes
+			const UE::Anim::FMeshAttributeContainer& AttributeContainer = SkeletalMesh->GetCustomAttributes();
+			// Section, Attribute, BoneData
+			TArray<TArray<TArray<uint8>>> AttributeBuffers;
+
+			AttributeBuffers.AddDefaulted(LodRenderData->RenderSections.Num());
+			for (int32 SectionIndex = 0; SectionIndex < LodRenderData->RenderSections.Num(); ++SectionIndex)
+			{
+				FSkelMeshRenderSection const& RenderSection = LodRenderData->RenderSections[SectionIndex];
+				AttributeBuffers[SectionIndex].AddDefaulted(AttributeBufferRuntimeData.Num());
+				for (int32 AttributeIndex = 0; AttributeIndex < AttributeBufferRuntimeData.Num(); ++AttributeIndex)
+				{
+					const FOptimusAnimAttributeBufferRuntimeData& AttributeData = AttributeBufferRuntimeData[AttributeIndex];
+					const int32 NumBones = RenderSection.BoneMap.Num();
+					AttributeBuffers[SectionIndex][AttributeIndex].AddDefaulted(NumBones * AttributeData.Size);
+
+					for (int32 BoneIndex = 0; BoneIndex < NumBones; BoneIndex++)
+					{
+						uint8* BoneDataPtr = AttributeBuffers[SectionIndex][AttributeIndex].GetData() + BoneIndex * AttributeData.Size;
+						int32 FinalBoneIndex = RenderSection.BoneMap[BoneIndex];
+
+						const UE::Anim::FAttributeId Id = {AttributeData.Name, FCompactPoseBoneIndex(FinalBoneIndex)} ;
+
+						bool bIsValueSet = false;
+				
+						FOptimusDataTypeRegistry::PropertyValueConvertFuncT ConvertFunc = AttributeData.ConvertFunc;
+
+						if (ConvertFunc)
+						{
+							UScriptStruct* AttributeType = AttributeData.AttributeType;
+				
+							if (const uint8* Attribute = AttributeContainer.Find(AttributeType, Id))
+							{
+								bIsValueSet = true;
+						
+								const uint8* ValuePtr = Attribute;
+								int32 ValueSize = AttributeType->GetStructureSize();
+
+								// TODO: use a specific function to extract the value from the attribute
+								// it works for now because even if the attribute type != actual value type
+								// it should only have a single property, whose type == actual property type
+						
+								ConvertFunc(
+									{ValuePtr, ValueSize},
+									{ {BoneDataPtr, AttributeData.Size}, {} });
+							}
+				
+							// Use the default value if the attribute was not found
+							if (!bIsValueSet)
+							{
+								const uint8* DefaultValuePtr = AttributeData.CachedDefaultValue.ShaderValue.GetData();
+								const uint32 DefaultValueSize = AttributeData.CachedDefaultValue.ShaderValue.Num();
+
+								FMemory::Memcpy(BoneDataPtr, DefaultValuePtr, DefaultValueSize);
+							}
+						}
+					}
+				}
+			}
+			
+			TArray<int32> AttributeBufferOffsets;
+			AttributeBufferOffsets.AddUninitialized(AttributeBufferRuntimeData.Num());
+			
+			for (int32 AttributeIndex = 0; AttributeIndex < AttributeBufferRuntimeData.Num(); ++AttributeIndex)
+			{
+				AttributeBufferOffsets[AttributeIndex] = AttributeBufferRuntimeData[AttributeIndex].Offset;
+			}
+
+			// Pipe data into proxy
+			Proxy->SkeletalMeshObject = SkeletalMeshObject;
+			Proxy->SkinWeightProfile = SkinWeightProfile;
+			Proxy->BoneRevisionNumber = SkeletalMesh->GetBoneTransformRevisionNumber();
+			Proxy->ParameterBuffer.AddDefaulted(ParameterBufferSize);
+			Proxy->LayeredBoneMatrixBuffer = MoveTemp(LayeredBoneMatrixBuffer);
+			
+			Proxy->AttributeBufferOffsets = MoveTemp(AttributeBufferOffsets);
+			Proxy->AttributeBuffers = MoveTemp(AttributeBuffers);	
+		}
+	}
+	
+	return Proxy;
+}
+
+
+FOptimusAdvancedSkeletonDataProviderProxy::FOptimusAdvancedSkeletonDataProviderProxy()
+{
+}
+
+bool FOptimusAdvancedSkeletonDataProviderProxy::IsValid(FValidationData const& InValidationData) const
+{
+	if (InValidationData.ParameterStructSize != ParameterBuffer.Num())
+	{
+		return false;
+	}
+	if (SkeletalMeshObject == nullptr)
+	{
+		return false;
+	}
+	const int32 LodIndex = SkeletalMeshObject->GetLOD();
+	FSkeletalMeshRenderData const& SkeletalMeshRenderData = SkeletalMeshObject->GetSkeletalMeshRenderData();
+	FSkeletalMeshLODRenderData const* LodRenderData = &SkeletalMeshRenderData.LODRenderData[LodIndex];
+	if (LodRenderData->RenderSections.Num() != InValidationData.NumInvocations)
+	{
+		return false;
+	}
+	
+	FSkinWeightVertexBuffer const* WeightBuffer = LodRenderData->GetSkinWeightVertexBuffer();
+	if (LodRenderData->SkinWeightProfilesData.ContainsProfile(SkinWeightProfile))
+	{
+		WeightBuffer = LodRenderData->SkinWeightProfilesData.GetOverrideBuffer(SkinWeightProfile);
+	}
+	
+	if (WeightBuffer == nullptr)
+	{
+		return false;
+	}
+
+	return true;
+}
+
+void FOptimusAdvancedSkeletonDataProviderProxy::AllocateResources(FRDGBuilder& GraphBuilder, FAllocationData const& InAllocationData)
+{
+	const int32 LodIndex = SkeletalMeshObject->GetLOD();
+	FSkeletalMeshRenderData const& SkeletalMeshRenderData = SkeletalMeshObject->GetSkeletalMeshRenderData();
+	FSkeletalMeshLODRenderData const* LodRenderData = &SkeletalMeshRenderData.LODRenderData[LodIndex];
+	
+	FallbackSRV = GraphBuilder.CreateSRV(FRDGBufferSRVDesc(GSystemTextures.GetDefaultStructuredBuffer(GraphBuilder, sizeof(int32))));
+	
+	LayeredBoneMatrixBuffer.AllocateResources(GraphBuilder);
+	
+	// Per-Bone Animation Attribute 
+	BuffersPerAttributePerSection.AddDefaulted(LodRenderData->RenderSections.Num());
+	BufferSRVsPerAttributePerSection.AddDefaulted(LodRenderData->RenderSections.Num());	
+	
+	for (int32 InvocationIndex = 0 ; InvocationIndex < LodRenderData->RenderSections.Num(); InvocationIndex++)
+	{
+		BuffersPerAttributePerSection[InvocationIndex].AddDefaulted(AttributeBuffers[InvocationIndex].Num());
+		BufferSRVsPerAttributePerSection[InvocationIndex].AddDefaulted(AttributeBuffers[InvocationIndex].Num());
+		
+		int32 NumBones = LodRenderData->RenderSections[InvocationIndex].BoneMap.Num();
+		
+		for (int32 AttributeIndex = 0; AttributeIndex < AttributeBuffers[InvocationIndex].Num(); ++AttributeIndex)
+		{
+			const TArray<uint8>& BoneData = AttributeBuffers[InvocationIndex][AttributeIndex];
+
+			int32 Stride = BoneData.Num() / NumBones;
+			
+			FRDGBufferRef Buffer = GraphBuilder.CreateBuffer(FRDGBufferDesc::CreateStructuredDesc(Stride,  NumBones), TEXT("AttributeBuffer"), ERDGBufferFlags::None);
+			FRDGBufferSRVRef BufferSRV = GraphBuilder.CreateSRV(Buffer);
+
+			BuffersPerAttributePerSection[InvocationIndex][AttributeIndex] = Buffer;
+			BufferSRVsPerAttributePerSection[InvocationIndex][AttributeIndex] = BufferSRV;
+
+			GraphBuilder.QueueBufferUpload(Buffer, BoneData.GetData(), BoneData.Num(), ERDGInitialDataFlags::None);	
+		}
+	}
+}
+
+struct FAdvancedSkeletonDataInterfacePermutationIds
+{
+	uint32 EnableDeformerBones = 0;
+	uint32 UnlimitedBoneInfluence = 0;
+	uint32 BoneIndexUint16 = 0;
+	uint32 BoneWeightsUint16 = 0;
+
+	FAdvancedSkeletonDataInterfacePermutationIds(FComputeKernelPermutationVector const& PermutationVector)
+	{
+		{
+			static FString Name(TEXT("ENABLE_DEFORMER_BONES"));
+			static uint32 Hash = GetTypeHash(Name);
+			EnableDeformerBones = PermutationVector.GetPermutationBits(Name, Hash, 1);
+		}
+		{
+			static FString Name(TEXT("GPUSKIN_UNLIMITED_BONE_INFLUENCE"));
+			static uint32 Hash = GetTypeHash(Name);
+			UnlimitedBoneInfluence = PermutationVector.GetPermutationBits(Name, Hash, 1);
+		}
+		{
+			static FString Name(TEXT("GPUSKIN_BONE_INDEX_UINT16"));
+			static uint32 Hash = GetTypeHash(Name);
+			BoneIndexUint16 = PermutationVector.GetPermutationBits(Name, Hash, 1);
+		}
+		{
+			static FString Name(TEXT("GPUSKIN_BONE_WEIGHTS_UINT16"));
+			static uint32 Hash = GetTypeHash(Name);
+			BoneWeightsUint16 = PermutationVector.GetPermutationBits(Name, Hash, 1);
+		}
+	}
+};
+
+void FOptimusAdvancedSkeletonDataProviderProxy::GatherPermutations(FPermutationData& InOutPermutationData) const
+{
+	const int32 LodIndex = SkeletalMeshObject->GetLOD();
+	FSkeletalMeshRenderData const& SkeletalMeshRenderData = SkeletalMeshObject->GetSkeletalMeshRenderData();
+	FSkeletalMeshLODRenderData const* LodRenderData = &SkeletalMeshRenderData.LODRenderData[LodIndex];
+
+	FAdvancedSkeletonDataInterfacePermutationIds PermutationIds(InOutPermutationData.PermutationVector);
+	for (int32 InvocationIndex = 0; InvocationIndex < InOutPermutationData.NumInvocations; ++InvocationIndex)
+	{
+		FSkelMeshRenderSection const& RenderSection = LodRenderData->RenderSections[InvocationIndex];
+
+		const bool bPreviousFrame = false;
+		FRHIShaderResourceView* BoneBufferSRV = FSkeletalMeshDeformerHelpers::GetBoneBufferForReading(SkeletalMeshObject, LodIndex, InvocationIndex, bPreviousFrame);
+
+		
+		FSkinWeightVertexBuffer const* WeightBuffer = LodRenderData->GetSkinWeightVertexBuffer();
+		if (LodRenderData->SkinWeightProfilesData.ContainsProfile(SkinWeightProfile))
+		{
+			WeightBuffer = LodRenderData->SkinWeightProfilesData.GetOverrideBuffer(SkinWeightProfile);
+		}
+		
+		check(WeightBuffer != nullptr);
+		FRHIShaderResourceView* SkinWeightBufferSRV = WeightBuffer->GetDataVertexBuffer()->GetSRV();
+		const bool bUnlimitedBoneInfluences = WeightBuffer->GetBoneInfluenceType() == GPUSkinBoneInfluenceType::UnlimitedBoneInfluence;
+		FRHIShaderResourceView* InputWeightLookupStreamSRV = bUnlimitedBoneInfluences ? WeightBuffer->GetLookupVertexBuffer()->GetSRV() : nullptr;
+		const bool bValidBones = (BoneBufferSRV != nullptr) && (SkinWeightBufferSRV != nullptr) && (!bUnlimitedBoneInfluences || InputWeightLookupStreamSRV != nullptr);
+		const bool bUse16BitBoneIndex = WeightBuffer->Use16BitBoneIndex();
+		const bool bUse16BitBoneWeights = WeightBuffer->Use16BitBoneWeight();
+
+		InOutPermutationData.PermutationIds[InvocationIndex] |= (bValidBones ? PermutationIds.EnableDeformerBones : 0);
+		InOutPermutationData.PermutationIds[InvocationIndex] |= (bUnlimitedBoneInfluences ? PermutationIds.UnlimitedBoneInfluence : 0);
+		InOutPermutationData.PermutationIds[InvocationIndex] |= (bUse16BitBoneIndex ? PermutationIds.BoneIndexUint16 : 0);
+		InOutPermutationData.PermutationIds[InvocationIndex] |= (bUse16BitBoneWeights ? PermutationIds.BoneWeightsUint16 : 0);
+	}
+}
+
+void FOptimusAdvancedSkeletonDataProviderProxy::GatherDispatchData(FDispatchData const& InDispatchData)
+{
+	const int32 LodIndex = SkeletalMeshObject->GetLOD();
+	FSkeletalMeshRenderData const& SkeletalMeshRenderData = SkeletalMeshObject->GetSkeletalMeshRenderData();
+	FSkeletalMeshLODRenderData const* LodRenderData = &SkeletalMeshRenderData.LODRenderData[LodIndex];
+
+	FRHIShaderResourceView* NullSRVBinding = GWhiteVertexBufferWithSRV->ShaderResourceViewRHI.GetReference();
+	
+
+	for (int32 InvocationIndex = 0; InvocationIndex < InDispatchData.NumInvocations; ++InvocationIndex)
+	{
+		uint8* ParameterBufferPtr = (InDispatchData.ParameterBuffer + InDispatchData.ParameterBufferOffset + InDispatchData.ParameterBufferStride * InvocationIndex);
+		
+		FSkelMeshRenderSection const& RenderSection = LodRenderData->RenderSections[InvocationIndex];
+
+		const bool bPreviousFrame = false;
+		FRHIShaderResourceView* BoneBufferSRV = FSkeletalMeshDeformerHelpers::GetBoneBufferForReading(SkeletalMeshObject, LodIndex, InvocationIndex, bPreviousFrame);
+
+		FSkinWeightVertexBuffer const* WeightBuffer = LodRenderData->GetSkinWeightVertexBuffer();
+		if (LodRenderData->SkinWeightProfilesData.ContainsProfile(SkinWeightProfile))
+		{
+			WeightBuffer = LodRenderData->SkinWeightProfilesData.GetOverrideBuffer(SkinWeightProfile);
+		}
+		check(WeightBuffer != nullptr);
+		
+		FRHIShaderResourceView* SkinWeightBufferSRV = WeightBuffer->GetDataVertexBuffer()->GetSRV();
+		const bool bUnlimitedBoneInfluences = WeightBuffer->GetBoneInfluenceType() == GPUSkinBoneInfluenceType::UnlimitedBoneInfluence;
+		FRHIShaderResourceView* InputWeightLookupStreamSRV = bUnlimitedBoneInfluences ? WeightBuffer->GetLookupVertexBuffer()->GetSRV() : nullptr;
+
+		FDefaultParameters* Parameters = reinterpret_cast<FDefaultParameters*>(ParameterBufferPtr);
+		Parameters->NumVertices = RenderSection.NumVertices;
+		Parameters->InputStreamStart = RenderSection.BaseVertexIndex;
+		Parameters->NumBoneInfluences = WeightBuffer->GetMaxBoneInfluences();
+		Parameters->InputWeightStart = (WeightBuffer->GetConstantInfluencesVertexStride() * RenderSection.GetVertexBufferIndex()) / sizeof(float);
+		Parameters->InputWeightStride = WeightBuffer->GetConstantInfluencesVertexStride();
+		Parameters->InputWeightIndexSize = WeightBuffer->GetBoneIndexByteSize() | (WeightBuffer->GetBoneWeightByteSize() << 8);
+		Parameters->BoneMatrices = BoneBufferSRV != nullptr ? BoneBufferSRV : NullSRVBinding;
+		Parameters->InputWeightStream = SkinWeightBufferSRV != nullptr ? SkinWeightBufferSRV : NullSRVBinding;
+		Parameters->InputWeightLookupStream = InputWeightLookupStreamSRV != nullptr ? InputWeightLookupStreamSRV : NullSRVBinding;
+		Parameters->LayeredBoneMatrices = LayeredBoneMatrixBuffer.BufferSRVPerSection[InvocationIndex] != nullptr ? LayeredBoneMatrixBuffer.BufferSRVPerSection[InvocationIndex] : FallbackSRV;
+
+		for(int32 AttributeIndex = 0 ; AttributeIndex < AttributeBufferOffsets.Num(); AttributeIndex++)
+		{
+			*((FRDGBufferSRV**)(ParameterBufferPtr + AttributeBufferOffsets[AttributeIndex])) = BufferSRVsPerAttributePerSection[InvocationIndex][AttributeIndex];
+		}
+	}
+}
diff --git a/Engine/Plugins/Animation/DeformerGraph/Source/OptimusCore/Private/DataInterfaces/OptimusDataInterfaceAdvancedSkeleton.h b/Engine/Plugins/Animation/DeformerGraph/Source/OptimusCore/Private/DataInterfaces/OptimusDataInterfaceAdvancedSkeleton.h
new file mode 100644
--- /dev/null
+++ b/Engine/Plugins/Animation/DeformerGraph/Source/OptimusCore/Private/DataInterfaces/OptimusDataInterfaceAdvancedSkeleton.h
@@ -0,0 +1,259 @@
+// Copyright Epic Games, Inc. All Rights Reserved.
+
+#pragma once
+
+#include "OptimusComputeDataInterface.h"
+#include "OptimusDataType.h"
+#include "OptimusDataTypeRegistry.h"
+#include "OptimusValueContainerStruct.h"
+#include "ComputeFramework/ComputeDataProvider.h"
+
+#include "OptimusDataInterfaceAdvancedSkeleton.generated.h"
+
+struct FOptimusAnimAttributeDescription;
+class FSkeletalMeshObject;
+class FAdvancedSkeletonDataInterfaceDefaultParameters;
+class USkeletalMeshComponent;
+class FSkeletalMeshLODRenderData;
+struct FReferenceSkeleton;
+
+USTRUCT()
+struct FOptimusAnimAttributeBufferDescription
+{
+	GENERATED_BODY()
+	
+	UPROPERTY(EditAnywhere, Category = "Data Interface")
+	FString Name;
+
+	UPROPERTY(EditAnywhere, Category = "Data Interface", meta=(UseInAnimAttribute))
+	FOptimusDataTypeRef DataType;
+
+	// Default value if the animation attribute is not found
+	UPROPERTY(EditAnywhere, Category = "Data Interface")
+	FOptimusValueContainerStruct DefaultValueStruct;
+
+	UPROPERTY()
+	FString HlslId;
+
+	UPROPERTY()
+	FName PinName;
+
+	static const TCHAR* PinNameDelimiter;
+	static const TCHAR* HlslIdDelimiter;
+	
+	void UpdatePinNameAndHlslId(bool bInIncludeTypeName = true);
+	
+	// Helpers
+	FOptimusAnimAttributeBufferDescription& Init(const FString& InName, const FOptimusDataTypeRef& InDataType);
+	
+private:
+	FString GetFormattedId(
+		const FString& InDelimiter,
+		bool bInIncludeTypeName) const;
+};
+
+USTRUCT()
+struct FOptimusAnimAttributeBufferArray
+{
+	GENERATED_BODY()
+
+	UPROPERTY(EditAnywhere, DisplayName = "Animation Attributes" ,Category = "Data Interface")
+	TArray<FOptimusAnimAttributeBufferDescription> InnerArray;
+
+	template <typename Predicate>
+	const FOptimusAnimAttributeBufferDescription* FindByPredicate(Predicate Pred) const
+	{
+		return InnerArray.FindByPredicate(Pred);
+	}
+
+	bool IsEmpty() const
+	{
+		return InnerArray.IsEmpty();
+	}
+
+	FOptimusAnimAttributeBufferDescription& Last(int32 IndexFromTheEnd = 0)
+	{
+		return InnerArray.Last(IndexFromTheEnd);
+	}
+
+	const FOptimusAnimAttributeBufferDescription& Last(int32 IndexFromTheEnd = 0) const
+	{
+		return InnerArray.Last(IndexFromTheEnd);
+	}
+
+	FOptimusAnimAttributeBufferArray& operator=(const TArray<FOptimusAnimAttributeBufferDescription>& Rhs)
+	{
+		InnerArray = Rhs;
+		return *this;
+	}
+	
+	int32 Num() const { return InnerArray.Num();}
+	bool IsValidIndex(int32 Index) const { return Index < InnerArray.Num() && Index >= 0; }
+	const FOptimusAnimAttributeBufferDescription& operator[](int32 InIndex) const { return InnerArray[InIndex]; }
+	FOptimusAnimAttributeBufferDescription& operator[](int32 InIndex) { return InnerArray[InIndex]; }
+	FORCEINLINE	TArray<FOptimusAnimAttributeBufferDescription>::RangedForIteratorType      begin()       { return InnerArray.begin(); }
+	FORCEINLINE	TArray<FOptimusAnimAttributeBufferDescription>::RangedForConstIteratorType begin() const { return InnerArray.begin(); }
+	FORCEINLINE	TArray<FOptimusAnimAttributeBufferDescription>::RangedForIteratorType      end()         { return InnerArray.end();   }
+	FORCEINLINE	TArray<FOptimusAnimAttributeBufferDescription>::RangedForConstIteratorType end() const   { return InnerArray.end();   }
+};
+
+
+/** Compute Framework Data Interface for skeletal data. */
+UCLASS(Category = ComputeFramework)
+class OPTIMUSCORE_API UOptimusAdvancedSkeletonDataInterface : public UOptimusComputeDataInterface
+{
+	GENERATED_BODY()
+	
+public:
+	
+	
+#if WITH_EDITOR
+	void PostEditChangeChainProperty(FPropertyChangedChainEvent& PropertyChangedEvent) override;
+#endif
+
+
+	
+	//~ Begin UOptimusComputeDataInterface Interface
+	FString GetDisplayName() const override;
+	TArray<FOptimusCDIPinDefinition> GetPinDefinitions() const override;
+	TSubclassOf<UActorComponent> GetRequiredComponentClass() const override;
+
+	void Initialize() override;
+	bool CanPinDefinitionChange() override {return true;};
+	void RegisterPropertyChangeDelegatesForOwningNode(UOptimusNode_DataInterface* InNode) override;
+
+	void OnDataTypeChanged(FName InTypeName) override;
+	//~ End UOptimusComputeDataInterface Interface
+	
+	const FOptimusAnimAttributeBufferDescription& AddAnimAttribute(const FString& InName, const FOptimusDataTypeRef& InDataType);
+
+
+	//~ Begin UComputeDataInterface Interface
+	TCHAR const* GetClassName() const override { return TEXT("AdvancedSkeleton"); }
+	void GetSupportedInputs(TArray<FShaderFunctionDefinition>& OutFunctions) const override;
+	void GetShaderParameters(TCHAR const* UID, FShaderParametersMetadataBuilder& InOutBuilder, FShaderParametersMetadataAllocations& InOutAllocations) const override;
+	void GetPermutations(FComputeKernelPermutationVector& OutPermutationVector) const override;
+	void GetShaderHash(FString& InOutKey) const override;
+	void GetHLSL(FString& OutHLSL, FString const& InDataInterfaceName) const override;
+	UComputeDataProvider* CreateDataProvider(TObjectPtr<UObject> InBinding, uint64 InInputMask, uint64 InOutputMask) const override;
+	//~ End UComputeDataInterface Interface
+	
+	UPROPERTY(EditAnywhere, Category = "Data Interface")
+	FName SkinWeightProfile = NAME_None;
+	
+	UPROPERTY(EditAnywhere, Category = "Data Interface")
+	bool bEnableLayeredSkinning = true;
+	
+	UPROPERTY(EditAnywhere, Category = "Data Interface", meta = (ShowOnlyInnerProperties))
+	FOptimusAnimAttributeBufferArray AttributeBufferArray;
+
+private:
+	FString GetUnusedAttributeName(int32 CurrentAttributeIndex ,const FString& InName) const;
+	void UpdateAttributePinNamesAndHlslIds();
+	TArray<FOptimusCDIPinDefinition> GetPinDefinitions_Internal(bool bGetAllPossiblePins = false, int32 AttributeIndexToExclude = INDEX_NONE) const;	
+	static TCHAR const* SkeletonTemplateFilePath;
+	static TCHAR const* AttributeTemplateFilePath;
+	FOnPinDefinitionChanged OnPinDefinitionChangedDelegate;
+	FOnPinDefinitionRenamed OnPinDefinitionRenamedDelegate;
+	FSimpleDelegate OnDisplayNameChangedDelegate;
+};
+
+// Runtime data with cached values baked out from AttributeDescription
+struct FOptimusAnimAttributeBufferRuntimeData
+{
+	FOptimusAnimAttributeBufferRuntimeData() = default;
+	
+	FOptimusAnimAttributeBufferRuntimeData(const FOptimusAnimAttributeBufferDescription& InDescription);
+	
+	FName Name;
+	FName HlslId;
+
+	int32 Offset = INDEX_NONE;
+
+	int32 Size = 0;
+
+	FOptimusDataTypeRegistry::PropertyValueConvertFuncT	ConvertFunc = nullptr;
+
+	UScriptStruct* AttributeType = nullptr;
+
+	FShaderValueContainer CachedDefaultValue;
+};
+
+struct FOptimusBoneTransformBuffer
+{
+	TArray<FRDGBufferRef> BufferRefPerSection;
+	TArray<FRDGBufferSRVRef> BufferSRVPerSection;
+	TArray<TArray<uint8>> BufferData;
+	TArray<int32> NumBones;
+
+	void SetData(FSkeletalMeshLODRenderData const* InLodRenderData, const TArray<FTransform>& InBoneTransforms);
+
+	void AllocateResources(FRDGBuilder& GraphBuilder);
+};
+
+/** Compute Framework Data Provider for reading skeletal mesh. */
+UCLASS(BlueprintType, editinlinenew, Category = ComputeFramework)
+class UOptimusAdvancedSkeletonDataProvider : public UComputeDataProvider
+{
+	GENERATED_BODY()
+
+public:
+	
+	UPROPERTY(EditAnywhere, BlueprintReadWrite, Category = Binding)
+	TObjectPtr<USkeletalMeshComponent> SkeletalMesh = nullptr;
+
+	TArray<FOptimusAnimAttributeBufferRuntimeData> AttributeBufferRuntimeData;
+
+	int32 ParameterBufferSize = 0;
+
+	FName SkinWeightProfile = NAME_None;
+	bool bSkinWeightBufferReady = false;
+	
+	bool bEnableLayeredSkinning = true;
+	bool bIsLayeredSkinInitialized = false;
+	TSet<uint32> CachedWeightedBoneIndices;
+	TArray<int32> CachedBoundaryBoneIndex;
+	TArray<FTransform> CachedLayerSpaceInitialBoneTransform; 
+	
+	//~ Begin UComputeDataProvider Interface
+	FComputeDataProviderRenderProxy* GetRenderProxy() override;
+	//~ End UComputeDataProvider Interface
+
+	void Init(
+		const UOptimusAdvancedSkeletonDataInterface* InDataInterface,
+		USkeletalMeshComponent* InSkeletalMesh
+	);
+
+	void ComputeBoneTransformsForLayeredSkinning(FOptimusBoneTransformBuffer& OutBoneBuffer, FSkeletalMeshLODRenderData const* InLodRenderData, const FReferenceSkeleton& InRefSkeleton);
+};
+
+class FOptimusAdvancedSkeletonDataProviderProxy : public FComputeDataProviderRenderProxy
+{
+public:
+	FOptimusAdvancedSkeletonDataProviderProxy();
+
+	//~ Begin FComputeDataProviderRenderProxy Interface
+	bool IsValid(FValidationData const& InValidationData) const override;
+	void AllocateResources(FRDGBuilder& GraphBuilder, FAllocationData const& InAllocationData) override;
+	void GatherPermutations(FPermutationData& InOutPermutationData) const override;
+	void GatherDispatchData(FDispatchData const& InDispatchData) override;
+	//~ End FComputeDataProviderRenderProxy Interface
+
+	using FDefaultParameters = FAdvancedSkeletonDataInterfaceDefaultParameters;
+
+	FSkeletalMeshObject* SkeletalMeshObject = nullptr;
+	FName SkinWeightProfile = NAME_None;
+	uint32 BoneRevisionNumber = 0;
+	
+	TArray<uint8> ParameterBuffer;
+	
+	FOptimusBoneTransformBuffer LayeredBoneMatrixBuffer;
+	
+	TArray<int32> AttributeBufferOffsets;
+	
+	TArray<TArray<FRDGBufferRef>> BuffersPerAttributePerSection;
+	TArray<TArray<FRDGBufferSRVRef>> BufferSRVsPerAttributePerSection;
+	TArray<TArray<TArray<uint8>>> AttributeBuffers;
+	
+	FRDGBufferSRVRef FallbackSRV = nullptr;
+};
diff --git a/Engine/Plugins/Animation/DeformerGraph/Source/OptimusCore/Private/DataInterfaces/OptimusDataInterfaceAnimAttribute.cpp b/Engine/Plugins/Animation/DeformerGraph/Source/OptimusCore/Private/DataInterfaces/OptimusDataInterfaceAnimAttribute.cpp
--- a/Engine/Plugins/Animation/DeformerGraph/Source/OptimusCore/Private/DataInterfaces/OptimusDataInterfaceAnimAttribute.cpp
+++ b/Engine/Plugins/Animation/DeformerGraph/Source/OptimusCore/Private/DataInterfaces/OptimusDataInterfaceAnimAttribute.cpp
@@ -112,7 +112,7 @@ void UOptimusAnimAttributeDataInterface::PostEditChangeChainProperty(FPropertyCh
 				{
 					ChangedAttribute.Name = TEXT("EmptyName");
 				}
-				
+
 				for (int32 Index = 0; Index < AttributeArray.Num(); Index++)
 				{
 					const FOptimusAnimAttributeDescription& Attribute = AttributeArray[Index];
@@ -233,7 +233,7 @@ void UOptimusAnimAttributeDataInterface::Initialize()
 			FOptimusDataTypeRegistry::Get().FindType(*FFloatProperty::StaticClass()) );
 }
 
-void UOptimusAnimAttributeDataInterface::RegisterDynamicPinDelegatesForOwningNode(UOptimusNode_DataInterface* InNode)
+void UOptimusAnimAttributeDataInterface::RegisterPropertyChangeDelegatesForOwningNode(UOptimusNode_DataInterface* InNode)
 {
 	OnPinDefinitionChangedDelegate.BindUObject(InNode, &UOptimusNode_DataInterface::RecreatePinsFromPinDefinitions);
 	OnPinDefinitionRenamedDelegate.BindUObject(InNode, &UOptimusNode_DataInterface::RenamePinFromPinDefinition);
diff --git a/Engine/Plugins/Animation/DeformerGraph/Source/OptimusCore/Private/DataInterfaces/OptimusDataInterfaceAnimAttribute.h b/Engine/Plugins/Animation/DeformerGraph/Source/OptimusCore/Private/DataInterfaces/OptimusDataInterfaceAnimAttribute.h
--- a/Engine/Plugins/Animation/DeformerGraph/Source/OptimusCore/Private/DataInterfaces/OptimusDataInterfaceAnimAttribute.h
+++ b/Engine/Plugins/Animation/DeformerGraph/Source/OptimusCore/Private/DataInterfaces/OptimusDataInterfaceAnimAttribute.h
@@ -128,7 +128,7 @@ public:
 	
 	void Initialize() override;
 	bool CanPinDefinitionChange() override {return true;};
-	void RegisterDynamicPinDelegatesForOwningNode(UOptimusNode_DataInterface* InNode) override;
+	void RegisterPropertyChangeDelegatesForOwningNode(UOptimusNode_DataInterface* InNode) override;
 	
 	//~ End UOptimusComputeDataInterface Interface
 	
diff --git a/Engine/Plugins/Animation/DeformerGraph/Source/OptimusCore/Private/Nodes/OptimusNode_DataInterface.cpp b/Engine/Plugins/Animation/DeformerGraph/Source/OptimusCore/Private/Nodes/OptimusNode_DataInterface.cpp
--- a/Engine/Plugins/Animation/DeformerGraph/Source/OptimusCore/Private/Nodes/OptimusNode_DataInterface.cpp
+++ b/Engine/Plugins/Animation/DeformerGraph/Source/OptimusCore/Private/Nodes/OptimusNode_DataInterface.cpp
@@ -187,6 +187,14 @@ void UOptimusNode_DataInterface::RenamePinFromPinDefinition(FName InOld, FName I
 	SetPinName(Pin, InNew);
 }
 
+void UOptimusNode_DataInterface::UpdateDisplayNameFromDataInterface()
+{
+	if (ensure(DataInterfaceData))
+	{
+		SetDisplayName(FText::FromString(DataInterfaceData->GetDisplayName()));
+	}
+}
+
 void UOptimusNode_DataInterface::InitializeTransientData()
 {
 	if (ensure(DataInterfaceData))
@@ -194,8 +202,9 @@ void UOptimusNode_DataInterface::InitializeTransientData()
 		if (DataInterfaceData->CanPinDefinitionChange())
 		{
 			EnableDynamicPins();
-			DataInterfaceData->RegisterDynamicPinDelegatesForOwningNode(this);
 		}
+		
+		DataInterfaceData->RegisterPropertyChangeDelegatesForOwningNode(this);
 	}
 }
 
diff --git a/Engine/Plugins/Animation/DeformerGraph/Source/OptimusCore/Private/Nodes/OptimusNode_DataInterface.h b/Engine/Plugins/Animation/DeformerGraph/Source/OptimusCore/Private/Nodes/OptimusNode_DataInterface.h
--- a/Engine/Plugins/Animation/DeformerGraph/Source/OptimusCore/Private/Nodes/OptimusNode_DataInterface.h
+++ b/Engine/Plugins/Animation/DeformerGraph/Source/OptimusCore/Private/Nodes/OptimusNode_DataInterface.h
@@ -33,6 +33,7 @@ public:
 
 	void RecreatePinsFromPinDefinitions();
 	void RenamePinFromPinDefinition(FName InOld, FName InNew);
+	void UpdateDisplayNameFromDataInterface();
 
 	// -- UOptimusNode overrides
 	FName GetNodeCategory() const override 
diff --git a/Engine/Plugins/Animation/DeformerGraph/Source/OptimusCore/Private/OptimusDataTypeRegistry.cpp b/Engine/Plugins/Animation/DeformerGraph/Source/OptimusCore/Private/OptimusDataTypeRegistry.cpp
--- a/Engine/Plugins/Animation/DeformerGraph/Source/OptimusCore/Private/OptimusDataTypeRegistry.cpp
+++ b/Engine/Plugins/Animation/DeformerGraph/Source/OptimusCore/Private/OptimusDataTypeRegistry.cpp
@@ -18,6 +18,8 @@
 
 #include "UObject/UnrealType.h"
 
+const TCHAR* FOptimusDataTypeRegistry::Matrix34TypeName = TEXT("3x4 float");
+
 static const TMap<FName, UScriptStruct*>& GetBuiltInAttributeTypes()
 {
 	static const TMap<FName, UScriptStruct*> BuiltInAttributeTypes =
@@ -454,7 +456,7 @@ void FOptimusDataTypeRegistry::RegisterBuiltinTypes()
 
 	// HLSL types
 	Registry.RegisterType(
-		FName("3x4 Float"),
+		Matrix34TypeName,
 		FText::FromString(TEXT("Matrix 3x4")),
 		FShaderValueType::Get(EShaderFundamentalType::Float, 3, 4),
 		FName("float3x4"),
@@ -1200,7 +1202,7 @@ bool FOptimusDataTypeRegistry::RegisterType(
 			}
 		}
 
-		const FName TypeName(*FString::Printf(TEXT("F%s"), *InStructType->GetName()));
+		const FName TypeName = Optimus::GetTypeName(InStructType, true);
 
 		PropertyCreateFuncT PropertyCreateFunc;
 		const int32 ExpectedShaderValueSize = InShaderValueType->GetResourceElementSize();
@@ -1258,7 +1260,7 @@ bool FOptimusDataTypeRegistry::RegisterType(
 {
 	if (ensure(InStructType))
 	{
-		const FName TypeName(*FString::Printf(TEXT("F%s"), *InStructType->GetName()));
+		const FName TypeName = Optimus::GetTypeName(InStructType, true);
 
 		PropertyCreateFuncT PropertyCreateFunc;
 		if (EnumHasAnyFlags(InUsageFlags, EOptimusDataTypeUsageFlags::Variable))
diff --git a/Engine/Plugins/Animation/DeformerGraph/Source/OptimusCore/Private/OptimusHelpers.cpp b/Engine/Plugins/Animation/DeformerGraph/Source/OptimusCore/Private/OptimusHelpers.cpp
--- a/Engine/Plugins/Animation/DeformerGraph/Source/OptimusCore/Private/OptimusHelpers.cpp
+++ b/Engine/Plugins/Animation/DeformerGraph/Source/OptimusCore/Private/OptimusHelpers.cpp
@@ -2,6 +2,7 @@
 
 #include "OptimusHelpers.h"
 
+#include "Matrix3x4.h"
 #include "OptimusDataTypeRegistry.h"
 #include "ComputeFramework/ShaderParamTypeDefinition.h"
 #include "ShaderParameterMetadataBuilder.h"
@@ -66,6 +67,36 @@ FName Optimus::GetSanitizedNameForHlsl(FName InName)
 	return *Name;
 }
 
+const void Optimus::ConvertFTransformToFMatrix3x4(const FTransform& InTransform, FShaderValueContainerView OutShaderValue)
+{
+	// Code taken from FGPUBaseSkinVertexFactory::FShaderDataType::UpdateBoneData
+	if (ensure(OutShaderValue.ShaderValue.Num() == FShaderValueType::Get(EShaderFundamentalType::Float, 3, 4)->GetResourceElementSize()))
+	{
+		uint8* ShaderValuePtr = OutShaderValue.ShaderValue.GetData();
+
+		FMatrix3x4& ShaderMat = *((FMatrix3x4*)(ShaderValuePtr));
+		const FMatrix44f& Matrix = ConvertFTransformToFMatrix44f(InTransform);
+		// Explicit SIMD implementation seems to be faster than standard implementation
+#if PLATFORM_ENABLE_VECTORINTRINSICS
+		VectorRegister4Float InRow0 = VectorLoadAligned(&(Matrix.M[0][0]));
+		VectorRegister4Float InRow1 = VectorLoadAligned(&(Matrix.M[1][0]));
+		VectorRegister4Float InRow2 = VectorLoadAligned(&(Matrix.M[2][0]));
+		VectorRegister4Float InRow3 = VectorLoadAligned(&(Matrix.M[3][0]));
+
+		VectorRegister4Float Temp0 = VectorShuffle(InRow0, InRow1, 0, 1, 0, 1);
+		VectorRegister4Float Temp1 = VectorShuffle(InRow2, InRow3, 0, 1, 0, 1);
+		VectorRegister4Float Temp2 = VectorShuffle(InRow0, InRow1, 2, 3, 2, 3);
+		VectorRegister4Float Temp3 = VectorShuffle(InRow2, InRow3, 2, 3, 2, 3);
+
+		VectorStoreAligned(VectorShuffle(Temp0, Temp1, 0, 2, 0, 2), &(ShaderMat.M[0][0]));
+		VectorStoreAligned(VectorShuffle(Temp0, Temp1, 1, 3, 1, 3), &(ShaderMat.M[1][0]));
+		VectorStoreAligned(VectorShuffle(Temp2, Temp3, 0, 2, 0, 2), &(ShaderMat.M[2][0]));
+#else
+		Matrix.To3x4MatrixTranspose((float*)ShaderMat.M);
+#endif
+	}
+}
+
 bool Optimus::RenameObject(UObject* InObjectToRename, const TCHAR* InNewName, UObject* InNewOuter)
 {
 	return InObjectToRename->Rename(InNewName, InNewOuter, REN_DoNotDirty | REN_DontCreateRedirectors | REN_NonTransactional);
diff --git a/Engine/Plugins/Animation/DeformerGraph/Source/OptimusCore/Public/OptimusComputeDataInterface.h b/Engine/Plugins/Animation/DeformerGraph/Source/OptimusCore/Public/OptimusComputeDataInterface.h
--- a/Engine/Plugins/Animation/DeformerGraph/Source/OptimusCore/Public/OptimusComputeDataInterface.h
+++ b/Engine/Plugins/Animation/DeformerGraph/Source/OptimusCore/Public/OptimusComputeDataInterface.h
@@ -155,9 +155,9 @@ public:
 	/// Whether the data interface allow users to add / remove pins
 	virtual bool CanPinDefinitionChange() { return false; };
 
-	/// Register delegates for data interface node to update its pins when the data interface changes
-	virtual void RegisterDynamicPinDelegatesForOwningNode(UOptimusNode_DataInterface* InNode) {};
-
+	/// Register delegates for data interface node to update when the data interface changes
+	virtual void RegisterPropertyChangeDelegatesForOwningNode(UOptimusNode_DataInterface* InNode) {};
+	
 	/**
 	 * @return Returns the component type that this data interface operates on.
 	 */
diff --git a/Engine/Plugins/Animation/DeformerGraph/Source/OptimusCore/Public/OptimusDataTypeRegistry.h b/Engine/Plugins/Animation/DeformerGraph/Source/OptimusCore/Public/OptimusDataTypeRegistry.h
--- a/Engine/Plugins/Animation/DeformerGraph/Source/OptimusCore/Public/OptimusDataTypeRegistry.h
+++ b/Engine/Plugins/Animation/DeformerGraph/Source/OptimusCore/Public/OptimusDataTypeRegistry.h
@@ -36,6 +36,9 @@ public:
 		int32 ElementShaderValueSize;
 		int32 ShaderValueOffset;
 	};
+
+	static const TCHAR* Matrix34TypeName;
+
 	
 	virtual ~FOptimusDataTypeRegistry() override;
 
diff --git a/Engine/Plugins/Animation/DeformerGraph/Source/OptimusCore/Public/OptimusHelpers.h b/Engine/Plugins/Animation/DeformerGraph/Source/OptimusCore/Public/OptimusHelpers.h
--- a/Engine/Plugins/Animation/DeformerGraph/Source/OptimusCore/Public/OptimusHelpers.h
+++ b/Engine/Plugins/Animation/DeformerGraph/Source/OptimusCore/Public/OptimusHelpers.h
@@ -11,6 +11,7 @@
 #include "UObject/Package.h"
 #include "UObject/Class.h"
 
+struct FShaderValueContainerView;
 class FShaderParametersMetadata;
 class FShaderParametersMetadataBuilder;
 struct FShaderValueTypeHandle;
@@ -91,7 +92,9 @@ namespace Optimus
 	{
 		return TransformConverter<FMatrix44f>::Convert<FMatrix44d>(InTransform.ToMatrixWithScale());	
 	};
-
+	
+	const void ConvertFTransformToFMatrix3x4(const FTransform& InTransform, FShaderValueContainerView OutShaderValue);
+	
 	bool RenameObject(UObject* InObjectToRename, const TCHAR* InNewName, UObject* InNewOuter);
 
 	/** Use this function to remove objects during postload safely */
diff --git a/Engine/Plugins/Runtime/ComputeFramework/Source/ComputeFramework/Public/ComputeFramework/ShaderParameterMetadataAllocation.h b/Engine/Plugins/Runtime/ComputeFramework/Source/ComputeFramework/Public/ComputeFramework/ShaderParameterMetadataAllocation.h
--- a/Engine/Plugins/Runtime/ComputeFramework/Source/ComputeFramework/Public/ComputeFramework/ShaderParameterMetadataAllocation.h
+++ b/Engine/Plugins/Runtime/ComputeFramework/Source/ComputeFramework/Public/ComputeFramework/ShaderParameterMetadataAllocation.h
@@ -13,7 +13,7 @@ class FShaderParametersMetadata;
  * This structure tracks the allocations, and should be released only after we are
  * done with the associated FShaderParametersMetadata.
  */
-struct FShaderParametersMetadataAllocations
+struct COMPUTEFRAMEWORK_API FShaderParametersMetadataAllocations
 {
 	/** Allocated metadata. Should include the parent metadata allocation. */
 	TArray<FShaderParametersMetadata*> ShaderParameterMetadatas;
diff --git a/Engine/Source/Runtime/Engine/Public/Animation/SkinWeightProfileManager.h b/Engine/Source/Runtime/Engine/Public/Animation/SkinWeightProfileManager.h
--- a/Engine/Source/Runtime/Engine/Public/Animation/SkinWeightProfileManager.h
+++ b/Engine/Source/Runtime/Engine/Public/Animation/SkinWeightProfileManager.h
@@ -121,12 +121,12 @@ protected:
 public: 
 	static void OnStartup();
 	static void OnShutdown();
-	static FSkinWeightProfileManager* Get(UWorld* World);
+	static ENGINE_API FSkinWeightProfileManager* Get(UWorld* World);
 
 	FSkinWeightProfileManager(UWorld* InWorld);
 	virtual ~FSkinWeightProfileManager() {}
 
-	void RequestSkinWeightProfile(FName InProfileName, USkinnedAsset* SkinnedAsset, UObject* Requester, FRequestFinished& Callback, int32 LODIndex = INDEX_NONE);
+	void ENGINE_API RequestSkinWeightProfile(FName InProfileName, USkinnedAsset* SkinnedAsset, UObject* Requester, FRequestFinished& Callback, int32 LODIndex = INDEX_NONE);
 	void CancelSkinWeightProfileRequest(UObject* Requester);
 	
 	void DoTick(float DeltaTime, ENamedThreads::Type CurrentThread, const FGraphEventRef& MyCompletionGraphEvent);
