[Deformer Graph] exclude struct type from per bone anim attribute

#rnx

diff --git a/Engine/Plugins/Animation/DeformerGraph/Source/OptimusCore/Private/DataInterfaces/OptimusDataInterfaceAdvancedSkeleton.h b/Engine/Plugins/Animation/DeformerGraph/Source/OptimusCore/Private/DataInterfaces/OptimusDataInterfaceAdvancedSkeleton.h
--- a/Engine/Plugins/Animation/DeformerGraph/Source/OptimusCore/Private/DataInterfaces/OptimusDataInterfaceAdvancedSkeleton.h
+++ b/Engine/Plugins/Animation/DeformerGraph/Source/OptimusCore/Private/DataInterfaces/OptimusDataInterfaceAdvancedSkeleton.h
@@ -25,7 +25,7 @@ struct FOptimusAnimAttributeBufferDescription
 	UPROPERTY(EditAnywhere, Category = "Data Interface")
 	FString Name;
 
-	UPROPERTY(EditAnywhere, Category = "Data Interface", meta=(UseInAnimAttribute))
+	UPROPERTY(EditAnywhere, Category = "Data Interface", meta=(UseInPerBoneAnimAttribute))
 	FOptimusDataTypeRef DataType;
 
 	// Default value if the animation attribute is not found
diff --git a/Engine/Plugins/Animation/DeformerGraph/Source/OptimusCore/Private/Nodes/OptimusNode_CustomComputeKernel.cpp b/Engine/Plugins/Animation/DeformerGraph/Source/OptimusCore/Private/Nodes/OptimusNode_CustomComputeKernel.cpp
--- a/Engine/Plugins/Animation/DeformerGraph/Source/OptimusCore/Private/Nodes/OptimusNode_CustomComputeKernel.cpp
+++ b/Engine/Plugins/Animation/DeformerGraph/Source/OptimusCore/Private/Nodes/OptimusNode_CustomComputeKernel.cpp
@@ -275,7 +275,7 @@ EOptimusDataTypeUsageFlags UOptimusNode_CustomComputeKernel::GetTypeUsageFlags(c
 		return EOptimusDataTypeUsageFlags::Variable | EOptimusDataTypeUsageFlags::AnimAttributes | EOptimusDataTypeUsageFlags::DataInterfaceOutput;
 	}
 
-	return EOptimusDataTypeUsageFlags::Resource;
+	return EOptimusDataTypeUsageFlags::Resource | EOptimusDataTypeUsageFlags::PerBoneAnimAttribute;
 }
 
 
diff --git a/Engine/Plugins/Animation/DeformerGraph/Source/OptimusCore/Private/OptimusDataTypeRegistry.cpp b/Engine/Plugins/Animation/DeformerGraph/Source/OptimusCore/Private/OptimusDataTypeRegistry.cpp
--- a/Engine/Plugins/Animation/DeformerGraph/Source/OptimusCore/Private/OptimusDataTypeRegistry.cpp
+++ b/Engine/Plugins/Animation/DeformerGraph/Source/OptimusCore/Private/OptimusDataTypeRegistry.cpp
@@ -271,7 +271,7 @@ void FOptimusDataTypeRegistry::RegisterBuiltinTypes()
 	    },
 		ConvertPropertyValuePOD<int32, int32>,
 		FName(TEXT("int")), {}, 
-	    EOptimusDataTypeUsageFlags::Resource | EOptimusDataTypeUsageFlags::Variable | EOptimusDataTypeUsageFlags::AnimAttributes | EOptimusDataTypeUsageFlags::DataInterfaceOutput | EOptimusDataTypeUsageFlags::PinType);
+	    EOptimusDataTypeUsageFlags::Resource | EOptimusDataTypeUsageFlags::Variable | EOptimusDataTypeUsageFlags::AnimAttributes | EOptimusDataTypeUsageFlags::DataInterfaceOutput | EOptimusDataTypeUsageFlags::PinType | EOptimusDataTypeUsageFlags::PerBoneAnimAttribute);
 
 	// FIntPoint -> int2
 	Registry.RegisterType(
@@ -359,7 +359,7 @@ void FOptimusDataTypeRegistry::RegisterBuiltinTypes()
 	    },
 		ConvertPropertyValuePOD<float, float>,
 		FName(TEXT("real")), {}, 
-	    EOptimusDataTypeUsageFlags::Resource | EOptimusDataTypeUsageFlags::AnimAttributes | EOptimusDataTypeUsageFlags::DataInterfaceOutput | EOptimusDataTypeUsageFlags::PinType);
+	    EOptimusDataTypeUsageFlags::Resource | EOptimusDataTypeUsageFlags::AnimAttributes | EOptimusDataTypeUsageFlags::DataInterfaceOutput | EOptimusDataTypeUsageFlags::PinType | EOptimusDataTypeUsageFlags::PerBoneAnimAttribute);
 
 	// double -> float 
 	Registry.RegisterType(
@@ -397,7 +397,7 @@ void FOptimusDataTypeRegistry::RegisterBuiltinTypes()
 		FShaderValueType::Get(EShaderFundamentalType::Float, 3),
 		{},
 		bShowElements,
-	    EOptimusDataTypeUsageFlags::Resource | EOptimusDataTypeUsageFlags::Variable | EOptimusDataTypeUsageFlags::AnimAttributes | EOptimusDataTypeUsageFlags::DataInterfaceOutput | EOptimusDataTypeUsageFlags::PinType);
+	    EOptimusDataTypeUsageFlags::Resource | EOptimusDataTypeUsageFlags::Variable | EOptimusDataTypeUsageFlags::AnimAttributes | EOptimusDataTypeUsageFlags::DataInterfaceOutput | EOptimusDataTypeUsageFlags::PinType | EOptimusDataTypeUsageFlags::PerBoneAnimAttribute);
 
 	// FVector4 -> float4
 	Registry.RegisterType(
@@ -422,7 +422,7 @@ void FOptimusDataTypeRegistry::RegisterBuiltinTypes()
 		FShaderValueType::Get(EShaderFundamentalType::Float, 4),
 		{},
 		bShowElements,
-		EOptimusDataTypeUsageFlags::Resource | EOptimusDataTypeUsageFlags::Variable | EOptimusDataTypeUsageFlags::AnimAttributes | EOptimusDataTypeUsageFlags::DataInterfaceOutput | EOptimusDataTypeUsageFlags::PinType);
+		EOptimusDataTypeUsageFlags::Resource | EOptimusDataTypeUsageFlags::Variable | EOptimusDataTypeUsageFlags::AnimAttributes | EOptimusDataTypeUsageFlags::DataInterfaceOutput | EOptimusDataTypeUsageFlags::PinType | EOptimusDataTypeUsageFlags::PerBoneAnimAttribute);
 
 	// FRotator -> float3x3
 	Registry.RegisterType(
@@ -430,7 +430,7 @@ void FOptimusDataTypeRegistry::RegisterBuiltinTypes()
 	    FShaderValueType::Get(EShaderFundamentalType::Float, 3, 3),
 	    {},
 	    bShowElements,
-	    EOptimusDataTypeUsageFlags::Resource | EOptimusDataTypeUsageFlags::Variable | EOptimusDataTypeUsageFlags::DataInterfaceOutput | EOptimusDataTypeUsageFlags::PinType);
+	    EOptimusDataTypeUsageFlags::Resource | EOptimusDataTypeUsageFlags::Variable | EOptimusDataTypeUsageFlags::DataInterfaceOutput | EOptimusDataTypeUsageFlags::PinType | EOptimusDataTypeUsageFlags::PerBoneAnimAttribute);
 
 	// FTransform -> float4x4
 	Registry.RegisterType(
@@ -452,7 +452,7 @@ void FOptimusDataTypeRegistry::RegisterBuiltinTypes()
 	    },
 	    {},
 	    bHideElements,
-	    EOptimusDataTypeUsageFlags::Resource | EOptimusDataTypeUsageFlags::Variable | EOptimusDataTypeUsageFlags::AnimAttributes | EOptimusDataTypeUsageFlags::DataInterfaceOutput | EOptimusDataTypeUsageFlags::PinType);
+	    EOptimusDataTypeUsageFlags::Resource | EOptimusDataTypeUsageFlags::Variable | EOptimusDataTypeUsageFlags::AnimAttributes | EOptimusDataTypeUsageFlags::DataInterfaceOutput | EOptimusDataTypeUsageFlags::PinType | EOptimusDataTypeUsageFlags::PerBoneAnimAttribute);
 
 	// HLSL types
 	Registry.RegisterType(
@@ -947,7 +947,7 @@ bool FOptimusDataTypeRegistry::RegisterArrayTypeIfApplicable(FOptimusDataTypeHan
 		ArrayDataType->DisplayName = FText::FromString(InElementDataType->DisplayName.ToString() + TEXT(" Array"));
 		ArrayDataType->ShaderValueType = FShaderValueType::MakeDynamicArrayType(InnerDataTypeForStructuredBuffer->ShaderValueType);
 		
-		EnumRemoveFlags(ArrayDataType->UsageFlags, EOptimusDataTypeUsageFlags::Resource | EOptimusDataTypeUsageFlags::AnimAttributes);
+		EnumRemoveFlags(ArrayDataType->UsageFlags, EOptimusDataTypeUsageFlags::Resource | EOptimusDataTypeUsageFlags::AnimAttributes | EOptimusDataTypeUsageFlags::PerBoneAnimAttribute);
 		
 		// Unused field
 		ArrayDataType->ShaderValueSize = 0;	
diff --git a/Engine/Plugins/Animation/DeformerGraph/Source/OptimusCore/Public/OptimusDataType.h b/Engine/Plugins/Animation/DeformerGraph/Source/OptimusCore/Public/OptimusDataType.h
--- a/Engine/Plugins/Animation/DeformerGraph/Source/OptimusCore/Public/OptimusDataType.h
+++ b/Engine/Plugins/Animation/DeformerGraph/Source/OptimusCore/Public/OptimusDataType.h
@@ -15,13 +15,14 @@
 UENUM(meta = (Bitflags))
 enum class EOptimusDataTypeUsageFlags : uint8
 {
-	None				= 0,
+	None				 = 0,
 	
-	Resource			= 1 << 0,		/** This type can be used in a resource */
-	Variable			= 1 << 1,		/** This type can be used in a variable */
-	AnimAttributes      = 1 << 2,       /** This type can be used to query an anim attribute*/
-	DataInterfaceOutput = 1 << 3,       /** This type can be used as output of a data interface*/
-	PinType				= 1 << 4,       /** This type can be used as pin type*/
+	Resource			 = 1 << 0,		/** This type can be used in a resource */
+	Variable			 = 1 << 1,		/** This type can be used in a variable */
+	AnimAttributes       = 1 << 2,       /** This type can be used to query a single anim attribute on a single bone*/
+	DataInterfaceOutput  = 1 << 3,       /** This type can be used as output of a data interface*/
+	PinType				 = 1 << 4,       /** This type can be used as pin type*/
+	PerBoneAnimAttribute = 1 << 5,       /** This type can be used as to query a per-bone anim attribute to produce a bone buffer */
 };
 ENUM_CLASS_FLAGS(EOptimusDataTypeUsageFlags)
 
diff --git a/Engine/Plugins/Animation/DeformerGraph/Source/OptimusEditor/Private/OptimusDetailsCustomization.cpp b/Engine/Plugins/Animation/DeformerGraph/Source/OptimusEditor/Private/OptimusDetailsCustomization.cpp
--- a/Engine/Plugins/Animation/DeformerGraph/Source/OptimusEditor/Private/OptimusDetailsCustomization.cpp
+++ b/Engine/Plugins/Animation/DeformerGraph/Source/OptimusEditor/Private/OptimusDetailsCustomization.cpp
@@ -86,6 +86,10 @@ void FOptimusDataTypeRefCustomization::CustomizeHeader(
 		{
 			UsageMask |= EOptimusDataTypeUsageFlags::AnimAttributes;
 		}
+		if (InPropertyHandle->HasMetaData(FName(TEXT("UseInPerBoneAnimAttribute"))))
+		{
+			UsageMask |= EOptimusDataTypeUsageFlags::PerBoneAnimAttribute;
+		}
 
 		return UsageMask;
 	};
