[Optimus Test] Fixed duplicated advanced skeleton DI node not reacting to new attributes being added.

diff --git a/Engine/Plugins/Animation/DeformerGraph/Source/OptimusCore/Private/Actions/OptimusNodeGraphActions.cpp b/Engine/Plugins/Animation/DeformerGraph/Source/OptimusCore/Private/Actions/OptimusNodeGraphActions.cpp
--- a/Engine/Plugins/Animation/DeformerGraph/Source/OptimusCore/Private/Actions/OptimusNodeGraphActions.cpp
+++ b/Engine/Plugins/Animation/DeformerGraph/Source/OptimusCore/Private/Actions/OptimusNodeGraphActions.cpp
@@ -349,13 +349,14 @@ FOptimusNodeGraphAction_DuplicateNode::FOptimusNodeGraphAction_DuplicateNode(
 {
 	if (ensure(InSourceNode != nullptr))
 	{
+		SourceNodePath = InSourceNode->GetNodePath();
 		GraphPath = InTargetGraphPath;
 		NodeName = InNodeName;
 		NodeClassPath = InSourceNode->GetClass()->GetPathName();
 		ConfigureNodeFunc = InConfigureNodeFunc;
-		
-		// Photocopy the node.
-		Optimus::FBinaryObjectWriter(InSourceNode, NodeData);
+
+		FMemoryWriter NodeArchive(CachedNodeData);
+		InSourceNode->ExportState(NodeArchive);
 	}
 }
 
@@ -372,10 +373,26 @@ bool FOptimusNodeGraphAction_DuplicateNode::Do(
 	IOptimusPathResolver* InRoot
 	)
 {
-	auto BootstrapNodeFunc = [this](UOptimusNode* InNode) -> bool
+	UOptimusNode* SourceNode = InRoot->ResolveNodePath(SourceNodePath);
+
+	TArray<uint8> NodeData;
+	if (!SourceNode)
+	{
+		// Used the cached data if we are duplicating from a clipboard graph
+		NodeData = CachedNodeData;
+	}
+	else
 	{
-		// Paste the data from the old node onto the new one.
-		Optimus::FBinaryObjectReader(InNode, NodeData);
+		FMemoryWriter NodeArchive(NodeData);
+		SourceNode->ExportState(NodeArchive);
+	}
+	
+	
+	auto BootstrapNodeFunc = [NodeData, this](UOptimusNode* InNode) -> bool
+	{
+		FMemoryReader NodeArchive(NodeData);
+		InNode->ImportState(NodeArchive);
+		
 		if (!ConfigureNodeFunc)
 		{
 			return true;
@@ -473,9 +490,8 @@ bool FOptimusNodeGraphAction_RemoveNode::Do(
 		return false;
 	}
 
-	// Take a copy of the node's contents but not sub-data (like pins).
 	FMemoryWriter NodeArchive(NodeData);
-	Node->SaveState(NodeArchive);
+	Node->ExportState(NodeArchive);
 
 	return Graph->RemoveNodeDirect(Node);
 }
@@ -499,7 +515,7 @@ bool FOptimusNodeGraphAction_RemoveNode::Undo(
 	UOptimusNode* Node = Graph->CreateNodeDirect(NodeClass, NodeName, [this](UOptimusNode* InNode)
 	{
 		FMemoryReader NodeArchive(NodeData);
-		InNode->RestoreState(NodeArchive);
+		InNode->ImportState(NodeArchive);
 		return true;
 	});
 
diff --git a/Engine/Plugins/Animation/DeformerGraph/Source/OptimusCore/Private/Actions/OptimusNodeGraphActions.h b/Engine/Plugins/Animation/DeformerGraph/Source/OptimusCore/Private/Actions/OptimusNodeGraphActions.h
--- a/Engine/Plugins/Animation/DeformerGraph/Source/OptimusCore/Private/Actions/OptimusNodeGraphActions.h
+++ b/Engine/Plugins/Animation/DeformerGraph/Source/OptimusCore/Private/Actions/OptimusNodeGraphActions.h
@@ -208,6 +208,10 @@ protected:
 	bool Undo(IOptimusPathResolver* InRoot) override;
 
 private:
+	
+	// The path of the node to duplicate
+	FString SourceNodePath;
+	
 	// The path of the graph the node should be added to.
 	FString GraphPath;
 	
@@ -222,8 +226,9 @@ private:
 	// gets added to the graph.
 	TFunction<bool(UOptimusNode*)> ConfigureNodeFunc;
 
-	// The stored node data to copy into the new node.
-	TArray<uint8> NodeData;
+	// State of the node when this action was created, useful when moving nodes from clipboard graph to a real graph
+	// not used if we can resolve the source node path during Do(), for regular duplicate usage
+	TArray<uint8> CachedNodeData;
 
 	// Path to the node that gets created through duplication.
 	FString NodePath;
diff --git a/Engine/Plugins/Animation/DeformerGraph/Source/OptimusCore/Private/Nodes/OptimusNode_DataInterface.cpp b/Engine/Plugins/Animation/DeformerGraph/Source/OptimusCore/Private/Nodes/OptimusNode_DataInterface.cpp
--- a/Engine/Plugins/Animation/DeformerGraph/Source/OptimusCore/Private/Nodes/OptimusNode_DataInterface.cpp
+++ b/Engine/Plugins/Animation/DeformerGraph/Source/OptimusCore/Private/Nodes/OptimusNode_DataInterface.cpp
@@ -103,6 +103,7 @@ void UOptimusNode_DataInterface::PostLoadNodeSpecificData()
 	if (DataInterfaceClass && !DataInterfaceData)
 	{
 		DataInterfaceData = NewObject<UOptimusComputeDataInterface>(this, DataInterfaceClass);
+		DataInterfaceData->SetFlags(RF_Transactional);
 	}
 
 	// Add in the component pin.
@@ -119,25 +120,22 @@ void UOptimusNode_DataInterface::OnDataTypeChanged(FName InTypeName)
 	DataInterfaceData->OnDataTypeChanged(InTypeName);
 }
 
-void UOptimusNode_DataInterface::SaveState(FArchive& Ar) const
+void UOptimusNode_DataInterface::ExportState(FArchive& Ar) const
 {
-	Super::SaveState(Ar);
-	// This fella does the heavy lifting of serializing object references. 
-	// FMemoryWriter and fam do not handle UObject* serialization on their own.
-	FObjectAndNameAsStringProxyArchive NodeProxyArchive(
-			Ar, /* bInLoadIfFindFails=*/ false);
-	DataInterfaceData->SerializeScriptProperties(NodeProxyArchive);	
+	Super::ExportState(Ar);
+
+	DataInterfaceData->ExportState(Ar);
+	
 }
 
-void UOptimusNode_DataInterface::RestoreState(FArchive& Ar)
+void UOptimusNode_DataInterface::ImportState(FArchive& Ar)
 {
-	Super::RestoreState(Ar);
+	Super::ImportState(Ar);
 
 	DataInterfaceData = NewObject<UOptimusComputeDataInterface>(this, DataInterfaceClass);
+	DataInterfaceData->SetFlags(RF_Transactional);
 	
-	FObjectAndNameAsStringProxyArchive NodeProxyArchive(
-			Ar, /* bInLoadIfFindFails=*/true);
-	DataInterfaceData->SerializeScriptProperties(NodeProxyArchive);
+	DataInterfaceData->ImportState(Ar);
 }
 
 bool UOptimusNode_DataInterface::IsComponentSourceCompatible(const UOptimusComponentSource* InComponentSource) const
@@ -352,6 +350,7 @@ void UOptimusNode_DataInterface::ConstructNode()
 		if (!DataInterfaceData)
 		{
 			DataInterfaceData = NewObject<UOptimusComputeDataInterface>(this, DataInterfaceClass);
+			DataInterfaceData->SetFlags(RF_Transactional);
 		}
 		SetDisplayName(FText::FromString(DataInterfaceData->GetDisplayName()));
 		CreateComponentPin();
@@ -362,14 +361,8 @@ void UOptimusNode_DataInterface::ConstructNode()
 
 void UOptimusNode_DataInterface::PostDuplicate(EDuplicateMode::Type DuplicateMode)
 {
-	// Currently duplication doesn't set the correct outer so fix here.
-	// We can remove this when duplication handles the outer correctly.
-	if (ensure(DataInterfaceData) && DataInterfaceData->GetOuter() != this)
-	{
-		FObjectDuplicationParameters DupParams = InitStaticDuplicateObjectParams(DataInterfaceData, this);
-		
-		DataInterfaceData = Cast<UOptimusComputeDataInterface>(StaticDuplicateObjectEx(DupParams));	
-	}
+	check(DataInterfaceData);
+	check(DataInterfaceData->GetOuter() == this);
 }
 
 
diff --git a/Engine/Plugins/Animation/DeformerGraph/Source/OptimusCore/Private/Nodes/OptimusNode_DataInterface.h b/Engine/Plugins/Animation/DeformerGraph/Source/OptimusCore/Private/Nodes/OptimusNode_DataInterface.h
--- a/Engine/Plugins/Animation/DeformerGraph/Source/OptimusCore/Private/Nodes/OptimusNode_DataInterface.h
+++ b/Engine/Plugins/Animation/DeformerGraph/Source/OptimusCore/Private/Nodes/OptimusNode_DataInterface.h
@@ -63,8 +63,8 @@ protected:
 	void PostLoadNodeSpecificData() override;
 	void OnDataTypeChanged(FName InTypeName) override;
 
-	void SaveState(FArchive& Ar) const override;
-	void RestoreState(FArchive& Ar) override;
+	void ExportState(FArchive& Ar) const override;
+	void ImportState(FArchive& Ar) override;
 
 	// -- UObject overrides
 	void PostDuplicate(EDuplicateMode::Type DuplicateMode) override;
diff --git a/Engine/Plugins/Animation/DeformerGraph/Source/OptimusCore/Private/OptimusComputeDataInterface.cpp b/Engine/Plugins/Animation/DeformerGraph/Source/OptimusCore/Private/OptimusComputeDataInterface.cpp
--- a/Engine/Plugins/Animation/DeformerGraph/Source/OptimusCore/Private/OptimusComputeDataInterface.cpp
+++ b/Engine/Plugins/Animation/DeformerGraph/Source/OptimusCore/Private/OptimusComputeDataInterface.cpp
@@ -3,6 +3,7 @@
 #include "OptimusComputeDataInterface.h"
 
 #include "DataInterfaces/OptimusDataInterfaceRawBuffer.h"
+#include "Serialization/ObjectAndNameAsStringProxyArchive.h"
 
 #include "Templates/SubclassOf.h"
 #include "UObject/UObjectIterator.h"
@@ -16,6 +17,24 @@ const FName UOptimusComputeDataInterface::CategoryName::ExecutionDataInterfaces(
 const FName UOptimusComputeDataInterface::CategoryName::OutputDataInterfaces("Output Data Interfaces");
 
 
+void UOptimusComputeDataInterface::ExportState(FArchive& Ar)
+{
+	// We have to use the proxy archive because FMemoryWriter simply asserts when the object has object references.
+	// However if the data interface spawns more objects, it needs to override this function with custom logic for undo/redo to work properly
+	FObjectAndNameAsStringProxyArchive NodeProxyArchive(
+			Ar, /* bInLoadIfFindFails=*/ false);
+
+	SerializeScriptProperties(NodeProxyArchive);	
+}
+
+void UOptimusComputeDataInterface::ImportState(FArchive& Ar)
+{
+	FObjectAndNameAsStringProxyArchive NodeProxyArchive(
+			Ar, /* bInLoadIfFindFails=*/false);
+	
+	SerializeScriptProperties(NodeProxyArchive);
+}
+
 TSet<TArray<FName>> UOptimusComputeDataInterface::GetUniqueNestedContexts() const
 {
 	TSet<TArray<FName>> UniqueContextNames;
diff --git a/Engine/Plugins/Animation/DeformerGraph/Source/OptimusCore/Private/OptimusNode.cpp b/Engine/Plugins/Animation/DeformerGraph/Source/OptimusCore/Private/OptimusNode.cpp
--- a/Engine/Plugins/Animation/DeformerGraph/Source/OptimusCore/Private/OptimusNode.cpp
+++ b/Engine/Plugins/Animation/DeformerGraph/Source/OptimusCore/Private/OptimusNode.cpp
@@ -487,26 +487,26 @@ void UOptimusNode::InitializeTransientData()
 }
 
 
-void UOptimusNode::SaveState(FArchive& Ar) const
+void UOptimusNode::ExportState(FArchive& Ar) const
 {
-	// Take a copy of the node's contents but not sub-data (like pins).
-	// Derived nodes may add additional data to the archive
+	// Take a copy of the node's core data that would allow the node to reconstruct itself during restore state
 	
-	// This fella does the heavy lifting of serializing object references. 
-	// FMemoryWriter and fam do not handle UObject* serialization on their own.
+	// We have to use the proxy archive because FMemoryWriter simply asserts when the object has object references.
+	// And even though we are serializing object references as strings. For things like pins we will recreate them anyway
 	FObjectAndNameAsStringProxyArchive NodeProxyArchive(
-			Ar, /* bInLoadIfFindFails=*/ false);
+			Ar, /* bInLoadIfFindFails=*/false);
+
 	SerializeScriptProperties(NodeProxyArchive);
 }
 
-void UOptimusNode::RestoreState(FArchive& Ar)
+void UOptimusNode::ImportState(FArchive& Ar)
 {
-	// Currently a warning appears when it can't find objects like pins and data interface data
-	// that were present during SaveState. However, they are harmless since we recreate those objects.
-	// But ideally there should be a better way to specify what should be saved/restored and what can be skipped/recreated
+	// Don't try to load any object references, the referenced object is likely no longer valid for this node and have to be recreated (such as pins)
+	// For objects we can find, it is likely that the object is owned by some top level object and is still alive and safe to reference.
 	
 	FObjectAndNameAsStringProxyArchive NodeProxyArchive(
-		Ar, /* bInLoadIfFindFails=*/true);
+		Ar, /* bInLoadIfFindFails=*/false);
+	
 	SerializeScriptProperties(NodeProxyArchive);
 }
 
diff --git a/Engine/Plugins/Animation/DeformerGraph/Source/OptimusCore/Public/OptimusComputeDataInterface.h b/Engine/Plugins/Animation/DeformerGraph/Source/OptimusCore/Public/OptimusComputeDataInterface.h
--- a/Engine/Plugins/Animation/DeformerGraph/Source/OptimusCore/Public/OptimusComputeDataInterface.h
+++ b/Engine/Plugins/Animation/DeformerGraph/Source/OptimusCore/Public/OptimusComputeDataInterface.h
@@ -152,6 +152,9 @@ public:
 	/// Data interface can use this to set default values/pins
 	virtual void Initialize() {};
 
+	virtual void ExportState(FArchive& Ar);
+	virtual void ImportState(FArchive& Ar);
+	
 	/// Whether the data interface allow users to add / remove pins
 	virtual bool CanPinDefinitionChange() { return false; };
 
diff --git a/Engine/Plugins/Animation/DeformerGraph/Source/OptimusCore/Public/OptimusNode.h b/Engine/Plugins/Animation/DeformerGraph/Source/OptimusCore/Public/OptimusNode.h
--- a/Engine/Plugins/Animation/DeformerGraph/Source/OptimusCore/Public/OptimusNode.h
+++ b/Engine/Plugins/Animation/DeformerGraph/Source/OptimusCore/Public/OptimusNode.h
@@ -208,6 +208,7 @@ protected:
 	friend struct FOptimusNodeGraphAction_PackageKernelFunction;
 	friend struct FOptimusNodeGraphAction_UnpackageKernelFunction;
 	friend struct FOptimusNodeGraphAction_RemoveNode;
+	friend struct FOptimusNodeGraphAction_DuplicateNode;
 
 	/**
 	 * Returns a unique name for pin. If given the same pins to compare to, the output
@@ -252,8 +253,8 @@ protected:
 		const UOptimusNodeGraph* InTargetGraph, 
 		FOptimusCompoundAction *InCompoundAction) {}
 
-	virtual void SaveState(FArchive& Ar) const;
-	virtual void RestoreState(FArchive& Ar);
+	virtual void ExportState(FArchive& Ar) const;
+	virtual void ImportState(FArchive& Ar);
 
 	void EnableDynamicPins();
 
