[Deformer Graph] make sure we serialize the type object with the data type ref so that we can register them on load to at least all the asset to open

diff --git a/Engine/Plugins/Animation/DeformerGraph/Source/OptimusCore/Private/OptimusDataType.cpp b/Engine/Plugins/Animation/DeformerGraph/Source/OptimusCore/Private/OptimusDataType.cpp
--- a/Engine/Plugins/Animation/DeformerGraph/Source/OptimusCore/Private/OptimusDataType.cpp
+++ b/Engine/Plugins/Animation/DeformerGraph/Source/OptimusCore/Private/OptimusDataType.cpp
@@ -24,7 +24,7 @@ void FOptimusDataTypeRef::Set(
 	if (InTypeHandle.IsValid())
 	{
 		TypeName = InTypeHandle->TypeName;
-		TypeObject = InTypeHandle->TypeObject;
+		TypeObject = InTypeHandle->TypeObject.Get();
 		checkSlow(FOptimusDataTypeRegistry::Get().FindType(TypeName) != nullptr);
 	}
 	else
@@ -66,7 +66,7 @@ void FOptimusDataTypeRef::PostSerialize(const FArchive& Ar)
 		const FOptimusDataTypeHandle TypeHandle = FOptimusDataTypeRegistry::Get().FindType(TypeName);
 		if (TypeHandle.IsValid())
 		{
-			TypeObject = TypeHandle->TypeObject;
+			TypeObject = TypeHandle->TypeObject.Get();
 		}
 	}
 }
diff --git a/Engine/Plugins/Animation/DeformerGraph/Source/OptimusCore/Public/OptimusDataType.h b/Engine/Plugins/Animation/DeformerGraph/Source/OptimusCore/Public/OptimusDataType.h
--- a/Engine/Plugins/Animation/DeformerGraph/Source/OptimusCore/Public/OptimusDataType.h
+++ b/Engine/Plugins/Animation/DeformerGraph/Source/OptimusCore/Public/OptimusDataType.h
@@ -165,9 +165,9 @@ struct OPTIMUSCORE_API FOptimusDataTypeRef
 	UPROPERTY(EditAnywhere, Category=Type)
 	FName TypeName;
 
-	// A weak pointer to the type object helps enforce asset dependency
+	// A soft pointer to the type object helps enforce asset dependency
 	UPROPERTY(EditAnywhere, Category=Type)
-	TWeakObjectPtr<UObject> TypeObject;
+	TSoftObjectPtr<UObject> TypeObject;
 	
 	void PostSerialize(const FArchive& Ar);
 };
