[Deformer Graph] when reading skinned mesh , it should use the same nullvertcolor buffer as the default GPU skin path.

diff --git a/Engine/Plugins/Animation/DeformerGraph/Source/OptimusCore/Private/DataInterfaces/OptimusDataInterfaceSkinnedMesh.cpp b/Engine/Plugins/Animation/DeformerGraph/Source/OptimusCore/Private/DataInterfaces/OptimusDataInterfaceSkinnedMesh.cpp
--- a/Engine/Plugins/Animation/DeformerGraph/Source/OptimusCore/Private/DataInterfaces/OptimusDataInterfaceSkinnedMesh.cpp
+++ b/Engine/Plugins/Animation/DeformerGraph/Source/OptimusCore/Private/DataInterfaces/OptimusDataInterfaceSkinnedMesh.cpp
@@ -177,7 +177,8 @@ void FOptimusSkinnedMeshDataProviderProxy::GatherDispatchData(FDispatchData cons
 	FSkeletalMeshLODRenderData const* LodRenderData = &SkeletalMeshRenderData.LODRenderData[LodIndex];
 
 	FRHIShaderResourceView* NullSRVBinding = GWhiteVertexBufferWithSRV->ShaderResourceViewRHI.GetReference();
-
+	FRHIShaderResourceView* NullColorVertexBufferSRV = GNullColorVertexBuffer.VertexBufferSRV;
+	
 	const TStridedView<FParameters> ParameterArray = MakeStridedParameterView<FParameters>(InDispatchData);
 	for (int32 InvocationIndex = 0; InvocationIndex < ParameterArray.Num(); ++InvocationIndex)
 	{
@@ -199,6 +200,6 @@ void FOptimusSkinnedMeshDataProviderProxy::GatherDispatchData(FDispatchData cons
 		Parameters.PositionInputBuffer = MeshVertexBufferSRV != nullptr ? MeshVertexBufferSRV : NullSRVBinding;
 		Parameters.TangentInputBuffer = MeshTangentBufferSRV != nullptr ? MeshTangentBufferSRV : NullSRVBinding;
 		Parameters.UVInputBuffer = MeshUVBufferSRV != nullptr ? MeshUVBufferSRV : NullSRVBinding;
-		Parameters.ColorInputBuffer = MeshColorBufferSRV != nullptr ? MeshColorBufferSRV : NullSRVBinding;
+		Parameters.ColorInputBuffer = MeshColorBufferSRV != nullptr ? MeshColorBufferSRV : NullColorVertexBufferSRV;
 	}
 }
