[Deformer Graph & Control Rig] Added AddDeformer Node in control rig to allow control rigs to add deformers to the skeletal mesh component on demand

+minor rig vm graph pin default value fix & unit test

known issues:
- deformer asset soft ptr is serialized as trait pin default value, which means if the deformer asset is renamed or moved, the automaically fixup process are not triggered and no redirector is created

diff --git a/Engine/Plugins/Animation/ControlRig/Source/ControlRig/Private/ControlRig.cpp b/Engine/Plugins/Animation/ControlRig/Source/ControlRig/Private/ControlRig.cpp
--- a/Engine/Plugins/Animation/ControlRig/Source/ControlRig/Private/ControlRig.cpp
+++ b/Engine/Plugins/Animation/ControlRig/Source/ControlRig/Private/ControlRig.cpp
@@ -129,6 +129,9 @@ void UControlRig::ResetRecordedTransforms(const FName& InEventName)
 
 void UControlRig::BeginDestroy()
 {
+	BeginDestroyEvent.Broadcast(this);
+	BeginDestroyEvent.Clear();
+	
 	Super::BeginDestroy();
 	SetRigVMExtendedExecuteContext(nullptr);
 
@@ -854,6 +857,7 @@ bool UControlRig::Execute(const FName& InEventName)
 	Context.InteractionType = InteractionType;
 	Context.ElementsBeingInteracted = ElementsBeingInteracted;
 	PublicContext.Hierarchy = GetHierarchy();
+	PublicContext.ControlRig = this;
 
 	// allow access to the hierarchy
 	Context.HierarchySettings = HierarchySettings;
diff --git a/Engine/Plugins/Animation/ControlRig/Source/ControlRig/Private/ModularRig.cpp b/Engine/Plugins/Animation/ControlRig/Source/ControlRig/Private/ModularRig.cpp
--- a/Engine/Plugins/Animation/ControlRig/Source/ControlRig/Private/ModularRig.cpp
+++ b/Engine/Plugins/Animation/ControlRig/Source/ControlRig/Private/ModularRig.cpp
@@ -799,6 +799,7 @@ FRigModuleInstance* UModularRig::AddModuleInstance(const FName& InModuleName, TS
 		NewModuleRig->bCopyHierarchyBeforeConstruction = false;
 		NewModuleRig->SetDynamicHierarchy(Hierarchy);
 		ModulePublicContext.Hierarchy = Hierarchy;
+		ModulePublicContext.ControlRig = this;
 		ModulePublicContext.RigModuleNameSpace = NewModuleRig->GetRigModuleNameSpace();
 		ModulePublicContext.RigModuleNameSpaceHash = GetTypeHash(ModulePublicContext.RigModuleNameSpace);
 		NewModuleRig->SetElementKeyRedirector(FRigElementKeyRedirector(InConnectionMap, Hierarchy));
diff --git a/Engine/Plugins/Animation/ControlRig/Source/ControlRig/Private/Rigs/AdditiveControlRig.cpp b/Engine/Plugins/Animation/ControlRig/Source/ControlRig/Private/Rigs/AdditiveControlRig.cpp
--- a/Engine/Plugins/Animation/ControlRig/Source/ControlRig/Private/Rigs/AdditiveControlRig.cpp
+++ b/Engine/Plugins/Animation/ControlRig/Source/ControlRig/Private/Rigs/AdditiveControlRig.cpp
@@ -48,6 +48,7 @@ bool UAdditiveControlRig::Execute_Internal(const FName& InEventName)
 	{
 		FControlRigExecuteContext ExecuteContext;
 		ExecuteContext.Hierarchy = GetHierarchy();
+		ExecuteContext.ControlRig = this;
 		ExecuteContext.SetEventName(InEventName);
 		ExecuteContext.UnitContext = FRigUnitContext();
 
diff --git a/Engine/Plugins/Animation/ControlRig/Source/ControlRig/Public/ControlRig.h b/Engine/Plugins/Animation/ControlRig/Source/ControlRig/Public/ControlRig.h
--- a/Engine/Plugins/Animation/ControlRig/Source/ControlRig/Public/ControlRig.h
+++ b/Engine/Plugins/Animation/ControlRig/Source/ControlRig/Public/ControlRig.h
@@ -205,6 +205,14 @@ public:
 	virtual void BeginDestroy() override;
 	// END UObject interface
 
+	DECLARE_EVENT_OneParam(UControlRig, FControlRigBeginDestroyEvent, class UControlRig*);
+	FControlRigBeginDestroyEvent& OnBeginDestroy() {return BeginDestroyEvent;};
+private:
+	/** Broadcasts a notification just before the controlrig is destroyed. */
+	FControlRigBeginDestroyEvent BeginDestroyEvent;
+	
+public:	
+	
 	UPROPERTY(transient)
 	ERigExecutionType ExecutionType;
 
diff --git a/Engine/Plugins/Animation/ControlRig/Source/ControlRig/Public/Units/RigUnitContext.h b/Engine/Plugins/Animation/ControlRig/Source/ControlRig/Public/Units/RigUnitContext.h
--- a/Engine/Plugins/Animation/ControlRig/Source/ControlRig/Public/Units/RigUnitContext.h
+++ b/Engine/Plugins/Animation/ControlRig/Source/ControlRig/Public/Units/RigUnitContext.h
@@ -122,6 +122,7 @@ public:
 	FControlRigExecuteContext()
 		: FRigVMExecuteContext()
 		, Hierarchy(nullptr)
+		, ControlRig(nullptr)
 		, RigModuleNameSpace()
 		, RigModuleNameSpaceHash(0)
 	{
@@ -133,6 +134,7 @@ public:
 
 		const FControlRigExecuteContext* OtherContext = (const FControlRigExecuteContext*)InOtherContext; 
 		Hierarchy = OtherContext->Hierarchy;
+		ControlRig = OtherContext->ControlRig;
 	}
 
 	/**
@@ -221,6 +223,7 @@ public:
 	
 	FRigUnitContext UnitContext;
 	URigHierarchy* Hierarchy;
+	UControlRig* ControlRig;
 
 #if WITH_EDITOR
 	virtual void Report(const FRigVMLogSettings& InLogSettings, const FName& InFunctionName, int32 InInstructionIndex, const FString& InMessage) const override
diff --git a/Engine/Plugins/Animation/DeformerGraph/DeformerGraph.uplugin b/Engine/Plugins/Animation/DeformerGraph/DeformerGraph.uplugin
--- a/Engine/Plugins/Animation/DeformerGraph/DeformerGraph.uplugin
+++ b/Engine/Plugins/Animation/DeformerGraph/DeformerGraph.uplugin
@@ -52,6 +52,10 @@
 		{
 			"Name": "ComputeFramework",
 			"Enabled": true
+		},
+		{
+			"Name": "ControlRig",
+			"Enabled": true
 		}
 	]
 }
diff --git a/Engine/Plugins/Animation/DeformerGraph/Shaders/Private/DataInterfaceSkinnedMeshRead.ush b/Engine/Plugins/Animation/DeformerGraph/Shaders/Private/DataInterfaceSkinnedMeshRead.ush
new file mode 100644
--- /dev/null
+++ b/Engine/Plugins/Animation/DeformerGraph/Shaders/Private/DataInterfaceSkinnedMeshRead.ush
@@ -0,0 +1,70 @@
+// Copyright Epic Games, Inc. All Rights Reserved.
+
+// Include required for platform TANGENT_RWBUFFER_FORMAT.
+#include "/Engine/Private/GpuSkinCommon.ush"
+
+uint {DataInterfaceName}_NumVertices;
+uint {DataInterfaceName}_InputStreamStart;
+
+RWBuffer<float> {DataInterfaceName}_PositionBufferUAV;
+RWBuffer<TANGENT_RWBUFFER_FORMAT> {DataInterfaceName}_TangentBufferUAV;
+RWBuffer<UNORM float4> {DataInterfaceName}_ColorBufferUAV;
+
+Buffer<float> {DataInterfaceName}_PositionStaticBuffer;
+Buffer<TANGENT_RWBUFFER_FORMAT> {DataInterfaceName}_TangentStaticBuffer;
+Buffer<UNORM float4> {DataInterfaceName}_ColorStaticBuffer;
+uint {DataInterfaceName}_ColorIndexMask;
+
+uint ReadNumVertices_{DataInterfaceName}()
+{
+	return {DataInterfaceName}_NumVertices;
+}
+
+float3 ReadPosition_{DataInterfaceName}(uint VertexIndex)
+{
+	uint BufferIndex = ({DataInterfaceName}_InputStreamStart + VertexIndex) * 3;
+
+#if READABLE_OUTPUT_BUFFERS & OPTIMUS_SKINNED_MESH_READ_POSITION
+	return float3(
+		{DataInterfaceName}_PositionBufferUAV[BufferIndex + 0], 
+		{DataInterfaceName}_PositionBufferUAV[BufferIndex + 1],
+		{DataInterfaceName}_PositionBufferUAV[BufferIndex + 2]);
+#else
+	return float3(
+		{DataInterfaceName}_PositionStaticBuffer[BufferIndex + 0], 
+		{DataInterfaceName}_PositionStaticBuffer[BufferIndex + 1],
+		{DataInterfaceName}_PositionStaticBuffer[BufferIndex + 2]);
+#endif 
+}
+
+float4 ReadTangentX_{DataInterfaceName}(uint VertexIndex)
+{
+	uint BufferIndex = ({DataInterfaceName}_InputStreamStart + VertexIndex) * 2;
+	
+#if READABLE_OUTPUT_BUFFERS & OPTIMUS_SKINNED_MESH_READ_TANGENTS
+	return TangentBias_SkinCache({DataInterfaceName}_TangentBufferUAV[BufferIndex]) ;
+#else	
+	return TangentBias_SkinCache({DataInterfaceName}_TangentStaticBuffer[BufferIndex]) ;
+#endif
+}
+
+float4 ReadTangentZ_{DataInterfaceName}(uint VertexIndex)
+{
+	uint BufferIndex = ({DataInterfaceName}_InputStreamStart + VertexIndex) * 2;
+#if READABLE_OUTPUT_BUFFERS & OPTIMUS_SKINNED_MESH_READ_TANGENTS
+	return TangentBias_SkinCache({DataInterfaceName}_TangentBufferUAV[BufferIndex + 1]) ;
+#else	
+	return TangentBias_SkinCache({DataInterfaceName}_TangentStaticBuffer[BufferIndex + 1]) ;
+#endif
+}
+
+float4 ReadColor_{DataInterfaceName}(uint VertexIndex)
+{
+	// Index Mask is used here in case ColorStaticBuffer is bound to the GWhiteVertexBufferWithSRV that is only a few bytes
+	uint BufferIndex = ({DataInterfaceName}_InputStreamStart + VertexIndex) & ({DataInterfaceName}_ColorIndexMask);
+#if READABLE_OUTPUT_BUFFERS & OPTIMUS_SKINNED_MESH_READ_COLOR
+	return {DataInterfaceName}_ColorBufferUAV[BufferIndex] FMANUALFETCH_COLOR_COMPONENT_SWIZZLE;
+#else	
+	return {DataInterfaceName}_ColorStaticBuffer[BufferIndex] FMANUALFETCH_COLOR_COMPONENT_SWIZZLE;
+#endif	
+}
diff --git a/Engine/Plugins/Animation/DeformerGraph/Source/OptimusCore/OptimusCore.Build.cs b/Engine/Plugins/Animation/DeformerGraph/Source/OptimusCore/OptimusCore.Build.cs
--- a/Engine/Plugins/Animation/DeformerGraph/Source/OptimusCore/OptimusCore.Build.cs
+++ b/Engine/Plugins/Animation/DeformerGraph/Source/OptimusCore/OptimusCore.Build.cs
@@ -25,10 +25,22 @@ namespace UnrealBuildTool.Rules
 					"Projects",
 					"RenderCore",
 					"Renderer",
-					"RHI"
+					"RHI",
+					"RigVM",
+					"ControlRig"
 				}
 			);
 
+            if (Target.bBuildEditor == true)
+            {
+	            PrivateDependencyModuleNames.AddRange(
+		            new string[]
+		            {
+			            "RigVMDeveloper",
+		            }
+	            );
+            }
+			
 			bDisableAutoRTFMInstrumentation = true;
         }
     }
diff --git a/Engine/Plugins/Animation/DeformerGraph/Source/OptimusCore/Private/ControlRig/RigUnit_Optimus.cpp b/Engine/Plugins/Animation/DeformerGraph/Source/OptimusCore/Private/ControlRig/RigUnit_Optimus.cpp
new file mode 100644
--- /dev/null
+++ b/Engine/Plugins/Animation/DeformerGraph/Source/OptimusCore/Private/ControlRig/RigUnit_Optimus.cpp
@@ -0,0 +1,463 @@
+﻿// Copyright Epic Games, Inc. All Rights Reserved.
+#include "RigUnit_Optimus.h"
+
+#if WITH_EDITOR
+#include "RigVMModel/RigVMController.h"
+#endif 
+#include "Units/RigUnitContext.h"
+#include "ControlRig.h"
+#include "OptimusDataTypeRegistry.h"
+#include "OptimusDeformerInstance.h"
+#include "OptimusHelpers.h"
+#include "OptimusVariableDescription.h"
+
+const TCHAR* FRigUnit_AddOptimusDeformer::DeformerTraitName = TEXT("DeformerGraphAsset");
+const TCHAR* FRigUnit_AddOptimusDeformer::DeformerSettingsTraitName = TEXT("Settings");
+
+FString FRigVMTrait_OptimusDeformer::GetDisplayName() const
+{
+	return TEXT("Deformer Graph Asset");
+}
+
+#if WITH_EDITOR
+bool FRigVMTrait_OptimusDeformer::ShouldCreatePinForProperty(const FProperty* InProperty) const
+{
+	if(!Super::ShouldCreatePinForProperty(InProperty))
+	{
+		return false;
+	}
+	
+	return InProperty->GetFName() != GET_MEMBER_NAME_CHECKED(FRigVMTrait_OptimusDeformer, DeformerGraph);
+}
+#endif
+
+
+bool FRigUnit_AddOptimusDeformer::IsVariableTraitName(const FString& InTraitName)
+{
+	return InTraitName != DeformerTraitName &&
+		InTraitName != DeformerSettingsTraitName;
+}
+
+
+void FRigUnit_AddOptimusDeformer::OnUnitNodeCreated(FRigVMUnitNodeCreatedContext& InContext) const
+{
+	FRigUnitMutable::OnUnitNodeCreated(InContext);
+#if WITH_EDITOR
+	if (InContext.GetReason() == ERigVMNodeCreatedReason::NodeSpawner)
+	{
+		InContext.GetController()->AddTrait(InContext.GetNodeName(), *FRigVMTrait_OptimusDeformer::StaticStruct()->GetPathName(), DeformerTraitName);
+		InContext.GetController()->AddTrait(InContext.GetNodeName(), *FRigVMTrait_OptimusDeformerSettings::StaticStruct()->GetPathName(), DeformerSettingsTraitName);
+	}
+#endif
+}
+
+TArray<FRigVMUserWorkflow> FRigUnit_AddOptimusDeformer::GetSupportedWorkflows(const UObject* InSubject) const
+{
+	TArray<FRigVMUserWorkflow> Workflows = Super::GetSupportedWorkflows(InSubject);
+	
+#if WITH_EDITOR
+
+	Workflows.Emplace(TEXT("Refresh Variables"), TEXT("Populate the node with available variables in the Deformer Graph"), ERigVMUserWorkflowType::NodeContext,
+		FRigVMPerformUserWorkflowDelegate::CreateLambda([](const URigVMUserWorkflowOptions* InOptions, UObject* InController)
+		{
+			URigVMController* Controller = CastChecked<URigVMController>(InController);
+			if(URigVMNode* Node = InOptions->GetSubject<URigVMNode>())
+			{
+				TSharedPtr<FStructOnScope> StructOnScope = Node->GetTraitInstance(DeformerTraitName, true);
+				FRigVMTrait_OptimusDeformer* TraitInstance = reinterpret_cast<FRigVMTrait_OptimusDeformer*>(StructOnScope->GetStructMemory());
+				UOptimusDeformer* Deformer = TraitInstance->DeformerGraph.LoadSynchronous();
+				
+				// TODO: Ideally pin display name should update immediately when the asset changes
+				Controller->SetPinDisplayName(Node->FindTrait(DeformerTraitName)->GetPinPath(),
+					Deformer ?
+					TraitInstance->DeformerGraph.GetAssetName() : TEXT("Deformer Graph Asset(Unassigned)")) ;
+
+				const TArray<URigVMController::FLinkedPath> LinkedPaths = URigVMController::GetLinkedPaths(Node);
+				TMap<FString, URigVMController::FPinState> PinStates = Controller->GetPinStates(Node);
+				
+				TArray<FString> TraitNames = Node->GetTraitNames();
+
+				for (const FString& TraitName : TraitNames)
+				{
+					if (IsVariableTraitName(TraitName))
+					{
+						Controller->RemoveTrait(Node->GetFName(), *TraitName);
+					}
+				}
+				
+				if (Deformer)
+				{
+					for (UOptimusVariableDescription* Variable : Deformer->GetVariables())
+					{
+						FName TraitName = NAME_None;
+						if (Variable->DataType == FOptimusDataTypeRegistry::Get().FindType(*FIntProperty::StaticClass()))
+						{
+							TraitName = Controller->AddTrait(Node->GetFName(), *FRigVMTrait_SetDeformerIntVariable::StaticStruct()->GetPathName(), Variable->VariableName);
+						}
+						else if (Variable->DataType == FOptimusDataTypeRegistry::Get().FindArrayType(*FIntProperty::StaticClass()))
+						{
+							TraitName = Controller->AddTrait(Node->GetFName(), *FRigVMTrait_SetDeformerIntArrayVariable::StaticStruct()->GetPathName(), Variable->VariableName);
+						}
+						if (Variable->DataType == FOptimusDataTypeRegistry::Get().FindType(TBaseStructure<FIntPoint>::Get()))
+						{
+							TraitName = Controller->AddTrait(Node->GetFName(), *FRigVMTrait_SetDeformerInt2Variable::StaticStruct()->GetPathName(), Variable->VariableName);
+						}
+						else if (Variable->DataType == FOptimusDataTypeRegistry::Get().FindArrayType(TBaseStructure<FIntPoint>::Get()))
+						{
+							TraitName = Controller->AddTrait(Node->GetFName(), *FRigVMTrait_SetDeformerInt2ArrayVariable::StaticStruct()->GetPathName(), Variable->VariableName);
+						}
+						if (Variable->DataType == FOptimusDataTypeRegistry::Get().FindType(TBaseStructure<FIntVector>::Get()))
+						{
+							TraitName = Controller->AddTrait(Node->GetFName(), *FRigVMTrait_SetDeformerInt3Variable::StaticStruct()->GetPathName(), Variable->VariableName);
+						}
+						else if (Variable->DataType == FOptimusDataTypeRegistry::Get().FindArrayType(TBaseStructure<FIntVector>::Get()))
+						{
+							TraitName = Controller->AddTrait(Node->GetFName(), *FRigVMTrait_SetDeformerInt3ArrayVariable::StaticStruct()->GetPathName(), Variable->VariableName);
+						}
+						if (Variable->DataType == FOptimusDataTypeRegistry::Get().FindType(TBaseStructure<FIntVector4>::Get()))
+						{
+							TraitName = Controller->AddTrait(Node->GetFName(), *FRigVMTrait_SetDeformerInt4Variable::StaticStruct()->GetPathName(), Variable->VariableName);
+						}
+						else if (Variable->DataType == FOptimusDataTypeRegistry::Get().FindArrayType(TBaseStructure<FIntVector4>::Get()))
+						{
+							TraitName = Controller->AddTrait(Node->GetFName(), *FRigVMTrait_SetDeformerInt4ArrayVariable::StaticStruct()->GetPathName(), Variable->VariableName);
+						}
+						else if (Variable->DataType == FOptimusDataTypeRegistry::Get().FindType(*FDoubleProperty::StaticClass()) ||
+							Variable->DataType == FOptimusDataTypeRegistry::Get().FindType(*FFloatProperty::StaticClass()))
+						{
+							TraitName = Controller->AddTrait(Node->GetFName(), *FRigVMTrait_SetDeformerFloatVariable::StaticStruct()->GetPathName(), Variable->VariableName);
+						}
+						else if (Variable->DataType == FOptimusDataTypeRegistry::Get().FindArrayType(*FDoubleProperty::StaticClass()) ||
+							Variable->DataType == FOptimusDataTypeRegistry::Get().FindArrayType(*FFloatProperty::StaticClass()))
+						{
+							TraitName = Controller->AddTrait(Node->GetFName(), *FRigVMTrait_SetDeformerFloatArrayVariable::StaticStruct()->GetPathName(), Variable->VariableName);
+						}
+						else if (Variable->DataType == FOptimusDataTypeRegistry::Get().FindType(TBaseStructure<FVector2D>::Get()))
+						{
+							TraitName = Controller->AddTrait(Node->GetFName(), *FRigVMTrait_SetDeformerVector2Variable::StaticStruct()->GetPathName(), Variable->VariableName);
+						}
+						else if (Variable->DataType == FOptimusDataTypeRegistry::Get().FindArrayType(TBaseStructure<FVector2D>::Get()))
+						{
+							TraitName = Controller->AddTrait(Node->GetFName(), *FRigVMTrait_SetDeformerVector2ArrayVariable::StaticStruct()->GetPathName(), Variable->VariableName);
+						}
+						else if (Variable->DataType == FOptimusDataTypeRegistry::Get().FindType(TBaseStructure<FVector>::Get()))
+						{
+							TraitName = Controller->AddTrait(Node->GetFName(), *FRigVMTrait_SetDeformerVectorVariable::StaticStruct()->GetPathName(), Variable->VariableName);
+						}
+						else if (Variable->DataType == FOptimusDataTypeRegistry::Get().FindArrayType(TBaseStructure<FVector>::Get()))
+						{
+							TraitName = Controller->AddTrait(Node->GetFName(), *FRigVMTrait_SetDeformerVectorArrayVariable::StaticStruct()->GetPathName(), Variable->VariableName);
+						}
+						else if (Variable->DataType == FOptimusDataTypeRegistry::Get().FindType(TBaseStructure<FVector4>::Get()))
+						{
+							TraitName = Controller->AddTrait(Node->GetFName(), *FRigVMTrait_SetDeformerVector4Variable::StaticStruct()->GetPathName(), Variable->VariableName);
+						}
+						else if (Variable->DataType == FOptimusDataTypeRegistry::Get().FindArrayType(TBaseStructure<FVector4>::Get()))
+						{
+							TraitName = Controller->AddTrait(Node->GetFName(), *FRigVMTrait_SetDeformerVector4ArrayVariable::StaticStruct()->GetPathName(), Variable->VariableName);
+						}
+						else if (Variable->DataType == FOptimusDataTypeRegistry::Get().FindType(TBaseStructure<FLinearColor>::Get()))
+						{
+							TraitName = Controller->AddTrait(Node->GetFName(), *FRigVMTrait_SetDeformerLinearColorVariable::StaticStruct()->GetPathName(), Variable->VariableName);
+						}
+						else if (Variable->DataType == FOptimusDataTypeRegistry::Get().FindArrayType(TBaseStructure<FLinearColor>::Get()))
+						{
+							TraitName = Controller->AddTrait(Node->GetFName(), *FRigVMTrait_SetDeformerLinearColorArrayVariable::StaticStruct()->GetPathName(), Variable->VariableName);
+						}
+						else if (Variable->DataType == FOptimusDataTypeRegistry::Get().FindType(TBaseStructure<FQuat>::Get()))
+						{
+							TraitName = Controller->AddTrait(Node->GetFName(), *FRigVMTrait_SetDeformerQuatVariable::StaticStruct()->GetPathName(), Variable->VariableName);
+						}
+						else if (Variable->DataType == FOptimusDataTypeRegistry::Get().FindArrayType(TBaseStructure<FQuat>::Get()))
+						{
+							TraitName = Controller->AddTrait(Node->GetFName(), *FRigVMTrait_SetDeformerQuatArrayVariable::StaticStruct()->GetPathName(), Variable->VariableName);
+						}
+						else if (Variable->DataType == FOptimusDataTypeRegistry::Get().FindType(TBaseStructure<FRotator>::Get()))
+						{
+							TraitName = Controller->AddTrait(Node->GetFName(), *FRigVMTrait_SetDeformerRotatorVariable::StaticStruct()->GetPathName(), Variable->VariableName);
+						}
+						else if (Variable->DataType == FOptimusDataTypeRegistry::Get().FindArrayType(TBaseStructure<FRotator>::Get()))
+						{
+							TraitName = Controller->AddTrait(Node->GetFName(), *FRigVMTrait_SetDeformerRotatorArrayVariable::StaticStruct()->GetPathName(), Variable->VariableName);
+						}
+						else if (Variable->DataType == FOptimusDataTypeRegistry::Get().FindType(TBaseStructure<FTransform>::Get()))
+						{
+							TraitName = Controller->AddTrait(Node->GetFName(), *FRigVMTrait_SetDeformerTransformVariable::StaticStruct()->GetPathName(), Variable->VariableName);
+						}
+						else if (Variable->DataType == FOptimusDataTypeRegistry::Get().FindArrayType(TBaseStructure<FTransform>::Get()))
+						{
+							TraitName = Controller->AddTrait(Node->GetFName(), *FRigVMTrait_SetDeformerTransformArrayVariable::StaticStruct()->GetPathName(), Variable->VariableName);
+						}
+						
+						Controller->SetPinExpansion(Node->FindTrait(TraitName)->GetPinPath(), true, true);
+					}
+				}
+
+				Controller->ApplyPinStates(Node, PinStates, {}, true);
+				Controller->RestoreLinkedPaths(LinkedPaths, {}, true);
+
+				return true;
+			}
+			return false;
+		}), URigVMUserWorkflowOptions::StaticClass());
+#endif
+	
+	return Workflows;
+	
+}
+
+
+FRigUnit_AddOptimusDeformer_Execute()
+{
+	if (!DeformerInstanceGuid.IsValid())
+	{
+		DeformerInstanceGuid = FGuid::NewGuid();
+	}
+	
+	USkeletalMeshComponent* RigMeshComponent = Cast<USkeletalMeshComponent>(ExecuteContext.GetMutableOwningComponent());
+	
+	if(RigMeshComponent)
+	{
+		const TArrayView<const FRigVMTraitScope> Traits = ExecuteContext.GetTraits();
+		
+		// Deformer Graph Trait spawns a task to adds the deformer on game thread
+		// Settings Trait enqueues the deformer based on execution settings
+		// Variable Traits sets variable value on deformer instances created by this rig unit
+
+		// By default, a deformer is added to all child components if possible.
+		// While Add is done on game thread, enqueue and set variable are done on anim thread
+		// Given that parent component always ticks before child components, it should be safe for the parent
+		// component to modify the deformer instance manager on the child components
+		
+		for(const FRigVMTraitScope& Scope : Traits)
+		{
+			if(const FRigVMTrait_OptimusDeformer* DeformerTrait = Scope.GetTrait<FRigVMTrait_OptimusDeformer>())
+			{
+				if (!DeformerTrait->DeformerGraph.IsNull())
+				{
+					FFunctionGraphTask::CreateAndDispatchWhenReady([
+						WeakMesh = TWeakObjectPtr(RigMeshComponent),
+						WeakRig = TWeakObjectPtr(ExecuteContext.ControlRig),
+						DeformerGraphAsset = DeformerTrait->DeformerGraph,
+						DeformerInstanceGuid]()
+					{
+						if (WeakMesh.IsValid() && WeakRig.IsValid())
+						{
+							// See StaticFindObjectFastSafe()
+							if (!UE::IsSavingPackage(nullptr) && !IsGarbageCollectingAndLockingUObjectHashTables())
+							{
+								if (UOptimusDeformer* DeformerGraph = DeformerGraphAsset.LoadSynchronous())
+								{
+									USkeletalMeshComponent* RigMeshComponent = WeakMesh.Get();
+
+									TArray<USceneComponent*> ChildComponents;
+									RigMeshComponent->GetChildrenComponents(true, ChildComponents);
+
+									TArray<USkeletalMeshComponent*> ComponentsToProcess = {RigMeshComponent};
+
+									for (USceneComponent* Component : ChildComponents)
+									{
+										if (USkeletalMeshComponent* SkeletalMeshComponent = Cast<USkeletalMeshComponent>(Component))
+										{
+											ComponentsToProcess.Add(SkeletalMeshComponent);
+										}
+									}
+
+									for (USkeletalMeshComponent* ComponentToProcess : ComponentsToProcess)
+									{
+										// Currently only one deformer for all LODs so using LOD0 for now is fine
+										// It may change in the future depending on how per-LOD deformer instance is implemented
+										if (!ComponentToProcess->GetMeshDeformerInstanceForLOD(0))
+										{
+											// Every time we re-set the mesh deformers we have to wait for anim eval to complete to avoid anim thread accessing
+											// deformer instances while we are trying to modify it. Similar to OnUnregister for skeletal mesh component
+											ComponentToProcess->HandleExistingParallelEvaluationTask(true, false);
+										
+											// In case there is no mesh deformer running, force turning it on and use a default deformer
+											// in project settings
+											ComponentToProcess->SetAlwaysUseMeshDeformer(true);
+										}
+								
+										if (UOptimusDeformerDynamicInstanceManager* DeformerInstanceManager = Cast<UOptimusDeformerDynamicInstanceManager>(ComponentToProcess->GetMeshDeformerInstanceForLOD(0)))
+										{
+											UControlRig* Rig = WeakRig.Get();
+											if (!DeformerInstanceManager->GetRigDeformer(DeformerInstanceGuid))
+											{
+												DeformerInstanceManager->AddRigDeformer(Rig, DeformerInstanceGuid, DeformerGraph);
+											}
+										}
+									}
+								}	
+							}
+						}
+					}, TStatId(), NULL, ENamedThreads::GameThread);
+				}	
+			}
+			else
+			{
+				TArray<USceneComponent*> ChildComponents;
+				RigMeshComponent->GetChildrenComponents(true, ChildComponents);
+				TArray<USkeletalMeshComponent*> ComponentsToProcess = {RigMeshComponent};
+				for (USceneComponent* Component : ChildComponents)
+				{
+					if (USkeletalMeshComponent* SkeletalMeshComponent = Cast<USkeletalMeshComponent>(Component))
+					{
+						ComponentsToProcess.Add(SkeletalMeshComponent);
+					}
+				}
+				
+				if (const FRigVMTrait_OptimusDeformerSettings* SettingsTrait = Scope.GetTrait<FRigVMTrait_OptimusDeformerSettings>())
+				{
+					for (const USkeletalMeshComponent* ComponentToProcess : ComponentsToProcess)
+					{
+						// Currently, there is only one deformer instance used by all LODs, so use LOD 0 here for now
+						// This might change in the future depending on how per-LOD instance is implemented
+						if (UOptimusDeformerDynamicInstanceManager* DeformerInstanceManager = Cast<UOptimusDeformerDynamicInstanceManager>(ComponentToProcess->GetMeshDeformerInstanceForLOD(0)))
+						{
+							DeformerInstanceManager->EnqueueRigDeformer(DeformerInstanceGuid, SettingsTrait->ExecutionPhase, SettingsTrait->ExecutionGroup);
+						}	
+					}
+				}
+				else if (const FRigVMTrait_OptimusVariableBase* VariableTrait = Scope.GetTrait<FRigVMTrait_OptimusVariableBase>())
+				{
+					for (const USkeletalMeshComponent* ComponentToProcess : ComponentsToProcess)
+					{
+						// Currently, there is only one deformer instance used by all LODs, so use LOD 0 here for now
+						// This might change in the future depending on how per-LOD instance is implemented
+						if (UOptimusDeformerDynamicInstanceManager* DeformerInstanceManager = Cast<UOptimusDeformerDynamicInstanceManager>(ComponentToProcess->GetMeshDeformerInstanceForLOD(0)))
+						{
+							if (DeformerInstanceGuid.IsValid())
+							{
+								if (UOptimusDeformerInstance* DeformerInstance = DeformerInstanceManager->GetRigDeformer(DeformerInstanceGuid))
+								{
+									VariableTrait->SetValue(DeformerInstance);
+								}
+							}
+						}
+					}
+				}
+			}
+		}
+	}
+}
+
+
+void FRigVMTrait_SetDeformerIntVariable::SetValue(UOptimusDeformerInstance* InInstance) const
+{
+	InInstance->SetIntVariable(*GetName(), Value);
+}
+
+void FRigVMTrait_SetDeformerIntArrayVariable::SetValue(UOptimusDeformerInstance* InInstance) const
+{
+	InInstance->SetIntArrayVariable(*GetName(), Value);
+}
+
+void FRigVMTrait_SetDeformerInt2Variable::SetValue(UOptimusDeformerInstance* InInstance) const
+{
+	InInstance->SetInt2Variable(*GetName(), Value);
+}
+
+void FRigVMTrait_SetDeformerInt2ArrayVariable::SetValue(UOptimusDeformerInstance* InInstance) const
+{
+	InInstance->SetInt2ArrayVariable(*GetName(), Value);
+}
+
+void FRigVMTrait_SetDeformerInt3Variable::SetValue(UOptimusDeformerInstance* InInstance) const
+{
+	InInstance->SetInt3Variable(*GetName(), Value);
+}
+
+void FRigVMTrait_SetDeformerInt3ArrayVariable::SetValue(UOptimusDeformerInstance* InInstance) const
+{
+	InInstance->SetInt3ArrayVariable(*GetName(), Value);
+}
+
+void FRigVMTrait_SetDeformerInt4Variable::SetValue(UOptimusDeformerInstance* InInstance) const
+{
+	InInstance->SetInt4Variable(*GetName(), Value);
+}
+
+void FRigVMTrait_SetDeformerInt4ArrayVariable::SetValue(UOptimusDeformerInstance* InInstance) const
+{
+	InInstance->SetInt4ArrayVariable(*GetName(), Value);
+}
+
+void FRigVMTrait_SetDeformerFloatVariable::SetValue(UOptimusDeformerInstance* InInstance) const
+{
+	InInstance->SetFloatVariable(*GetName(), Value);
+}
+
+void FRigVMTrait_SetDeformerFloatArrayVariable::SetValue(UOptimusDeformerInstance* InInstance) const
+{
+	InInstance->SetFloatArrayVariable(*GetName(), Value);
+}
+
+void FRigVMTrait_SetDeformerVector2Variable::SetValue(UOptimusDeformerInstance* InInstance) const
+{
+	InInstance->SetVector2Variable(*GetName(), Value);
+}
+
+void FRigVMTrait_SetDeformerVector2ArrayVariable::SetValue(UOptimusDeformerInstance* InInstance) const
+{
+	InInstance->SetVector2ArrayVariable(*GetName(), Value);
+}
+
+void FRigVMTrait_SetDeformerVectorVariable::SetValue(UOptimusDeformerInstance* InInstance) const
+{
+	InInstance->SetVectorVariable(*GetName(), Value);
+}
+
+void FRigVMTrait_SetDeformerVectorArrayVariable::SetValue(UOptimusDeformerInstance* InInstance) const
+{
+	InInstance->SetVectorArrayVariable(*GetName(), Value);
+}
+
+void FRigVMTrait_SetDeformerVector4Variable::SetValue(UOptimusDeformerInstance* InInstance) const
+{
+	InInstance->SetVector4Variable(*GetName(), Value);
+}
+
+void FRigVMTrait_SetDeformerVector4ArrayVariable::SetValue(UOptimusDeformerInstance* InInstance) const
+{
+	InInstance->SetVector4ArrayVariable(*GetName(), Value);
+}
+
+void FRigVMTrait_SetDeformerLinearColorVariable::SetValue(UOptimusDeformerInstance* InInstance) const
+{
+	InInstance->SetLinearColorVariable(*GetName(), Value);
+}
+
+void FRigVMTrait_SetDeformerLinearColorArrayVariable::SetValue(UOptimusDeformerInstance* InInstance) const
+{
+	InInstance->SetLinearColorArrayVariable(*GetName(), Value);
+}
+
+void FRigVMTrait_SetDeformerQuatVariable::SetValue(UOptimusDeformerInstance* InInstance) const
+{
+	InInstance->SetQuatVariable(*GetName(), Value);
+}
+
+void FRigVMTrait_SetDeformerQuatArrayVariable::SetValue(UOptimusDeformerInstance* InInstance) const
+{
+	InInstance->SetQuatArrayVariable(*GetName(), Value);
+}
+
+void FRigVMTrait_SetDeformerRotatorVariable::SetValue(UOptimusDeformerInstance* InInstance) const
+{
+	InInstance->SetRotatorVariable(*GetName(), Value);
+}
+
+void FRigVMTrait_SetDeformerRotatorArrayVariable::SetValue(UOptimusDeformerInstance* InInstance) const
+{
+	InInstance->SetRotatorArrayVariable(*GetName(), Value);
+}
+
+void FRigVMTrait_SetDeformerTransformVariable::SetValue(UOptimusDeformerInstance* InInstance) const
+{
+	InInstance->SetTransformVariable(*GetName(), Value);
+}
+
+void FRigVMTrait_SetDeformerTransformArrayVariable::SetValue(UOptimusDeformerInstance* InInstance) const
+{
+	InInstance->SetTransformArrayVariable(*GetName(), Value);
+}
diff --git a/Engine/Plugins/Animation/DeformerGraph/Source/OptimusCore/Private/ControlRig/RigUnit_Optimus.h b/Engine/Plugins/Animation/DeformerGraph/Source/OptimusCore/Private/ControlRig/RigUnit_Optimus.h
new file mode 100644
--- /dev/null
+++ b/Engine/Plugins/Animation/DeformerGraph/Source/OptimusCore/Private/ControlRig/RigUnit_Optimus.h
@@ -0,0 +1,345 @@
+﻿// Copyright Epic Games, Inc. All Rights Reserved.
+#pragma once
+#include "CoreMinimal.h"
+#include "Units/RigUnit.h"
+#include "RigVMCore/RigVMTrait.h"
+#include "OptimusDeformer.h"
+#include "OptimusDeformerDynamicInstanceManager.h"
+#include "RigUnit_Optimus.generated.h"
+
+USTRUCT()
+struct OPTIMUSCORE_API FRigVMTrait_OptimusDeformer: public FRigVMTrait
+{
+	GENERATED_BODY()
+
+	FRigVMTrait_OptimusDeformer() = default;
+	FString GetDisplayName() const override;
+#if WITH_EDITOR
+	bool ShouldCreatePinForProperty(const FProperty* InProperty) const override;
+#endif
+
+	UPROPERTY(EditAnywhere, Category = "Deformer Graph")
+	TSoftObjectPtr<UOptimusDeformer> DeformerGraph;
+};
+
+
+USTRUCT()
+struct OPTIMUSCORE_API FRigVMTrait_OptimusDeformerSettings: public FRigVMTrait
+{
+	GENERATED_BODY()
+
+	FRigVMTrait_OptimusDeformerSettings() = default;
+
+	UPROPERTY(EditAnywhere, Category = "Trait", meta=(Input))
+	EOptimusDeformerExecutionPhase ExecutionPhase = EOptimusDeformerExecutionPhase::AfterDefaultDeformer;
+	
+	// Deformers are first sorted by execution group index, then by the order in which they are added
+	UPROPERTY(EditAnywhere, Category = "Trait", meta=(Input))
+	int32 ExecutionGroup = 1;
+};
+
+/** Adds a deformer to the Skeletal Mesh Component*/
+USTRUCT(meta = (DisplayName = "Add Deformer", Category = "Deformer Graph", NodeColor="1 1 1"))
+struct OPTIMUSCORE_API FRigUnit_AddOptimusDeformer: public FRigUnitMutable
+{
+	GENERATED_BODY()
+	
+	static const TCHAR* DeformerTraitName;
+	static const TCHAR* DeformerSettingsTraitName;
+
+	static bool IsVariableTraitName(const FString& InTraitName);
+
+	FRigUnit_AddOptimusDeformer() = default;
+	
+	void OnUnitNodeCreated(FRigVMUnitNodeCreatedContext& InContext) const override;
+	TArray<FRigVMUserWorkflow> GetSupportedWorkflows(const UObject* InSubject) const override;
+	
+	/** Execute logic for this rig unit */
+	RIGVM_METHOD()
+	virtual void Execute() override;
+
+	UPROPERTY(meta=(Hidden))
+	FGuid DeformerInstanceGuid;
+};
+
+
+
+USTRUCT(meta=(Abstract))
+struct OPTIMUSCORE_API FRigVMTrait_OptimusVariableBase: public FRigVMTrait
+{
+	GENERATED_BODY()
+
+	virtual void SetValue(UOptimusDeformerInstance* InInstance) const PURE_VIRTUAL(FRigVMTrait_OptimusVariableBase::SetValue, {});
+};
+
+USTRUCT(BlueprintType)
+struct OPTIMUSCORE_API FRigVMTrait_SetDeformerIntVariable: public FRigVMTrait_OptimusVariableBase
+{
+	GENERATED_BODY()
+
+	void SetValue(UOptimusDeformerInstance* InInstance) const override;
+	
+	UPROPERTY(EditAnywhere, Category = "Trait")
+	int32 Value;	
+};
+
+USTRUCT(BlueprintType)
+struct OPTIMUSCORE_API FRigVMTrait_SetDeformerIntArrayVariable: public FRigVMTrait_OptimusVariableBase
+{
+	GENERATED_BODY()
+
+	void SetValue(UOptimusDeformerInstance* InInstance) const override;
+	
+	UPROPERTY(EditAnywhere, Category = "Trait")
+	TArray<int32> Value;	
+};
+
+USTRUCT(BlueprintType)
+struct OPTIMUSCORE_API FRigVMTrait_SetDeformerInt2Variable: public FRigVMTrait_OptimusVariableBase
+{
+	GENERATED_BODY()
+
+	void SetValue(UOptimusDeformerInstance* InInstance) const override;
+	
+	UPROPERTY(EditAnywhere, Category = "Trait")
+	FIntPoint Value;	
+};
+
+USTRUCT(BlueprintType)
+struct OPTIMUSCORE_API FRigVMTrait_SetDeformerInt2ArrayVariable: public FRigVMTrait_OptimusVariableBase
+{
+	GENERATED_BODY()
+
+	void SetValue(UOptimusDeformerInstance* InInstance) const override;
+	
+	UPROPERTY(EditAnywhere, Category = "Trait")
+	TArray<FIntPoint> Value;	
+};
+
+USTRUCT(BlueprintType)
+struct OPTIMUSCORE_API FRigVMTrait_SetDeformerInt3Variable: public FRigVMTrait_OptimusVariableBase
+{
+	GENERATED_BODY()
+
+	void SetValue(UOptimusDeformerInstance* InInstance) const override;
+	
+	UPROPERTY(EditAnywhere, Category = "Trait")
+	FIntVector Value;	
+};
+
+USTRUCT(BlueprintType)
+struct OPTIMUSCORE_API FRigVMTrait_SetDeformerInt3ArrayVariable: public FRigVMTrait_OptimusVariableBase
+{
+	GENERATED_BODY()
+
+	void SetValue(UOptimusDeformerInstance* InInstance) const override;
+	
+	UPROPERTY(EditAnywhere, Category = "Trait")
+	TArray<FIntVector> Value;	
+};
+
+USTRUCT(BlueprintType)
+struct OPTIMUSCORE_API FRigVMTrait_SetDeformerInt4Variable: public FRigVMTrait_OptimusVariableBase
+{
+	GENERATED_BODY()
+
+	void SetValue(UOptimusDeformerInstance* InInstance) const override;
+	
+	UPROPERTY(EditAnywhere, Category = "Trait")
+	FIntVector4 Value;	
+};
+
+USTRUCT(BlueprintType)
+struct OPTIMUSCORE_API FRigVMTrait_SetDeformerInt4ArrayVariable: public FRigVMTrait_OptimusVariableBase
+{
+	GENERATED_BODY()
+
+	void SetValue(UOptimusDeformerInstance* InInstance) const override;
+	
+	UPROPERTY(EditAnywhere, Category = "Trait")
+	TArray<FIntVector4> Value;	
+};
+
+
+
+
+USTRUCT(BlueprintType)
+struct OPTIMUSCORE_API FRigVMTrait_SetDeformerFloatVariable: public FRigVMTrait_OptimusVariableBase
+{
+	GENERATED_BODY()
+	
+	void SetValue(UOptimusDeformerInstance* InInstance) const override;
+	
+	UPROPERTY(EditAnywhere, Category = "Trait")
+	double Value;	
+};
+
+USTRUCT(BlueprintType)
+struct OPTIMUSCORE_API FRigVMTrait_SetDeformerFloatArrayVariable: public FRigVMTrait_OptimusVariableBase
+{
+	GENERATED_BODY()
+	
+	void SetValue(UOptimusDeformerInstance* InInstance) const override;
+	
+	UPROPERTY(EditAnywhere, Category = "Trait")
+	TArray<double> Value;	
+};
+
+USTRUCT(BlueprintType)
+struct OPTIMUSCORE_API FRigVMTrait_SetDeformerVector2Variable: public FRigVMTrait_OptimusVariableBase
+{
+	GENERATED_BODY()
+
+	void SetValue(UOptimusDeformerInstance* InInstance) const override;
+	
+	UPROPERTY(EditAnywhere, Category = "Trait")
+	FVector2D Value;	
+};
+
+
+USTRUCT(BlueprintType)
+struct OPTIMUSCORE_API FRigVMTrait_SetDeformerVector2ArrayVariable: public FRigVMTrait_OptimusVariableBase
+{
+	GENERATED_BODY()
+
+	void SetValue(UOptimusDeformerInstance* InInstance) const override;
+	
+	UPROPERTY(EditAnywhere, Category = "Trait")
+	TArray<FVector2D> Value;	
+};
+
+USTRUCT(BlueprintType)
+struct OPTIMUSCORE_API FRigVMTrait_SetDeformerVectorVariable: public FRigVMTrait_OptimusVariableBase
+{
+	GENERATED_BODY()
+
+	void SetValue(UOptimusDeformerInstance* InInstance) const override;
+	
+	UPROPERTY(EditAnywhere, Category = "Trait")
+	FVector Value;	
+};
+
+USTRUCT(BlueprintType)
+struct OPTIMUSCORE_API FRigVMTrait_SetDeformerVectorArrayVariable: public FRigVMTrait_OptimusVariableBase
+{
+	GENERATED_BODY()
+
+	void SetValue(UOptimusDeformerInstance* InInstance) const override;
+	
+	UPROPERTY(EditAnywhere, Category = "Trait")
+	TArray<FVector> Value;	
+};
+
+
+USTRUCT(BlueprintType)
+struct OPTIMUSCORE_API FRigVMTrait_SetDeformerVector4Variable: public FRigVMTrait_OptimusVariableBase
+{
+	GENERATED_BODY()
+
+	void SetValue(UOptimusDeformerInstance* InInstance) const override;
+	
+	UPROPERTY(EditAnywhere, Category = "Trait")
+	FVector4 Value;	
+};
+
+USTRUCT(BlueprintType)
+struct OPTIMUSCORE_API FRigVMTrait_SetDeformerVector4ArrayVariable: public FRigVMTrait_OptimusVariableBase
+{
+	GENERATED_BODY()
+
+	void SetValue(UOptimusDeformerInstance* InInstance) const override;
+	
+	UPROPERTY(EditAnywhere, Category = "Trait")
+	TArray<FVector4> Value;	
+};
+
+USTRUCT(BlueprintType)
+struct OPTIMUSCORE_API FRigVMTrait_SetDeformerLinearColorVariable: public FRigVMTrait_OptimusVariableBase
+{
+	GENERATED_BODY()
+
+	void SetValue(UOptimusDeformerInstance* InInstance) const override;
+	
+	UPROPERTY(EditAnywhere, Category = "Trait")
+	FLinearColor Value;	
+};
+
+USTRUCT(BlueprintType)
+struct OPTIMUSCORE_API FRigVMTrait_SetDeformerLinearColorArrayVariable: public FRigVMTrait_OptimusVariableBase
+{
+	GENERATED_BODY()
+
+	void SetValue(UOptimusDeformerInstance* InInstance) const override;
+	
+	UPROPERTY(EditAnywhere, Category = "Trait")
+	TArray<FLinearColor> Value;	
+};
+
+
+USTRUCT(BlueprintType)
+struct OPTIMUSCORE_API FRigVMTrait_SetDeformerQuatVariable: public FRigVMTrait_OptimusVariableBase
+{
+	GENERATED_BODY()
+
+	void SetValue(UOptimusDeformerInstance* InInstance) const override;
+	
+	UPROPERTY(EditAnywhere, Category = "Trait")
+	FQuat Value;	
+};
+
+
+USTRUCT(BlueprintType)
+struct OPTIMUSCORE_API FRigVMTrait_SetDeformerQuatArrayVariable: public FRigVMTrait_OptimusVariableBase
+{
+	GENERATED_BODY()
+
+	void SetValue(UOptimusDeformerInstance* InInstance) const override;
+	
+	UPROPERTY(EditAnywhere, Category = "Trait")
+	TArray<FQuat> Value;	
+};
+
+USTRUCT(BlueprintType)
+struct OPTIMUSCORE_API FRigVMTrait_SetDeformerRotatorVariable: public FRigVMTrait_OptimusVariableBase
+{
+	GENERATED_BODY()
+
+	void SetValue(UOptimusDeformerInstance* InInstance) const override;
+	
+	UPROPERTY(EditAnywhere, Category = "Trait")
+	FRotator Value;	
+};
+
+USTRUCT(BlueprintType)
+struct OPTIMUSCORE_API FRigVMTrait_SetDeformerRotatorArrayVariable: public FRigVMTrait_OptimusVariableBase
+{
+	GENERATED_BODY()
+
+	void SetValue(UOptimusDeformerInstance* InInstance) const override;
+	
+	UPROPERTY(EditAnywhere, Category = "Trait")
+	TArray<FRotator> Value;	
+};
+
+USTRUCT(BlueprintType)
+struct OPTIMUSCORE_API FRigVMTrait_SetDeformerTransformVariable: public FRigVMTrait_OptimusVariableBase
+{
+	GENERATED_BODY()
+
+	void SetValue(UOptimusDeformerInstance* InInstance) const override;
+	
+	UPROPERTY(EditAnywhere, Category = "Trait")
+	FTransform Value;	
+};
+
+USTRUCT(BlueprintType)
+struct OPTIMUSCORE_API FRigVMTrait_SetDeformerTransformArrayVariable: public FRigVMTrait_OptimusVariableBase
+{
+	GENERATED_BODY()
+
+	void SetValue(UOptimusDeformerInstance* InInstance) const override;
+	
+	UPROPERTY(EditAnywhere, Category = "Trait")
+	TArray<FTransform> Value;	
+};
+
diff --git a/Engine/Plugins/Animation/DeformerGraph/Source/OptimusCore/Private/DataInterfaces/OptimusDataInterfaceGraph.cpp b/Engine/Plugins/Animation/DeformerGraph/Source/OptimusCore/Private/DataInterfaces/OptimusDataInterfaceGraph.cpp
--- a/Engine/Plugins/Animation/DeformerGraph/Source/OptimusCore/Private/DataInterfaces/OptimusDataInterfaceGraph.cpp
+++ b/Engine/Plugins/Animation/DeformerGraph/Source/OptimusCore/Private/DataInterfaces/OptimusDataInterfaceGraph.cpp
@@ -190,11 +190,6 @@ void UOptimusGraphDataProvider::SetDeformerInstance(UOptimusDeformerInstance* In
 	DeformerInstance = InInstance;
 }
 
-UOptimusDeformerInstance* UOptimusGraphDataProvider::GetDeformerInstance() const
-{
-	return DeformerInstance;
-}
-
 FOptimusGraphDataProviderProxy::FOptimusGraphDataProviderProxy(
 	UOptimusDeformerInstance const* DeformerInstance,
 	TArray<FOptimusGraphVariableDescription> const& Variables,
diff --git a/Engine/Plugins/Animation/DeformerGraph/Source/OptimusCore/Private/DataInterfaces/OptimusDataInterfaceGraph.h b/Engine/Plugins/Animation/DeformerGraph/Source/OptimusCore/Private/DataInterfaces/OptimusDataInterfaceGraph.h
--- a/Engine/Plugins/Animation/DeformerGraph/Source/OptimusCore/Private/DataInterfaces/OptimusDataInterfaceGraph.h
+++ b/Engine/Plugins/Animation/DeformerGraph/Source/OptimusCore/Private/DataInterfaces/OptimusDataInterfaceGraph.h
@@ -111,7 +111,6 @@ public:
 
 	//~ Begin IOptimusDeformerInstanceAccessor Interface
 	void SetDeformerInstance(UOptimusDeformerInstance* InInstance) override;
-	UOptimusDeformerInstance* GetDeformerInstance() const override;
 	//~ End IOptimusDeformerInstanceAccessor Interface
 
 private:
diff --git a/Engine/Plugins/Animation/DeformerGraph/Source/OptimusCore/Private/DataInterfaces/OptimusDataInterfaceSkinnedMeshRead.cpp b/Engine/Plugins/Animation/DeformerGraph/Source/OptimusCore/Private/DataInterfaces/OptimusDataInterfaceSkinnedMeshRead.cpp
new file mode 100644
--- /dev/null
+++ b/Engine/Plugins/Animation/DeformerGraph/Source/OptimusCore/Private/DataInterfaces/OptimusDataInterfaceSkinnedMeshRead.cpp
@@ -0,0 +1,321 @@
+﻿// Copyright Epic Games, Inc. All Rights Reserved.
+
+#include "OptimusDataInterfaceSkinnedMeshRead.h"
+
+#include "Components/SkinnedMeshComponent.h"
+#include "ComputeFramework/ShaderParamTypeDefinition.h"
+#include "OptimusDataDomain.h"
+#include "RenderGraphBuilder.h"
+#include "Rendering/SkeletalMeshLODRenderData.h"
+#include "Rendering/SkeletalMeshRenderData.h"
+#include "ShaderCompilerCore.h"
+#include "ShaderParameterMetadataBuilder.h"
+#include "SkeletalMeshDeformerHelpers.h"
+#include "SkeletalRenderPublic.h"
+#include "DataDrivenShaderPlatformInfo.h"
+#include "OptimusDeformerInstance.h"
+#include "Animation/MeshDeformerInstance.h"
+#include "ComputeFramework/ComputeKernelPermutationSet.h"
+#include "ComputeFramework/ComputeKernelPermutationVector.h"
+
+#include UE_INLINE_GENERATED_CPP_BY_NAME(OptimusDataInterfaceSkinnedMeshRead)
+
+const TCHAR* UOptimusSkinnedMeshReadDataInterface::ReadableOutputBufferPermutationName = TEXT("READABLE_OUTPUT_BUFFERS");
+
+FString UOptimusSkinnedMeshReadDataInterface::GetDisplayName() const
+{
+	return TEXT("Read Skinned Mesh");
+}
+
+FName UOptimusSkinnedMeshReadDataInterface::GetCategory() const
+{
+	return CategoryName::OutputDataInterfaces;
+}
+
+TArray<FOptimusCDIPinDefinition> UOptimusSkinnedMeshReadDataInterface::GetPinDefinitions() const
+{
+	TArray<FOptimusCDIPinDefinition> Defs;
+	Defs.Add({ "Position", "ReadPosition", Optimus::DomainName::Vertex, "ReadNumVertices" });
+	Defs.Add({ "TangentX", "ReadTangentX", Optimus::DomainName::Vertex, "ReadNumVertices" });
+	Defs.Add({ "TangentZ", "ReadTangentZ", Optimus::DomainName::Vertex, "ReadNumVertices" });
+	Defs.Add({ "Color", "ReadColor", Optimus::DomainName::Vertex, "ReadNumVertices" });
+	return Defs;
+}
+
+
+TSubclassOf<UActorComponent> UOptimusSkinnedMeshReadDataInterface::GetRequiredComponentClass() const
+{
+	return USkinnedMeshComponent::StaticClass();
+}
+
+// Should be kept in sync with GetSupportedInput
+enum class ESkinnedMeshReadDataInterfaceInputSelectorMask : uint64
+{
+	NumVertices = 1 << 0,
+	Position = 1 << 1,
+	TangentX = 1 << 2,
+	TangentZ = 1 << 3,
+	Color = 1 << 4,
+};
+ENUM_CLASS_FLAGS(ESkinnedMeshReadDataInterfaceInputSelectorMask);
+
+void UOptimusSkinnedMeshReadDataInterface::GetSupportedInputs(TArray<FShaderFunctionDefinition>& OutFunctions) const
+{
+	OutFunctions.AddDefaulted_GetRef()
+		.SetName(TEXT("ReadNumVertices"))
+		.AddReturnType(EShaderFundamentalType::Uint);
+	
+	OutFunctions.AddDefaulted_GetRef()
+		.SetName(TEXT("ReadPosition"))
+		.AddReturnType(EShaderFundamentalType::Float, 3)
+		.AddParam(EShaderFundamentalType::Uint);
+
+	OutFunctions.AddDefaulted_GetRef()
+		.SetName(TEXT("ReadTangentX"))
+		.AddReturnType(EShaderFundamentalType::Float, 4)
+		.AddParam(EShaderFundamentalType::Uint);
+
+	OutFunctions.AddDefaulted_GetRef()
+		.SetName(TEXT("ReadTangentZ"))
+		.AddReturnType(EShaderFundamentalType::Float, 4)
+		.AddParam(EShaderFundamentalType::Uint);
+
+	OutFunctions.AddDefaulted_GetRef()
+		.SetName(TEXT("ReadColor"))
+		.AddReturnType(EShaderFundamentalType::Float, 4)
+		.AddParam(EShaderFundamentalType::Uint);
+}
+
+BEGIN_SHADER_PARAMETER_STRUCT(FSkinnedMeshReadDataInterfaceParameters, )
+	SHADER_PARAMETER(uint32, NumVertices)
+	SHADER_PARAMETER(uint32, InputStreamStart)
+	SHADER_PARAMETER_RDG_BUFFER_UAV(RWBuffer<float>, PositionBufferUAV)
+	SHADER_PARAMETER_RDG_BUFFER_UAV(RWBuffer<SNORM float4>, TangentBufferUAV)
+	SHADER_PARAMETER_RDG_BUFFER_UAV(RWBuffer<UNORM float4>, ColorBufferUAV)
+	SHADER_PARAMETER_SRV(Buffer<float>, PositionStaticBuffer)
+	SHADER_PARAMETER_SRV(Buffer<SNORM float4>, TangentStaticBuffer)
+	SHADER_PARAMETER_SRV(Buffer<float4>, ColorStaticBuffer)
+	SHADER_PARAMETER(uint32, ColorIndexMask)
+END_SHADER_PARAMETER_STRUCT()
+
+void UOptimusSkinnedMeshReadDataInterface::GetShaderParameters(TCHAR const* UID, FShaderParametersMetadataBuilder& InOutBuilder, FShaderParametersMetadataAllocations& InOutAllocations) const
+{
+	InOutBuilder.AddNestedStruct<FSkinnedMeshReadDataInterfaceParameters>(UID);
+}
+
+TCHAR const* UOptimusSkinnedMeshReadDataInterface::TemplateFilePath = TEXT("/Plugin/Optimus/Private/DataInterfaceSkinnedMeshRead.ush");
+
+TCHAR const* UOptimusSkinnedMeshReadDataInterface::GetShaderVirtualPath() const
+{
+	return TemplateFilePath;
+}
+
+void UOptimusSkinnedMeshReadDataInterface::GetShaderHash(FString& InOutKey) const
+{
+	GetShaderFileHash(TemplateFilePath, EShaderPlatform::SP_PCD3D_SM5).AppendString(InOutKey);
+}
+
+void UOptimusSkinnedMeshReadDataInterface::GetHLSL(FString& OutHLSL, FString const& InDataInterfaceName) const
+{
+	TMap<FString, FStringFormatArg> TemplateArgs =
+	{
+		{ TEXT("DataInterfaceName"), InDataInterfaceName },
+	};
+
+	FString TemplateFile;
+	LoadShaderSourceFile(TemplateFilePath, EShaderPlatform::SP_PCD3D_SM5, &TemplateFile, nullptr);
+	OutHLSL += FString::Format(*TemplateFile, TemplateArgs);
+}
+
+enum class EOptimusSkinnedMeshReadReadableOutputBuffer: uint32
+{
+	None = 0,
+	Position = 1<<0,
+	Tangents = 1<<1,
+	VertexColor = 1<<2,
+	NumPermutations = Position | Tangents | VertexColor
+};
+ENUM_CLASS_FLAGS(EOptimusSkinnedMeshReadReadableOutputBuffer);
+
+void UOptimusSkinnedMeshReadDataInterface::GetDefines(FComputeKernelDefinitionSet& OutDefinitionSet) const
+{
+	OutDefinitionSet.Defines.Add(FComputeKernelDefinition(
+		TEXT("OPTIMUS_SKINNED_MESH_READ_POSITION"),
+		FString::FromInt(static_cast<uint32>(EOptimusSkinnedMeshReadReadableOutputBuffer::Position))));
+	OutDefinitionSet.Defines.Add(FComputeKernelDefinition(
+		TEXT("OPTIMUS_SKINNED_MESH_READ_TANGENTS"),
+		FString::FromInt(static_cast<uint32>(EOptimusSkinnedMeshReadReadableOutputBuffer::Tangents))));
+	OutDefinitionSet.Defines.Add(FComputeKernelDefinition(
+		TEXT("OPTIMUS_SKINNED_MESH_READ_COLOR"),
+		FString::FromInt(static_cast<uint32>(EOptimusSkinnedMeshReadReadableOutputBuffer::VertexColor))));
+}
+
+
+
+void UOptimusSkinnedMeshReadDataInterface::GetPermutations(FComputeKernelPermutationVector& OutPermutationVector) const
+{
+	OutPermutationVector.AddPermutation(ReadableOutputBufferPermutationName, static_cast<uint32>(EOptimusSkinnedMeshReadReadableOutputBuffer::NumPermutations));
+}
+
+UComputeDataProvider* UOptimusSkinnedMeshReadDataInterface::CreateDataProvider(TObjectPtr<UObject> InBinding, uint64 InInputMask, uint64 InOutputMask) const
+{
+	UOptimusSkinnedMeshReadDataProvider* Provider = NewObject<UOptimusSkinnedMeshReadDataProvider>();
+	Provider->SkinnedMesh = Cast<USkinnedMeshComponent>(InBinding);
+	Provider->InputMask = InInputMask;
+	return Provider;
+}
+
+FComputeDataProviderRenderProxy* UOptimusSkinnedMeshReadDataProvider::GetRenderProxy()
+{
+	return new FOptimusSkinnedMeshReadDataProviderProxy(SkinnedMesh, InputMask, DeformerInstance->OutputBuffersFromPreviousInstances);
+}
+
+void UOptimusSkinnedMeshReadDataProvider::SetDeformerInstance(UOptimusDeformerInstance* InInstance)
+{
+	DeformerInstance = InInstance;
+}
+
+FOptimusSkinnedMeshReadDataProviderProxy::FOptimusSkinnedMeshReadDataProviderProxy(USkinnedMeshComponent* InSkinnedMeshComponent, uint64 InInputMask, EMeshDeformerOutputBuffer InOutputBuffersFromPreviousInstances)
+{
+	SkeletalMeshObject = InSkinnedMeshComponent != nullptr ? InSkinnedMeshComponent->MeshObject : nullptr;
+	InputMask = InInputMask;
+	OutputBuffersFromPreviousInstances = InOutputBuffersFromPreviousInstances;
+}
+
+bool FOptimusSkinnedMeshReadDataProviderProxy::IsValid(FValidationData const& InValidationData) const
+{
+	if (InValidationData.ParameterStructSize != sizeof(FParameters))
+	{
+		return false;
+	}
+	if (SkeletalMeshObject == nullptr)
+	{
+		return false;
+	}
+	if (SkeletalMeshObject->IsCPUSkinned())
+	{
+		return false;
+	}
+	if (SkeletalMeshObject->GetSkeletalMeshRenderData().LODRenderData[SkeletalMeshObject->GetLOD()].RenderSections.Num() != InValidationData.NumInvocations)
+	{
+		return false;
+	}
+
+	return true;
+}
+
+void FOptimusSkinnedMeshReadDataProviderProxy::AllocateResources(FRDGBuilder& GraphBuilder)
+{
+	// Allocate required buffers
+	const int32 LodIndex = SkeletalMeshObject->GetLOD();
+	
+	if (InputMask & static_cast<uint64>(ESkinnedMeshReadDataInterfaceInputSelectorMask::Position) )
+	{
+		FRDGBuffer* PositionBuffer = FSkeletalMeshDeformerHelpers::AllocateVertexFactoryPositionBuffer(GraphBuilder, SkeletalMeshObject, LodIndex, TEXT("OptimusSkinnedMeshPosition"));
+		PositionBufferUAV = GraphBuilder.CreateUAV(PositionBuffer, PF_R32_FLOAT, ERDGUnorderedAccessViewFlags::SkipBarrier);
+	}
+	else
+	{
+		PositionBufferUAV = GraphBuilder.CreateUAV(GraphBuilder.RegisterExternalBuffer(GWhiteVertexBufferWithRDG->Buffer), PF_R32_FLOAT);
+	}	
+
+	// OpenGL ES does not support writing to RGBA16_SNORM images, instead pack data into SINT in the shader
+	const EPixelFormat TangentsFormat = IsOpenGLPlatform(GMaxRHIShaderPlatform) ? PF_R16G16B16A16_SINT : PF_R16G16B16A16_SNORM;
+
+	if ((InputMask & static_cast<uint64>(ESkinnedMeshReadDataInterfaceInputSelectorMask::TangentX)) ||
+		(InputMask & static_cast<uint64>(ESkinnedMeshReadDataInterfaceInputSelectorMask::TangentZ)))
+	{
+		FRDGBuffer* TangentBuffer = FSkeletalMeshDeformerHelpers::AllocateVertexFactoryTangentBuffer(GraphBuilder, SkeletalMeshObject, LodIndex, TEXT("OptimusSkinnedMeshTangent"));
+		TangentBufferUAV = GraphBuilder.CreateUAV(TangentBuffer, TangentsFormat, ERDGUnorderedAccessViewFlags::SkipBarrier);
+	}
+	else
+	{
+		TangentBufferUAV = GraphBuilder.CreateUAV(GraphBuilder.RegisterExternalBuffer(GWhiteVertexBufferWithRDG->Buffer), TangentsFormat);
+	}
+
+	if (InputMask & static_cast<uint64>(ESkinnedMeshReadDataInterfaceInputSelectorMask::Color))
+	{
+		FRDGBuffer* ColorBuffer = FSkeletalMeshDeformerHelpers::AllocateVertexFactoryColorBuffer(GraphBuilder, SkeletalMeshObject, LodIndex, TEXT("OptimusSkinnedMeshColor"));
+		// using RGBA here and do a manual fetch swizzle in shader instead of BGRA directly because some Mac does not support it. See GMetalBufferFormats[PF_B8G8R8A8] 
+		ColorBufferUAV = GraphBuilder.CreateUAV(ColorBuffer, PF_R8G8B8A8, ERDGUnorderedAccessViewFlags::SkipBarrier);
+	}
+	else
+	{
+		ColorBufferUAV = GraphBuilder.CreateUAV(GraphBuilder.RegisterExternalBuffer(GWhiteVertexBufferWithRDG->Buffer), PF_A32B32G32R32F);
+	}
+
+	FSkeletalMeshDeformerHelpers::UpdateVertexFactoryBufferOverrides(GraphBuilder.RHICmdList, SkeletalMeshObject, LodIndex);
+}
+
+
+
+
+
+void FOptimusSkinnedMeshReadDataProviderProxy::GatherPermutations(FPermutationData& InOutPermutationData) const
+{
+	EOptimusSkinnedMeshReadReadableOutputBuffer ReadableOutputBuffers = EOptimusSkinnedMeshReadReadableOutputBuffer::None; 
+	if (EnumHasAnyFlags(OutputBuffersFromPreviousInstances, EMeshDeformerOutputBuffer::SkinnedMeshPosition))
+	{
+		ReadableOutputBuffers |= EOptimusSkinnedMeshReadReadableOutputBuffer::Position;
+	}
+	if (EnumHasAnyFlags(OutputBuffersFromPreviousInstances, EMeshDeformerOutputBuffer::SkinnedMeshTangents))
+	{
+		ReadableOutputBuffers |= EOptimusSkinnedMeshReadReadableOutputBuffer::Tangents;
+	}
+	if (EnumHasAnyFlags(OutputBuffersFromPreviousInstances, EMeshDeformerOutputBuffer::SkinnedMeshVertexColor))
+	{
+		ReadableOutputBuffers |= EOptimusSkinnedMeshReadReadableOutputBuffer::VertexColor;
+	}
+	static FString Name(UOptimusSkinnedMeshReadDataInterface::ReadableOutputBufferPermutationName);
+	static uint32 Hash = GetTypeHash(Name);
+	const uint32 ReadableOutputBuffersBits = InOutPermutationData.PermutationVector.GetPermutationBits(Name, Hash, static_cast<uint32>(ReadableOutputBuffers));
+	
+	for (int32 InvocationIndex = 0; InvocationIndex < InOutPermutationData.NumInvocations; ++InvocationIndex)
+	{
+		InOutPermutationData.PermutationIds[InvocationIndex] |= ReadableOutputBuffersBits;
+	}
+}
+
+
+
+
+void FOptimusSkinnedMeshReadDataProviderProxy::GatherDispatchData(FDispatchData const& InDispatchData)
+{
+	const int32 LodIndex = SkeletalMeshObject->GetLOD();
+	FSkeletalMeshRenderData const& SkeletalMeshRenderData = SkeletalMeshObject->GetSkeletalMeshRenderData();
+	FSkeletalMeshLODRenderData const* LodRenderData = &SkeletalMeshRenderData.LODRenderData[LodIndex];
+
+	FRHIShaderResourceView* NullSRVBinding = GWhiteVertexBufferWithSRV->ShaderResourceViewRHI.GetReference();
+
+
+	
+	const TStridedView<FParameters> ParameterArray = MakeStridedParameterView<FParameters>(InDispatchData);
+	for (int32 InvocationIndex = 0; InvocationIndex < ParameterArray.Num(); ++InvocationIndex)
+	{
+		FSkelMeshRenderSection const& RenderSection = LodRenderData->RenderSections[InvocationIndex];
+
+		FParameters& Parameters = ParameterArray[InvocationIndex];
+		Parameters.NumVertices = InDispatchData.bUnifiedDispatch ? LodRenderData->GetNumVertices() : RenderSection.GetNumVertices();
+		Parameters.InputStreamStart= InDispatchData.bUnifiedDispatch ? 0 : RenderSection.GetVertexBufferIndex();
+
+		Parameters.PositionBufferUAV = PositionBufferUAV;
+		Parameters.TangentBufferUAV = TangentBufferUAV;
+		Parameters.ColorBufferUAV = ColorBufferUAV;
+		
+		FRHIShaderResourceView* MeshVertexBufferSRV = LodRenderData->StaticVertexBuffers.PositionVertexBuffer.GetSRV();
+		FRHIShaderResourceView* MeshTangentBufferSRV = LodRenderData->StaticVertexBuffers.StaticMeshVertexBuffer.GetTangentsSRV();
+		FRHIShaderResourceView* MeshColorBufferSRV = LodRenderData->StaticVertexBuffers.ColorVertexBuffer.GetColorComponentsSRV();
+
+		Parameters.PositionStaticBuffer = MeshVertexBufferSRV != nullptr ? MeshVertexBufferSRV : NullSRVBinding;
+		Parameters.TangentStaticBuffer = MeshTangentBufferSRV != nullptr ? MeshTangentBufferSRV : NullSRVBinding;
+		Parameters.ColorStaticBuffer = MeshColorBufferSRV != nullptr ? MeshColorBufferSRV : NullSRVBinding;
+		
+		// Basically when we are accessing GWhiteVertexBufferWithSRV(NullSRVBinding),
+		// we should not access beyond index 0 since the buffer is only a few bytes
+		
+		// See FGPUSkinPassthroughVertexFactory::UpdateUniformBuffer() and LocalVertexFactory.ush :: GetVertexFactoryIntermediates()
+		// Ideally we should be getting this value from the GPUBaseSkinVertexFactory but the need for
+		// section index make it tricky when we are doing unified dispatch
+		Parameters.ColorIndexMask = MeshColorBufferSRV != nullptr ? ~0u : 0;
+	}
+}
diff --git a/Engine/Plugins/Animation/DeformerGraph/Source/OptimusCore/Private/DataInterfaces/OptimusDataInterfaceSkinnedMeshRead.h b/Engine/Plugins/Animation/DeformerGraph/Source/OptimusCore/Private/DataInterfaces/OptimusDataInterfaceSkinnedMeshRead.h
new file mode 100644
--- /dev/null
+++ b/Engine/Plugins/Animation/DeformerGraph/Source/OptimusCore/Private/DataInterfaces/OptimusDataInterfaceSkinnedMeshRead.h
@@ -0,0 +1,102 @@
+﻿// Copyright Epic Games, Inc. All Rights Reserved.
+
+#pragma once
+
+#include "IOptimusDeformerInstanceAccessor.h"
+#include "OptimusComputeDataInterface.h"
+#include "ComputeFramework/ComputeDataProvider.h"
+
+#include "OptimusDataInterfaceSkinnedMeshRead.generated.h"
+
+class FRDGBuffer;
+class FRDGBufferSRV;
+class FSkeletalMeshObject;
+class FSkinnedMeshReadDataInterfaceParameters;
+class USkinnedMeshComponent;
+enum class EMeshDeformerOutputBuffer : uint8;
+
+/** Compute Framework Data Interface for writing skinned mesh. */
+UCLASS(Category = ComputeFramework)
+class OPTIMUSCORE_API UOptimusSkinnedMeshReadDataInterface :
+	public UOptimusComputeDataInterface
+{
+	GENERATED_BODY()
+
+public:
+	static const TCHAR* ReadableOutputBufferPermutationName;
+	//~ Begin UOptimusComputeDataInterface Interface
+	FString GetDisplayName() const override;
+	FName GetCategory() const override;
+	TArray<FOptimusCDIPinDefinition> GetPinDefinitions() const override;
+	TSubclassOf<UActorComponent> GetRequiredComponentClass() const override;
+	//~ End UOptimusComputeDataInterface Interface
+	
+	//~ Begin UComputeDataInterface Interface
+	TCHAR const* GetClassName() const override { return TEXT("SkinnedMeshRead"); }
+	bool CanSupportUnifiedDispatch() const override { return true; }
+	void GetSupportedInputs(TArray<FShaderFunctionDefinition>& OutFunctions) const override;
+	void GetShaderParameters(TCHAR const* UID, FShaderParametersMetadataBuilder& InOutBuilder, FShaderParametersMetadataAllocations& InOutAllocations) const override;
+	TCHAR const* GetShaderVirtualPath() const override;
+	void GetShaderHash(FString& InOutKey) const override;
+	void GetHLSL(FString& OutHLSL, FString const& InDataInterfaceName) const override;
+	void GetDefines(FComputeKernelDefinitionSet& OutDefinitionSet) const override;
+	void GetPermutations(FComputeKernelPermutationVector& OutPermutationVector) const override;
+	UComputeDataProvider* CreateDataProvider(TObjectPtr<UObject> InBinding, uint64 InInputMask, uint64 InOutputMask) const override;
+	//~ End UComputeDataInterface Interface
+
+private:
+	static TCHAR const* TemplateFilePath;
+};
+
+/** Compute Framework Data Provider for writing skinned mesh. */
+UCLASS(BlueprintType, editinlinenew, Category = ComputeFramework)
+class UOptimusSkinnedMeshReadDataProvider :
+	public UComputeDataProvider,
+	public IOptimusDeformerInstanceAccessor
+{
+	GENERATED_BODY()
+
+public:
+	UPROPERTY(EditAnywhere, BlueprintReadWrite, Category = Binding)
+	TObjectPtr<USkinnedMeshComponent> SkinnedMesh = nullptr;
+
+	uint64 InputMask;
+
+	//~ Begin UComputeDataProvider Interface
+	FComputeDataProviderRenderProxy* GetRenderProxy() override;
+	//~ End UComputeDataProvider Interface
+
+	//~ Begin IOptimusDeformerInstanceAccessor Interface
+	void SetDeformerInstance(UOptimusDeformerInstance* InInstance) override;
+	//~ End IOptimusDeformerInstanceAccessor Interface
+
+private:
+	UPROPERTY()
+	TObjectPtr<UOptimusDeformerInstance> DeformerInstance = nullptr;
+};
+
+class FOptimusSkinnedMeshReadDataProviderProxy : public FComputeDataProviderRenderProxy
+{
+public:
+	FOptimusSkinnedMeshReadDataProviderProxy(USkinnedMeshComponent* InSkinnedMeshComponent, uint64 InInputMask, EMeshDeformerOutputBuffer InOutputBuffersWithValidData);
+
+	//~ Begin FComputeDataProviderRenderProxy Interface
+	bool IsValid(FValidationData const& InValidationData) const override;
+	void AllocateResources(FRDGBuilder& GraphBuilder) override;
+	void GatherPermutations(FPermutationData& InOutPermutationData) const override;
+	void GatherDispatchData(FDispatchData const& InDispatchData) override;
+	//~ End FComputeDataProviderRenderProxy Interface
+
+private:
+	using FParameters = FSkinnedMeshReadDataInterfaceParameters;
+
+	FSkeletalMeshObject* SkeletalMeshObject;
+	uint64 InputMask;
+	EMeshDeformerOutputBuffer OutputBuffersFromPreviousInstances;
+
+	// Using UAV here because we might be reading and writing to these buffers in the same kernel
+	// if we have a setup like Read -> Kernel -> Write 
+	FRDGBufferUAV* PositionBufferUAV = nullptr;
+	FRDGBufferUAV* TangentBufferUAV = nullptr;
+	FRDGBufferUAV* ColorBufferUAV = nullptr;
+};
diff --git a/Engine/Plugins/Animation/DeformerGraph/Source/OptimusCore/Private/DataInterfaces/OptimusDataInterfaceSkinnedMeshWrite.cpp b/Engine/Plugins/Animation/DeformerGraph/Source/OptimusCore/Private/DataInterfaces/OptimusDataInterfaceSkinnedMeshWrite.cpp
--- a/Engine/Plugins/Animation/DeformerGraph/Source/OptimusCore/Private/DataInterfaces/OptimusDataInterfaceSkinnedMeshWrite.cpp
+++ b/Engine/Plugins/Animation/DeformerGraph/Source/OptimusCore/Private/DataInterfaces/OptimusDataInterfaceSkinnedMeshWrite.cpp
@@ -52,6 +52,17 @@ void UOptimusSkinnedMeshWriteDataInterface::GetSupportedInputs(TArray<FShaderFun
 		.AddReturnType(EShaderFundamentalType::Uint);
 }
 
+// Should be kept in sync with GetSupportedOutputs
+enum class ESkinnedMeshWriteDataInterfaceOutputSelectorMask : uint64
+{
+	Position = 1 << 0,
+	TangentX = 1 << 1,
+	TangentZ = 1 << 2,
+	Color = 1 << 3,
+};
+ENUM_CLASS_FLAGS(ESkinnedMeshWriteDataInterfaceOutputSelectorMask);
+
+
 void UOptimusSkinnedMeshWriteDataInterface::GetSupportedOutputs(TArray<FShaderFunctionDefinition>& OutFunctions) const
 {
 	OutFunctions.AddDefaulted_GetRef()
@@ -179,11 +190,7 @@ void FOptimusSkinnedMeshWriteDataProviderProxy::AllocateResources(FRDGBuilder& G
 	// Allocate required buffers
 	const int32 LodIndex = SkeletalMeshObject->GetLOD();
 
-	// We will extract buffers from RDG.
-	// It could be better to for memory to use QueueBufferExtraction instead of ConvertToExternalBuffer but that will require an extra hook after graph execution.
-	TRefCountPtr<FRDGPooledBuffer> PositionBufferExternal;
-
-	if (OutputMask & 1)
+	if (OutputMask & static_cast<uint64>(ESkinnedMeshWriteDataInterfaceOutputSelectorMask::Position))
 	{
 		PositionBuffer = FSkeletalMeshDeformerHelpers::AllocateVertexFactoryPositionBuffer(GraphBuilder, SkeletalMeshObject, LodIndex, TEXT("OptimusSkinnedMeshPosition"));
 		PositionBufferUAV = GraphBuilder.CreateUAV(PositionBuffer, PF_R32_FLOAT, ERDGUnorderedAccessViewFlags::SkipBarrier);
@@ -196,7 +203,8 @@ void FOptimusSkinnedMeshWriteDataProviderProxy::AllocateResources(FRDGBuilder& G
 	// OpenGL ES does not support writing to RGBA16_SNORM images, instead pack data into SINT in the shader
 	const EPixelFormat TangentsFormat = IsOpenGLPlatform(GMaxRHIShaderPlatform) ? PF_R16G16B16A16_SINT : PF_R16G16B16A16_SNORM;
 
-	if (OutputMask & 2)
+	if ((OutputMask & static_cast<uint64>(ESkinnedMeshWriteDataInterfaceOutputSelectorMask::TangentX)) ||
+		(OutputMask & static_cast<uint64>(ESkinnedMeshWriteDataInterfaceOutputSelectorMask::TangentZ)))
 	{
 		TangentBuffer = FSkeletalMeshDeformerHelpers::AllocateVertexFactoryTangentBuffer(GraphBuilder, SkeletalMeshObject, LodIndex, TEXT("OptimusSkinnedMeshTangent"));
 		TangentBufferUAV = GraphBuilder.CreateUAV(TangentBuffer, TangentsFormat, ERDGUnorderedAccessViewFlags::SkipBarrier);
@@ -206,7 +214,7 @@ void FOptimusSkinnedMeshWriteDataProviderProxy::AllocateResources(FRDGBuilder& G
 		TangentBufferUAV = GraphBuilder.CreateUAV(GraphBuilder.RegisterExternalBuffer(GWhiteVertexBufferWithRDG->Buffer), TangentsFormat);
 	}
 
-	if (OutputMask & 8)
+	if (OutputMask & static_cast<uint64>(ESkinnedMeshWriteDataInterfaceOutputSelectorMask::Color))
 	{
 		ColorBuffer = FSkeletalMeshDeformerHelpers::AllocateVertexFactoryColorBuffer(GraphBuilder, SkeletalMeshObject, LodIndex, TEXT("OptimusSkinnedMeshColor"));
 		// using RGBA here and do a manual fetch swizzle in shader instead of BGRA directly because some Mac does not support it. See GMetalBufferFormats[PF_B8G8R8A8] 
@@ -214,7 +222,7 @@ void FOptimusSkinnedMeshWriteDataProviderProxy::AllocateResources(FRDGBuilder& G
 	}
 	else
 	{
-		ColorBufferUAV = GraphBuilder.CreateUAV(GraphBuilder.RegisterExternalBuffer(GWhiteVertexBufferWithRDG->Buffer), PF_R8G8B8A8);
+		ColorBufferUAV = GraphBuilder.CreateUAV(GraphBuilder.RegisterExternalBuffer(GWhiteVertexBufferWithRDG->Buffer), PF_A32B32G32R32F);
 	}
 
 	FSkeletalMeshDeformerHelpers::UpdateVertexFactoryBufferOverrides(GraphBuilder.RHICmdList, SkeletalMeshObject, LodIndex);
diff --git a/Engine/Plugins/Animation/DeformerGraph/Source/OptimusCore/Private/Nodes/OptimusNode_CustomComputeKernel.cpp b/Engine/Plugins/Animation/DeformerGraph/Source/OptimusCore/Private/Nodes/OptimusNode_CustomComputeKernel.cpp
--- a/Engine/Plugins/Animation/DeformerGraph/Source/OptimusCore/Private/Nodes/OptimusNode_CustomComputeKernel.cpp
+++ b/Engine/Plugins/Animation/DeformerGraph/Source/OptimusCore/Private/Nodes/OptimusNode_CustomComputeKernel.cpp
@@ -1772,7 +1772,7 @@ void UOptimusNode_CustomComputeKernel::PostLoadExtractExecutionDomain()
 {
 	ExecutionDomain = {};
 	// Check if there's an execution node connected and grab the domain from it.
-	FOptimusDataTypeHandle IntVector3Type = FOptimusDataTypeRegistry::Get().FindType(Optimus::GetTypeName(TBaseStructure<FIntVector3>::Get()));
+	FOptimusDataTypeHandle IntVector3Type = FOptimusDataTypeRegistry::Get().FindType(TBaseStructure<FIntVector3>::Get());
 
 	for (UOptimusNodePin* Pin: GetPins())
 	{
diff --git a/Engine/Plugins/Animation/DeformerGraph/Source/OptimusCore/Private/OptimusDataTypeRegistry.cpp b/Engine/Plugins/Animation/DeformerGraph/Source/OptimusCore/Private/OptimusDataTypeRegistry.cpp
--- a/Engine/Plugins/Animation/DeformerGraph/Source/OptimusCore/Private/OptimusDataTypeRegistry.cpp
+++ b/Engine/Plugins/Animation/DeformerGraph/Source/OptimusCore/Private/OptimusDataTypeRegistry.cpp
@@ -100,7 +100,7 @@ EOptimusDataTypeUsageFlags FOptimusDataTypeRegistry::GetStructTypeUsageFlag(UScr
         		return EOptimusDataTypeUsageFlags::None;
         	}
         	
-        	FOptimusDataTypeHandle DataType = Get().FindType(Optimus::GetTypeName(InStruct));
+        	FOptimusDataTypeHandle DataType = Get().FindType(InStruct);
         	if (DataType.IsValid())
         	{
         		PropertyValueConvertFuncT ConvertFunc = Get().FindPropertyValueConvertFunc(DataType->TypeName);
@@ -223,7 +223,7 @@ EOptimusDataTypeUsageFlags FOptimusDataTypeRegistry::GetStructTypeUsageFlag(UScr
 			EOptimusDataTypeUsageFlags::Resource | EOptimusDataTypeUsageFlags::Variable | EOptimusDataTypeUsageFlags::AnimAttributes | EOptimusDataTypeUsageFlags::DataInterfaceOutput | EOptimusDataTypeUsageFlags::PinType;
     };
 
-	if (FOptimusDataTypeHandle DataType = Get().FindType(Optimus::GetTypeName(InStruct)))
+	if (FOptimusDataTypeHandle DataType = Get().FindType(InStruct))
 	{
 		return DataType->UsageFlags;
 	}
@@ -496,6 +496,32 @@ void FOptimusDataTypeRegistry::RegisterBuiltinTypes()
 
 FOptimusDataTypeRegistry::~FOptimusDataTypeRegistry()
 {
+}
+
+
+FName FOptimusDataTypeRegistry::GetTypeName(const FFieldClass& InFieldClass)
+{
+	return InFieldClass.GetFName();
+}
+
+FName FOptimusDataTypeRegistry::GetTypeName(UScriptStruct* InStruct)
+{
+	return Optimus::GetTypeName(InStruct, true);
+}
+
+FName FOptimusDataTypeRegistry::GetArrayTypeName(const FFieldClass& InFieldClass)
+{
+	return GetArrayTypeName(GetTypeName(InFieldClass));
+}
+
+FName FOptimusDataTypeRegistry::GetArrayTypeName(UScriptStruct* InStruct)
+{
+	return GetArrayTypeName(GetTypeName(InStruct));
+}
+
+FName FOptimusDataTypeRegistry::GetArrayTypeName(FName InElementTypeName)
+{
+	return *(TEXT("TArray<") + InElementTypeName.ToString() + TEXT(">"));	
 };
 
 FOptimusDataTypeRegistry& FOptimusDataTypeRegistry::Get()
@@ -556,7 +582,7 @@ bool FOptimusDataTypeRegistry::RegisterStructType(UScriptStruct* InStructType)
 
 	FText DisplayName = Optimus::GetTypeDisplayName(InStructType);
 
-	FName TypeName = Optimus::GetTypeName(InStructType);
+	FName TypeName = GetTypeName(InStructType);
 
 	const bool bIsHashable = IsStructHashable(InStructType);
 
@@ -943,7 +969,7 @@ bool FOptimusDataTypeRegistry::RegisterArrayTypeIfApplicable(FOptimusDataTypeHan
 
 	{
 		*ArrayDataType = *InElementDataType;
-		ArrayDataType->TypeName = *(TEXT("TArray<") + InElementDataType->TypeName.ToString() + TEXT(">"));
+		ArrayDataType->TypeName = GetArrayTypeName(InElementDataType->TypeName);
 		ArrayDataType->DisplayName = FText::FromString(InElementDataType->DisplayName.ToString() + TEXT(" Array"));
 		ArrayDataType->ShaderValueType = FShaderValueType::MakeDynamicArrayType(InnerDataTypeForStructuredBuffer->ShaderValueType);
 		
@@ -961,7 +987,7 @@ bool FOptimusDataTypeRegistry::RegisterArrayTypeIfApplicable(FOptimusDataTypeHan
 
 void FOptimusDataTypeRegistry::RefreshStructType(UUserDefinedStruct* InStructType)
 {
-	FName TypeName = Optimus::GetTypeName(InStructType);
+	FName TypeName = GetTypeName(InStructType);
 	if (RegisteredTypes.Contains(TypeName))
 	{
 		UnregisterType(TypeName);
@@ -981,8 +1007,8 @@ bool FOptimusDataTypeRegistry::RegisterType(
 	EOptimusDataTypeUsageFlags InUsageFlags
 	)
 {
-	return RegisterType(InFieldType.GetFName(), [&](FOptimusDataType& InDataType) {
-		InDataType.TypeName = InFieldType.GetFName();
+	return RegisterType(GetTypeName(InFieldType), [&](FOptimusDataType& InDataType) {
+		InDataType.TypeName = GetTypeName(InFieldType);
 		InDataType.DisplayName = InDisplayName;
 		InDataType.ShaderValueType = InShaderValueType;
 		InDataType.ShaderValueSize = InShaderValueType->GetResourceElementSize();
@@ -1041,7 +1067,7 @@ bool FOptimusDataTypeRegistry::RegisterType(
 			}
 		}
 
-		const FName TypeName = Optimus::GetTypeName(InStructType);
+		const FName TypeName = GetTypeName(InStructType);
 
 		PropertyCreateFuncT PropertyCreateFunc;
 		PropertyValueConvertFuncT PropertyValueConvertFunc;
@@ -1202,7 +1228,7 @@ bool FOptimusDataTypeRegistry::RegisterType(
 			}
 		}
 
-		const FName TypeName = Optimus::GetTypeName(InStructType, true);
+		const FName TypeName = GetTypeName(InStructType);
 
 		PropertyCreateFuncT PropertyCreateFunc;
 		const int32 ExpectedShaderValueSize = InShaderValueType->GetResourceElementSize();
@@ -1260,7 +1286,7 @@ bool FOptimusDataTypeRegistry::RegisterType(
 {
 	if (ensure(InStructType))
 	{
-		const FName TypeName = Optimus::GetTypeName(InStructType, true);
+		const FName TypeName = GetTypeName(InStructType);
 
 		PropertyCreateFuncT PropertyCreateFunc;
 		if (EnumHasAnyFlags(InUsageFlags, EOptimusDataTypeUsageFlags::Variable))
@@ -1388,7 +1414,7 @@ FOptimusDataTypeHandle FOptimusDataTypeRegistry::FindType(const FProperty& InPro
 {
 	if (const FStructProperty* StructProperty = CastField<const FStructProperty>(&InProperty))
 	{
-		const FName TypeName = Optimus::GetTypeName(StructProperty->Struct);
+		const FName TypeName = GetTypeName(StructProperty->Struct);
 		return FindType(TypeName);
 	}
 	else if (const FObjectProperty* ObjectProperty = CastField<const FObjectProperty>(&InProperty))
@@ -1419,12 +1445,32 @@ FOptimusDataTypeHandle FOptimusDataTypeRegistry::FindType(const FProperty& InPro
 #endif
 }
 
+FOptimusDataTypeHandle FOptimusDataTypeRegistry::FindArrayType(const FProperty& InProperty) const
+{
+	if(FOptimusDataTypeHandle Type = FindType(InProperty))
+	{
+		return FindType(GetArrayTypeName(Type->TypeName));
+	}
+
+	return {};	
+}
+
 
 FOptimusDataTypeHandle FOptimusDataTypeRegistry::FindType(const FFieldClass& InFieldType) const
 {
 	return FindType(InFieldType.GetFName());
 }
 
+FOptimusDataTypeHandle FOptimusDataTypeRegistry::FindArrayType(const FFieldClass& InFieldType) const
+{
+	if(FOptimusDataTypeHandle Type = FindType(InFieldType))
+	{
+		return FindType(GetArrayTypeName(Type->TypeName));
+	}
+
+	return {};	
+}
+
 
 FOptimusDataTypeHandle FOptimusDataTypeRegistry::FindType(const UClass& InClassType) const
 {
@@ -1432,6 +1478,16 @@ FOptimusDataTypeHandle FOptimusDataTypeRegistry::FindType(const UClass& InClassT
 	return FindType(TypeName);
 }
 
+FOptimusDataTypeHandle FOptimusDataTypeRegistry::FindArrayType(const UClass& InClassType) const
+{
+	if(FOptimusDataTypeHandle Type = FindType(InClassType))
+	{
+		return FindType(GetArrayTypeName(Type->TypeName));
+	}
+
+	return {};	
+}
+
 
 FOptimusDataTypeHandle FOptimusDataTypeRegistry::FindType(FName InTypeName) const
 {
@@ -1439,6 +1495,21 @@ FOptimusDataTypeHandle FOptimusDataTypeRegistry::FindType(FName InTypeName) cons
 	return InfoPtr ? InfoPtr->DataType : FOptimusDataTypeHandle();
 }
 
+FOptimusDataTypeHandle FOptimusDataTypeRegistry::FindArrayType(FName InTypeName) const
+{
+	return FindType(GetArrayTypeName(InTypeName));
+}
+
+FOptimusDataTypeHandle FOptimusDataTypeRegistry::FindType(UScriptStruct* InStruct) const
+{
+	return FindType(GetTypeName(InStruct));
+}
+
+FOptimusDataTypeHandle FOptimusDataTypeRegistry::FindArrayType(UScriptStruct* InStruct) const
+{
+	return FindArrayType(GetTypeName(InStruct));
+}
+
 
 FOptimusDataTypeHandle FOptimusDataTypeRegistry::FindType(FShaderValueTypeHandle InValueType) const
 {
@@ -1518,7 +1589,7 @@ void FOptimusDataTypeRegistry::OnFilesLoaded()
 		
 		if (ScriptStruct && !BuiltInAttributeTypeArray.Contains(ScriptStruct))
 		{
-			FOptimusDataTypeHandle DataType = FindType(Optimus::GetTypeName(ScriptStruct));
+			FOptimusDataTypeHandle DataType = FindType(ScriptStruct);
 
 			if (!DataType.IsValid())
 			{
@@ -1534,7 +1605,7 @@ void FOptimusDataTypeRegistry::OnAssetRemoved(const FAssetData& InAssetData)
 	{
 		if (UUserDefinedStruct* UserDefinedStruct = Cast<UUserDefinedStruct>(InAssetData.GetAsset()))
 		{
-			FName TypeName = Optimus::GetTypeName(UserDefinedStruct);
+			FName TypeName = GetTypeName(UserDefinedStruct);
 			UnregisterType(TypeName);
 		}
 	}
@@ -1544,7 +1615,7 @@ void FOptimusDataTypeRegistry::OnAssetRenamed(const FAssetData& InAssetData, con
 {
 	if (UUserDefinedStruct* UserDefinedStruct = Cast<UUserDefinedStruct>(InAssetData.GetAsset()))
 	{
-		FName TypeName = Optimus::GetTypeName(UserDefinedStruct);
+		FName TypeName = GetTypeName(UserDefinedStruct);
 
 		if (FOptimusDataTypeHandle DataTypeHandle = FindType(TypeName))
 		{
@@ -1563,7 +1634,7 @@ void FOptimusDataTypeRegistry::OnAnimationAttributeRegistryChanged(const UScript
 	{
 		
 		UScriptStruct* ScriptStruct = const_cast<UScriptStruct*>(InScriptStruct);
-		FOptimusDataTypeHandle DataType = FindType(Optimus::GetTypeName(ScriptStruct));
+		FOptimusDataTypeHandle DataType = FindType(ScriptStruct);
 		
 		if (bIsAdded)
 		{
diff --git a/Engine/Plugins/Animation/DeformerGraph/Source/OptimusCore/Private/OptimusDeformer.cpp b/Engine/Plugins/Animation/DeformerGraph/Source/OptimusCore/Private/OptimusDeformer.cpp
--- a/Engine/Plugins/Animation/DeformerGraph/Source/OptimusCore/Private/OptimusDeformer.cpp
+++ b/Engine/Plugins/Animation/DeformerGraph/Source/OptimusCore/Private/OptimusDeformer.cpp
@@ -57,6 +57,7 @@
 
 #include <limits>
 
+#include "OptimusDeformerDynamicInstanceManager.h"
 #include "OptimusFunctionNodeGraphHeader.h"
 #include "Nodes/OptimusNode_FunctionReference.h"
 #include "Nodes/OptimusNode_SubGraphReference.h"
@@ -3357,7 +3358,7 @@ void UOptimusDeformer::PostLoadRemoveDeprecatedExecutionNodes()
 				const UOptimusNodePin* PrimaryGroupPin = KernelNode->GetPrimaryGroupPin();
 
 				UOptimusNode_ComponentSource* ComponentSourceNode = nullptr;
-				FOptimusDataTypeHandle IntVector3Type = FOptimusDataTypeRegistry::Get().FindType(Optimus::GetTypeName(TBaseStructure<FIntVector3>::Get()));
+				FOptimusDataTypeHandle IntVector3Type = FOptimusDataTypeRegistry::Get().FindType(TBaseStructure<FIntVector3>::Get());
 				
 				for (const UOptimusNodePin* Pin : PrimaryGroupPin->GetSubPins())
 				{
@@ -3614,7 +3615,17 @@ UMeshDeformerInstance* UOptimusDeformer::CreateInstance(
 		return nullptr;
 	}
 
-	const FName InstanceName(GetName() + TEXT("_Instance"));
+	UOptimusDeformerDynamicInstanceManager* InstanceManager = NewObject<UOptimusDeformerDynamicInstanceManager>(InMeshComponent);
+	
+	InstanceManager->DefaultInstance = CreateOptimusInstance(InMeshComponent, InSettings);
+	
+	return InstanceManager;
+}
+
+UOptimusDeformerInstance* UOptimusDeformer::CreateOptimusInstance(UMeshComponent* InMeshComponent, UMeshDeformerInstanceSettings* InSettings)
+{
+	const FName InstanceName = Optimus::GetUniqueNameForScope(InMeshComponent ,*(GetName() + TEXT("_Instance")));
+	
 	UOptimusDeformerInstance* Instance = NewObject<UOptimusDeformerInstance>(InMeshComponent, InstanceName);
 	Instance->SetMeshComponent(InMeshComponent);
 	Instance->SetInstanceSettings(Cast<UOptimusDeformerInstanceSettings>(InSettings));
diff --git a/Engine/Plugins/Animation/DeformerGraph/Source/OptimusCore/Private/OptimusDeformerDynamicInstanceManager.cpp b/Engine/Plugins/Animation/DeformerGraph/Source/OptimusCore/Private/OptimusDeformerDynamicInstanceManager.cpp
new file mode 100644
--- /dev/null
+++ b/Engine/Plugins/Animation/DeformerGraph/Source/OptimusCore/Private/OptimusDeformerDynamicInstanceManager.cpp
@@ -0,0 +1,211 @@
+﻿// Copyright Epic Games, Inc. All Rights Reserved.
+
+#include "OptimusDeformerDynamicInstanceManager.h"
+
+#include "ControlRig.h"
+#include "IControlRigObjectBinding.h"
+#include "OptimusDeformerInstance.h"
+#include "ControlRig/RigUnit_Optimus.h"
+
+void UOptimusDeformerDynamicInstanceManager::AllocateResources()
+{
+	// Typically called during recreate render state
+	
+	DefaultInstance->AllocateResources();
+
+	for (TPair<FGuid, TObjectPtr<UOptimusDeformerInstance>>& GuidToInstance : GuidToRigDeformerInstanceMap)
+	{
+		GuidToInstance.Value->AllocateResources();
+	}
+}
+
+void UOptimusDeformerDynamicInstanceManager::ReleaseResources()
+{
+	// Typically called during recreate render state
+	
+	DefaultInstance->ReleaseResources();
+	
+	for (TPair<FGuid, TObjectPtr<UOptimusDeformerInstance>>& GuidToInstance : GuidToRigDeformerInstanceMap)
+	{
+		GuidToInstance.Value->ReleaseResources();
+	}
+}
+
+void UOptimusDeformerDynamicInstanceManager::EnqueueWork(FEnqueueWorkDesc const& InDesc)
+{
+	// Runs during UWorld::SendAllEndOfFrameUpdates
+
+	for (FGuid Guid : RigDeformerInstancePendingInit)
+	{
+		if (TObjectPtr<UOptimusDeformerInstance>* InstancePtr = GuidToRigDeformerInstanceMap.Find(Guid))
+		{
+			(*InstancePtr)->AllocateResources();
+		}
+	}
+
+	RigDeformerInstancePendingInit.Reset();
+	
+	// Enqueue work
+
+	// Making sure instances in the queue are dispatched sequentially
+	uint8 NumComputeGraphsPossiblyEnqueued = 0;
+	// Used to inform later instances whether specific buffers have valid data in them
+	EMeshDeformerOutputBuffer OutputBuffers = EMeshDeformerOutputBuffer::None;
+	
+	auto EnqueueInstance = [&NumComputeGraphsPossiblyEnqueued, &InDesc, &OutputBuffers](UOptimusDeformerInstance* Instance)
+	{
+		Instance->OutputBuffersFromPreviousInstances = OutputBuffers;
+		OutputBuffers |= Instance->GetOutputBuffers();
+		
+		Instance->GraphSortPriorityOffset = NumComputeGraphsPossiblyEnqueued;
+		NumComputeGraphsPossiblyEnqueued += Instance->ComputeGraphExecInfos.Num();
+		
+		Instance->EnqueueWork(InDesc);
+	};
+
+	static const EOptimusDeformerExecutionPhase Phases[] = {
+		EOptimusDeformerExecutionPhase::BeforeDefaultDeformer,
+		EOptimusDeformerExecutionPhase::OverrideDefaultDeformer,
+		EOptimusDeformerExecutionPhase::AfterDefaultDeformer,
+	};
+
+	for (EOptimusDeformerExecutionPhase Phase : Phases)
+	{
+		if (TMap<int32, TArray<FGuid>>* ExecutionGroupQueueMapPtr = InstanceQueueMap.Find(Phase))
+		{
+			TArray<int32> SortedExecutionGroups;
+			ExecutionGroupQueueMapPtr->GenerateKeyArray(SortedExecutionGroups);
+			SortedExecutionGroups.Sort();
+
+			if (Phase == EOptimusDeformerExecutionPhase::OverrideDefaultDeformer)
+			{
+				// The last instance in the override queue is the one getting used
+				int32 LastOverrideExecutionGroup = SortedExecutionGroups.Last();
+				FGuid LastOverrideInstanceGuid = (*ExecutionGroupQueueMapPtr)[LastOverrideExecutionGroup].Last();
+
+				if (TObjectPtr<UOptimusDeformerInstance>* InstancePtr = GuidToRigDeformerInstanceMap.Find(LastOverrideInstanceGuid))
+				{
+					EnqueueInstance((*InstancePtr));
+				}
+			}
+			else
+			{
+				for (int32 ExecutionGroup : SortedExecutionGroups)
+				{
+					for (FGuid Guid : (*ExecutionGroupQueueMapPtr)[ExecutionGroup])
+					{
+						if (TObjectPtr<UOptimusDeformerInstance>* InstancePtr = GuidToRigDeformerInstanceMap.Find(Guid))
+						{
+							EnqueueInstance((*InstancePtr));
+						}
+					}
+				}
+			}
+		}
+		else if (Phase == EOptimusDeformerExecutionPhase::OverrideDefaultDeformer)
+		{
+			// Use the default instance if there is no override
+			EnqueueInstance(DefaultInstance);
+		}
+	}
+
+	InstanceQueueMap.Reset();
+}
+
+EMeshDeformerOutputBuffer UOptimusDeformerDynamicInstanceManager::GetOutputBuffers() const
+{
+	// Since instances can be added dynamically, no way to know in advance if some of these are not written to, so just declare all of them
+	EMeshDeformerOutputBuffer Result = EMeshDeformerOutputBuffer::SkinnedMeshPosition | EMeshDeformerOutputBuffer::SkinnedMeshTangents | EMeshDeformerOutputBuffer::SkinnedMeshVertexColor;
+	
+	return Result;
+}
+
+
+void UOptimusDeformerDynamicInstanceManager::OnControlRigBeginDestroy(UControlRig* InControlRig)
+{
+	if (TArray<FGuid>* Instances = RigToInstanceGuidsMap.Find(InControlRig))
+	{
+		TArray<FGuid> InstancesToRemove = *Instances;
+		for (FGuid Guid : InstancesToRemove)
+		{
+			if (TObjectPtr<UOptimusDeformerInstance>* InstancePtr = GuidToRigDeformerInstanceMap.Find(Guid))
+			{
+				(*InstancePtr)->ReleaseResources();
+			}
+
+			GuidToRigDeformerInstanceMap.Remove(Guid);
+		}
+	}
+	
+	RigToInstanceGuidsMap.Remove(InControlRig);
+	
+	InControlRig->OnBeginDestroy().RemoveAll(this);
+}
+
+void UOptimusDeformerDynamicInstanceManager::BeginDestroy()
+{
+	TArray<TWeakObjectPtr<UControlRig>> Rigs;
+	RigToInstanceGuidsMap.GenerateKeyArray(Rigs);
+
+	// Release resources should have been called, so just unregister callbacks for good measure
+	
+	for (TWeakObjectPtr<UControlRig>& Rig : Rigs)
+	{
+		if (Rig.IsValid())
+		{
+			Rig->OnBeginDestroy().RemoveAll(this);
+		}
+	}
+	
+	Super::BeginDestroy();
+}
+
+
+void UOptimusDeformerDynamicInstanceManager::AddRigDeformer(UControlRig* InControlRig, FGuid InInstanceGuid, UOptimusDeformer* InDeformer)
+{
+	check(IsInGameThread());
+	
+	if (ensure(!GuidToRigDeformerInstanceMap.Contains(InInstanceGuid)))
+	{
+		UOptimusDeformerInstance* DeformerInstance = InDeformer->CreateOptimusInstance(CastChecked<USkeletalMeshComponent>(GetOuter()), nullptr);
+		GuidToRigDeformerInstanceMap.Add(InInstanceGuid, DeformerInstance);
+		RigDeformerInstancePendingInit.Add(InInstanceGuid);
+
+		if (TArray<FGuid>* GuidArray = RigToInstanceGuidsMap.Find(InControlRig))
+		{
+			GuidArray->Add(InInstanceGuid);
+		}
+		else
+		{
+			RigToInstanceGuidsMap.Add(InControlRig, {InInstanceGuid});
+			
+			// First time for this control rig, register some callbacks as well
+			check(!InControlRig->OnBeginDestroy().IsBoundToObject(this));
+			// Assuming owning component of the rig cannot change
+			InControlRig->OnBeginDestroy().AddUObject(this, &UOptimusDeformerDynamicInstanceManager::OnControlRigBeginDestroy);
+
+		}
+	}
+}
+
+UOptimusDeformerInstance* UOptimusDeformerDynamicInstanceManager::GetRigDeformer(FGuid InInstanceGuid)
+{
+	TObjectPtr<UOptimusDeformerInstance>* InstancePtr = GuidToRigDeformerInstanceMap.Find(InInstanceGuid);
+	
+	if (InstancePtr)
+	{
+		return *InstancePtr;
+	}
+
+	return nullptr;
+}
+
+
+void UOptimusDeformerDynamicInstanceManager::EnqueueRigDeformer(FGuid InInstanceGuid, EOptimusDeformerExecutionPhase InExecutionPhase, int32 InExecutionGroup)
+{
+	// Typically called from anim thread, but there shouldn't be concurrent access to this queue. All rigs running on the current mesh should run sequentially.
+	
+	// There are cases where control rig can run multiple times in one tick(like when moving controls), but we only need to enqueue once per instance
+	InstanceQueueMap.FindOrAdd(InExecutionPhase).FindOrAdd(InExecutionGroup).AddUnique(InInstanceGuid);
+}
+
diff --git a/Engine/Plugins/Animation/DeformerGraph/Source/OptimusCore/Private/OptimusDeformerInstance.cpp b/Engine/Plugins/Animation/DeformerGraph/Source/OptimusCore/Private/OptimusDeformerInstance.cpp
--- a/Engine/Plugins/Animation/DeformerGraph/Source/OptimusCore/Private/OptimusDeformerInstance.cpp
+++ b/Engine/Plugins/Animation/DeformerGraph/Source/OptimusCore/Private/OptimusDeformerInstance.cpp
@@ -456,8 +456,8 @@ void UOptimusDeformerInstance::SetupFromDeformer(UOptimusDeformer* InDeformer)
 		VariableDescriptionCopy->Guid = VariableDescription->Guid;
 		VariableDescriptionCopy->VariableName = VariableDescription->VariableName;
 		VariableDescriptionCopy->DataType = VariableDescription->DataType;
-		 // No need to copy the default value, we directly copy the shader value of the default value below
-		VariableDescriptionCopy->CachedShaderValue = VariableDescription->CachedShaderValue;
+		// No need to copy the default value, we directly copy the shader value of the default value below
+		VariableDescriptionCopy->CachedShaderValue = VariableDescription->DefaultValueStruct.GetShaderValue(VariableDescription->DataType);
 		Variables->Descriptions.Add(VariableDescriptionCopy);
 	}
 
@@ -503,7 +503,7 @@ void UOptimusDeformerInstance::ReleaseResources()
 	}
 }
 
-void UOptimusDeformerInstance::EnqueueWork(FEnqueueWorkDesc const& InDesc)
+void UOptimusDeformerInstance::EnqueueWork(UMeshDeformerInstance::FEnqueueWorkDesc const& InDesc)
 {
 	// Convert execution group enum to ComputeTaskExecutionGroup name.
 	FName ExecutionGroupName;
@@ -550,7 +550,7 @@ void UOptimusDeformerInstance::EnqueueWork(FEnqueueWorkDesc const& InDesc)
 			{
 				if (Info.GraphType == EOptimusNodeGraphType::Update || GraphsToRun.Contains(Info.GraphName))
 				{
-					bIsWorkEnqueued |= Info.ComputeGraphInstance.EnqueueWork(Info.ComputeGraph, InDesc.Scene, ExecutionGroupName, InDesc.OwnerName, InDesc.FallbackDelegate, this);
+					bIsWorkEnqueued |= Info.ComputeGraphInstance.EnqueueWork(Info.ComputeGraph, InDesc.Scene, ExecutionGroupName, InDesc.OwnerName, InDesc.FallbackDelegate, this, GraphSortPriorityOffset);
 				}
 			}
 		}
@@ -615,14 +615,50 @@ namespace
 
 bool UOptimusDeformerInstance::SetBoolVariable(FName InVariableName, bool InValue)
 {
-	return SetVariableValue(Variables, InVariableName, FBoolProperty::StaticClass()->GetFName(), InValue);
+	return SetVariableValue(Variables, InVariableName, FOptimusDataTypeRegistry::GetTypeName(*FBoolProperty::StaticClass()), InValue);
 }
 
 bool UOptimusDeformerInstance::SetIntVariable(FName InVariableName, int32 InValue)
 {
-	return SetVariableValue<int32>(Variables, InVariableName, FIntProperty::StaticClass()->GetFName(), InValue);
+	return SetVariableValue<int32>(Variables, InVariableName, FOptimusDataTypeRegistry::GetTypeName(*FIntProperty::StaticClass()), InValue);
 }
 
+bool UOptimusDeformerInstance::SetIntArrayVariable(FName InVariableName, const TArray<int32>& InValue)
+{
+	return SetVariableValue<TArray<int32>>(Variables, InVariableName, FOptimusDataTypeRegistry::GetArrayTypeName(*FIntProperty::StaticClass()), InValue);
+}
+
+bool UOptimusDeformerInstance::SetInt2Variable(FName InVariableName, const FIntPoint& InValue)
+{
+	return SetVariableValue<FIntPoint>(Variables, InVariableName, FOptimusDataTypeRegistry::GetTypeName(TBaseStructure<FIntPoint>::Get()), InValue);
+}
+
+bool UOptimusDeformerInstance::SetInt2ArrayVariable(FName InVariableName, const TArray<FIntPoint>& InValue)
+{
+	return SetVariableValue<TArray<FIntPoint>>(Variables, InVariableName, FOptimusDataTypeRegistry::GetArrayTypeName(TBaseStructure<FIntPoint>::Get()), InValue);
+}
+
+bool UOptimusDeformerInstance::SetInt3Variable(FName InVariableName, const FIntVector& InValue)
+{
+	return SetVariableValue<FIntVector>(Variables, InVariableName, FOptimusDataTypeRegistry::GetTypeName(TBaseStructure<FIntVector>::Get()), InValue);
+}
+
+bool UOptimusDeformerInstance::SetInt3ArrayVariable(FName InVariableName, const TArray<FIntVector>& InValue)
+{
+	return SetVariableValue<TArray<FIntVector>>(Variables, InVariableName, FOptimusDataTypeRegistry::GetArrayTypeName(TBaseStructure<FIntVector>::Get()), InValue);
+}
+
+bool UOptimusDeformerInstance::SetInt4Variable(FName InVariableName, const FIntVector4& InValue)
+{
+	return SetVariableValue<FIntVector4>(Variables, InVariableName, FOptimusDataTypeRegistry::GetTypeName(TBaseStructure<FIntVector4>::Get()), InValue);
+}
+
+bool UOptimusDeformerInstance::SetInt4ArrayVariable(FName InVariableName, const TArray<FIntVector4>& InValue)
+{
+	return SetVariableValue<TArray<FIntVector4>>(Variables, InVariableName, FOptimusDataTypeRegistry::GetArrayTypeName(TBaseStructure<FIntVector4>::Get()), InValue);
+}
+
+
 bool UOptimusDeformerInstance::SetFloatVariable(FName InVariableName, double InValue)
 {
 	if (SetVariableValue<double>(Variables, InVariableName, FDoubleProperty::StaticClass()->GetFName(), InValue))
@@ -634,19 +670,79 @@ bool UOptimusDeformerInstance::SetFloatVariable(FName InVariableName, double InV
 	return SetVariableValue<float>(Variables, InVariableName, FFloatProperty::StaticClass()->GetFName(), static_cast<float>(InValue));
 }
 
+bool UOptimusDeformerInstance::SetFloatArrayVariable(FName InVariableName, const TArray<double>& InValue)
+{
+	return SetVariableValue<TArray<double>>(Variables, InVariableName, FOptimusDataTypeRegistry::GetArrayTypeName(FDoubleProperty::StaticClass()->GetFName()), InValue);
+}
+
+bool UOptimusDeformerInstance::SetVector2Variable(FName InVariableName, const FVector2D& InValue)
+{
+	return SetVariableValue<FVector2D>(Variables, InVariableName, FOptimusDataTypeRegistry::GetTypeName(TBaseStructure<FVector2D>::Get()), InValue);
+}
+
+bool UOptimusDeformerInstance::SetVector2ArrayVariable(FName InVariableName, const TArray<FVector2D>& InValue)
+{
+	return SetVariableValue<TArray<FVector2D>>(Variables, InVariableName, FOptimusDataTypeRegistry::GetArrayTypeName(TBaseStructure<FVector2D>::Get()), InValue);
+}
+
 bool UOptimusDeformerInstance::SetVectorVariable(FName InVariableName, const FVector& InValue)
 {
-	return SetVariableValue<FVector>(Variables, InVariableName, "FVector", InValue);
+	return SetVariableValue<FVector>(Variables, InVariableName, FOptimusDataTypeRegistry::GetTypeName(TBaseStructure<FVector>::Get()), InValue);
+}
+
+bool UOptimusDeformerInstance::SetVectorArrayVariable(FName InVariableName, const TArray<FVector>& InValue)
+{
+	return SetVariableValue<TArray<FVector>>(Variables, InVariableName, FOptimusDataTypeRegistry::GetArrayTypeName(TBaseStructure<FVector>::Get()), InValue);
 }
 
 bool UOptimusDeformerInstance::SetVector4Variable(FName InVariableName, const FVector4& InValue)
 {
-	return SetVariableValue<FVector4>(Variables, InVariableName, "FVector4", InValue);
+	return SetVariableValue<FVector4>(Variables, InVariableName, FOptimusDataTypeRegistry::GetTypeName(TBaseStructure<FVector4>::Get()), InValue);
+}
+
+bool UOptimusDeformerInstance::SetVector4ArrayVariable(FName InVariableName, const TArray<FVector4>& InValue)
+{
+	return SetVariableValue<TArray<FVector4>>(Variables, InVariableName, FOptimusDataTypeRegistry::GetArrayTypeName(TBaseStructure<FVector4>::Get()), InValue);
+}
+
+bool UOptimusDeformerInstance::SetLinearColorVariable(FName InVariableName, const FLinearColor& InValue)
+{
+	return SetVariableValue<FLinearColor>(Variables, InVariableName, FOptimusDataTypeRegistry::GetTypeName(TBaseStructure<FLinearColor>::Get()), InValue);
+}
+
+bool UOptimusDeformerInstance::SetLinearColorArrayVariable(FName InVariableName, const TArray<FLinearColor>& InValue)
+{
+	return SetVariableValue<TArray<FLinearColor>>(Variables, InVariableName, FOptimusDataTypeRegistry::GetArrayTypeName(TBaseStructure<FLinearColor>::Get()), InValue);
+}
+
+bool UOptimusDeformerInstance::SetQuatVariable(FName InVariableName, const FQuat& InValue)
+{
+	return SetVariableValue<FQuat>(Variables, InVariableName, FOptimusDataTypeRegistry::GetTypeName(TBaseStructure<FQuat>::Get()), InValue);
+}
+
+bool UOptimusDeformerInstance::SetQuatArrayVariable(FName InVariableName, const TArray<FQuat>& InValue)
+{
+	return SetVariableValue<TArray<FQuat>>(Variables, InVariableName, FOptimusDataTypeRegistry::GetArrayTypeName(TBaseStructure<FQuat>::Get()), InValue);
+}
+
+bool UOptimusDeformerInstance::SetRotatorVariable(FName InVariableName, const FRotator& InValue)
+{
+	return SetVariableValue<FRotator>(Variables, InVariableName, FOptimusDataTypeRegistry::GetTypeName(TBaseStructure<FRotator>::Get()), InValue);
+}
+
+bool UOptimusDeformerInstance::SetRotatorArrayVariable(FName InVariableName, const TArray<FRotator>& InValue)
+{
+	return SetVariableValue<TArray<FRotator>>(Variables, InVariableName, FOptimusDataTypeRegistry::GetArrayTypeName(TBaseStructure<FRotator>::Get()), InValue);
 }
 
 bool UOptimusDeformerInstance::SetTransformVariable(FName InVariableName, const FTransform& InValue)
 {
-	return SetVariableValue<FTransform>(Variables, InVariableName, "FTransform", InValue);
+	return SetVariableValue<FTransform>(Variables, InVariableName, FOptimusDataTypeRegistry::GetTypeName(TBaseStructure<FTransform>::Get()), InValue);
+}
+
+bool UOptimusDeformerInstance::SetTransformArrayVariable(FName InVariableName, const TArray<FTransform>& InValue)
+{
+	return SetVariableValue<TArray<FTransform>>(Variables, InVariableName, FOptimusDataTypeRegistry::GetArrayTypeName(TBaseStructure<FTransform>::Get()), InValue);
 }
 
 const TArray<UOptimusVariableDescription*>& UOptimusDeformerInstance::GetVariables() const
diff --git a/Engine/Plugins/Animation/DeformerGraph/Source/OptimusCore/Private/OptimusVariableDescription.cpp b/Engine/Plugins/Animation/DeformerGraph/Source/OptimusCore/Private/OptimusVariableDescription.cpp
--- a/Engine/Plugins/Animation/DeformerGraph/Source/OptimusCore/Private/OptimusVariableDescription.cpp
+++ b/Engine/Plugins/Animation/DeformerGraph/Source/OptimusCore/Private/OptimusVariableDescription.cpp
@@ -20,7 +20,6 @@ void UOptimusVariableDescription::SetDataType(FOptimusDataTypeRef InDataType)
 	}
 	else
 	{
-		DataType = InDataType;
 		if (!DefaultValueStruct.IsInitialized())
 		{
 			DefaultValueStruct.SetType(InDataType);
@@ -57,11 +56,6 @@ void UOptimusVariableDescription::PostLoad()
 	}
 	
 PRAGMA_DISABLE_DEPRECATION_WARNINGS
-	if (!ValueData_DEPRECATED.IsEmpty())
-	{
-		CachedShaderValue.ShaderValue = ValueData_DEPRECATED;
-		ValueData_DEPRECATED.Reset();
-	}
 
 	if (GetLinkerCustomVersion(FOptimusObjectVersion::GUID) < FOptimusObjectVersion::PropertyBagValueContainer)
 	{
@@ -78,13 +72,6 @@ PRAGMA_ENABLE_DEPRECATION_WARNINGS
 	{
 		SetDataType(DataType);
 	}
-
-	// The cached shader value might have incorrect size due to LWC changes,
-	// so we refresh the cache here. Example: FVector shader value was saved as 3*8 bytes at some point
-	if (DataType->ShaderValueSize != CachedShaderValue.ShaderValue.Num())
-	{
-		CachedShaderValue = DefaultValueStruct.GetShaderValue(DataType);
-	}
 }
 
 
@@ -119,11 +106,6 @@ void UOptimusVariableDescription::PostEditChangeProperty(FPropertyChangedEvent&
 			Deformer->SetVariableDataType(this, DataType, bForceChange);
 		}
 	}
-	else if (PropertyName == GET_MEMBER_NAME_CHECKED(UOptimusVariableDescription, DefaultValueStruct))
-	{
-		// Store the default shader value.
-		CachedShaderValue = DefaultValueStruct.GetShaderValue(DataType);
-	}
 }
 
 
diff --git a/Engine/Plugins/Animation/DeformerGraph/Source/OptimusCore/Public/IOptimusDeformerInstanceAccessor.h b/Engine/Plugins/Animation/DeformerGraph/Source/OptimusCore/Public/IOptimusDeformerInstanceAccessor.h
--- a/Engine/Plugins/Animation/DeformerGraph/Source/OptimusCore/Public/IOptimusDeformerInstanceAccessor.h
+++ b/Engine/Plugins/Animation/DeformerGraph/Source/OptimusCore/Public/IOptimusDeformerInstanceAccessor.h
@@ -23,5 +23,4 @@ class IOptimusDeformerInstanceAccessor
 
 public:
 	virtual void SetDeformerInstance(UOptimusDeformerInstance* InInstance) = 0;
-	virtual UOptimusDeformerInstance* GetDeformerInstance() const = 0;
 };
diff --git a/Engine/Plugins/Animation/DeformerGraph/Source/OptimusCore/Public/OptimusDataTypeRegistry.h b/Engine/Plugins/Animation/DeformerGraph/Source/OptimusCore/Public/OptimusDataTypeRegistry.h
--- a/Engine/Plugins/Animation/DeformerGraph/Source/OptimusCore/Public/OptimusDataTypeRegistry.h
+++ b/Engine/Plugins/Animation/DeformerGraph/Source/OptimusCore/Public/OptimusDataTypeRegistry.h
@@ -42,6 +42,13 @@ public:
 	
 	virtual ~FOptimusDataTypeRegistry() override;
 
+	static FName GetTypeName(const FFieldClass& InFieldClass);
+	static FName GetTypeName(UScriptStruct* InStruct);
+	
+	static FName GetArrayTypeName(const FFieldClass& InFieldClass);
+	static FName GetArrayTypeName(UScriptStruct* InStruct);
+	static FName GetArrayTypeName(FName InElementTypeName);
+	
 	/** Get the singleton registry object */
 	OPTIMUSCORE_API static FOptimusDataTypeRegistry &Get();
 
@@ -149,30 +156,67 @@ public:
 	*/
 	OPTIMUSCORE_API FOptimusDataTypeHandle FindType(const FProperty &InProperty) const;
 
+	/** Find the registered array type associated with the given property's type. Returns an invalid
+	  * handle if no registered type is associated.
+	*/
+	OPTIMUSCORE_API FOptimusDataTypeHandle FindArrayType(const FProperty &InProperty) const;
+
 	/** Find the registered type associated with the given field class. Returns an invalid
 	  * handle if no registered type is associated.
 	*/
 	OPTIMUSCORE_API FOptimusDataTypeHandle FindType(const FFieldClass& InFieldType) const;
 
+	/** Find the registered array type associated with the given field class. Returns an invalid
+	  * handle if no registered array type is associated.
+	*/
+	OPTIMUSCORE_API FOptimusDataTypeHandle FindArrayType(const FFieldClass& InFieldType) const;
+
 	/** Find the registered type associated with the given object class. Returns an invalid
 	  * handle if no registered type is associated.
 	*/
 	OPTIMUSCORE_API FOptimusDataTypeHandle FindType(const UClass& InClassType) const;
 	
+	/** Find the registered array type associated with the given object class. Returns an invalid
+      * handle if no registered array type is associated.
+	*/
+	OPTIMUSCORE_API FOptimusDataTypeHandle FindArrayType(const UClass& InClassType) const;
+	
 	/** Find the registered type with the given name. Returns an invalid handle if no registered 
 	  * type with that name exists.
 	*/
 	OPTIMUSCORE_API FOptimusDataTypeHandle FindType(FName InTypeName) const;
 
+	/** Find the registered array type with the given element type name. Returns an invalid handle if no registered
+	  * type with that name exists.
+	*/
+	OPTIMUSCORE_API FOptimusDataTypeHandle FindArrayType(FName InTypeName) const;
+
+	/** Find the registered array type with the given struct. Returns an invalid handle if no registered 
+	  * type with that name exists.
+	*/
+	OPTIMUSCORE_API FOptimusDataTypeHandle FindType(UScriptStruct* InStruct) const;
+
+	/** Find the registered arraytype with the given element type struct. Returns an invalid handle if no registered 
+	  * type with that name exists.
+	*/
+	OPTIMUSCORE_API FOptimusDataTypeHandle FindArrayType(UScriptStruct* InStruct) const;
+
 	/** Find a registered type from a FShaderValueTypeHandle. If multiple types are using the
 	  * same shader value type, then the first one found in the registration order will be
 	  * returned.
-	  */
+	*/
 	// FIXME: We should allow for some kind of type hinting from the HLSL side (e.g. vector4 a color or a vector of four independent scalars).
 	OPTIMUSCORE_API FOptimusDataTypeHandle FindType(FShaderValueTypeHandle InValueType) const;
 
+	/** Find a registered type from a FShaderValueTypeHandle. If multiple types are using the
+	  * same shader value type, then the first one found in the registration order will be
+	  * returned.
+	*/
+	// FIXME: We should allow for some kind of type hinting from the HLSL side (e.g. vector4 a color or a vector of four independent scalars).
+	OPTIMUSCORE_API FOptimusDataTypeHandle FindArrayType(FShaderValueTypeHandle InValueType) const;
+
 	/** A helper function to return a property conversion function. The function can be unbound
-	*  and that should be checked prior to calling
+	  * and that should be checked prior to calling
 	*/
 	PropertyValueConvertFuncT FindPropertyValueConvertFunc(FName InTypeName) const;
 
diff --git a/Engine/Plugins/Animation/DeformerGraph/Source/OptimusCore/Public/OptimusDeformer.h b/Engine/Plugins/Animation/DeformerGraph/Source/OptimusCore/Public/OptimusDeformer.h
--- a/Engine/Plugins/Animation/DeformerGraph/Source/OptimusCore/Public/OptimusDeformer.h
+++ b/Engine/Plugins/Animation/DeformerGraph/Source/OptimusCore/Public/OptimusDeformer.h
@@ -26,6 +26,7 @@ class USkeletalMesh;
 class UOptimusActionStack;
 class UOptimusComputeGraph;
 class UOptimusDeformer;
+class UOptimusDeformerInstance;
 class UOptimusResourceDescription;
 class UOptimusVariableDescription;
 class UOptimusFunctionNodeGraph;
@@ -322,6 +323,12 @@ public:
 		UMeshDeformerInstanceSettings* InSettings
 		) override;
 	
+	
+	UOptimusDeformerInstance* CreateOptimusInstance(
+		UMeshComponent* InMeshComponent,
+		UMeshDeformerInstanceSettings* InSettings
+		);
+	
 	// IInterface_PreviewMeshProvider overrides
 	void SetPreviewMesh(USkeletalMesh* PreviewMesh, bool bMarkAsDirty = true) override;
 	USkeletalMesh* GetPreviewMesh() const override;
diff --git a/Engine/Plugins/Animation/DeformerGraph/Source/OptimusCore/Public/OptimusDeformerDynamicInstanceManager.h b/Engine/Plugins/Animation/DeformerGraph/Source/OptimusCore/Public/OptimusDeformerDynamicInstanceManager.h
new file mode 100644
--- /dev/null
+++ b/Engine/Plugins/Animation/DeformerGraph/Source/OptimusCore/Public/OptimusDeformerDynamicInstanceManager.h
@@ -0,0 +1,69 @@
+﻿// Copyright Epic Games, Inc. All Rights Reserved.
+
+#pragma once
+
+#include "Animation/MeshDeformerInstance.h"
+#include "OptimusDeformerDynamicInstanceManager.generated.h"
+
+
+class UControlRig;
+class UOptimusDeformerInstance;
+class UOptimusDeformer;
+
+UENUM()
+enum class EOptimusDeformerExecutionPhase : uint8
+{
+	AfterDefaultDeformer = 0,
+	OverrideDefaultDeformer = 1,
+	BeforeDefaultDeformer = 2,
+};
+
+/** 
+ * Enables composition of multiple deformer instances dynamically
+ */
+UCLASS()
+class OPTIMUSCORE_API UOptimusDeformerDynamicInstanceManager : public UMeshDeformerInstance
+{
+	GENERATED_BODY()
+
+public:
+	/** Called to allocate any persistent render resources */
+	void AllocateResources() override;;
+
+	/** Called when persistent render resources should be released */
+	void ReleaseResources() override;
+
+	/** Enqueue the mesh deformer workload on a scene. */
+	void EnqueueWork(FEnqueueWorkDesc const& InDesc) override;
+
+	/** Return the buffers that this deformer can potentially write to */
+	EMeshDeformerOutputBuffer GetOutputBuffers() const override;
+
+	/** Remove associated deformer instances when the rig is removed */
+	void OnControlRigBeginDestroy(UControlRig* InControlRig);
+
+	void BeginDestroy() override;
+
+	void AddRigDeformer(UControlRig* InControlRig, FGuid InInstanceGuid, UOptimusDeformer* InDeformer);
+	UOptimusDeformerInstance* GetRigDeformer(FGuid InInstanceGuid);
+	void EnqueueRigDeformer(FGuid InInstanceGuid, EOptimusDeformerExecutionPhase InExecutionPhase, int32 InExecutionGroup);
+	
+	UPROPERTY()
+	TObjectPtr<UOptimusDeformerInstance> DefaultInstance;
+
+	UPROPERTY()
+	TMap<FGuid, TObjectPtr<UOptimusDeformerInstance>> GuidToRigDeformerInstanceMap;
+
+	TMap<TWeakObjectPtr<UControlRig>, TArray<FGuid>> RigToInstanceGuidsMap;
+
+	// Freshly created deformer instances should be initialized before dispatch 
+	TArray<FGuid> RigDeformerInstancePendingInit;
+
+	// Instances per execution group per execution phase
+	TMap<EOptimusDeformerExecutionPhase, TMap<int32, TArray<FGuid>>> InstanceQueueMap;
+};
+
+
+
+
+
diff --git a/Engine/Plugins/Animation/DeformerGraph/Source/OptimusCore/Public/OptimusDeformerInstance.h b/Engine/Plugins/Animation/DeformerGraph/Source/OptimusCore/Public/OptimusDeformerInstance.h
--- a/Engine/Plugins/Animation/DeformerGraph/Source/OptimusCore/Public/OptimusDeformerInstance.h
+++ b/Engine/Plugins/Animation/DeformerGraph/Source/OptimusCore/Public/OptimusDeformerInstance.h
@@ -182,13 +182,11 @@ protected:
 
 
 /** 
- * Class representing an instance of an Optimus Mesh Deformer.
- * This implements the UMeshDeformerInstance interface to enqueue the graph execution.
- * It also contains the per instance deformer variable state and local state for each of the graphs in the deformer.
+ * Class representing an instance of an Optimus Mesh Deformer, used in a OptimusDeformerDynamicInstanceManager
+ * It contains the per instance deformer variable state and local state for each of the graphs in the deformer.
  */
 UCLASS(BlueprintType)
-class UOptimusDeformerInstance :
-	public UMeshDeformerInstance
+class UOptimusDeformerInstance : public UObject
 {
 	GENERATED_BODY()
 
@@ -220,22 +218,84 @@ public:
 	UFUNCTION(BlueprintCallable, Category="Deformer", meta=(DisplayName="Set Variable (Integer)"))
 	bool SetIntVariable(FName InVariableName, int32 InValue);
 
-	/** Set the value of a float variable. */
+
+	UFUNCTION(BlueprintCallable, Category="Deformer", meta=(DisplayName="Set Variable (Integer Array)"))
+	bool SetIntArrayVariable(FName InVariableName, const TArray<int32>& InValue);
+	
+
+	UFUNCTION(BlueprintCallable, Category="Deformer", meta=(DisplayName="Set Variable (Integer2)"))
+	bool SetInt2Variable(FName InVariableName, const FIntPoint& InValue);
+
+
+	UFUNCTION(BlueprintCallable, Category="Deformer", meta=(DisplayName="Set Variable (Integer2 Array)"))
+	bool SetInt2ArrayVariable(FName InVariableName, const TArray<FIntPoint>& InValue);
+
+
+	UFUNCTION(BlueprintCallable, Category="Deformer", meta=(DisplayName="Set Variable (Integer3)"))
+	bool SetInt3Variable(FName InVariableName, const FIntVector& InValue);
+
+
+	UFUNCTION(BlueprintCallable, Category="Deformer", meta=(DisplayName="Set Variable (Integer3 Array)"))
+	bool SetInt3ArrayVariable(FName InVariableName, const TArray<FIntVector>& InValue);
+
+
+	UFUNCTION(BlueprintCallable, Category="Deformer", meta=(DisplayName="Set Variable (Integer4)"))
+	bool SetInt4Variable(FName InVariableName, const FIntVector4& InValue);
+
+
+	UFUNCTION(BlueprintCallable, Category="Deformer", meta=(DisplayName="Set Variable (Integer4 Array)"))
+	bool SetInt4ArrayVariable(FName InVariableName, const TArray<FIntVector4>& InValue);
+
+
 	UFUNCTION(BlueprintCallable, Category="Deformer", meta=(DisplayName="Set Variable (Float)"))
 	bool SetFloatVariable(FName InVariableName, double InValue);
 
-	/** Set the value of a 3-vector variable. */
+	UFUNCTION(BlueprintCallable, Category="Deformer", meta=(DisplayName="Set Variable (Float Array)"))
+	bool SetFloatArrayVariable(FName InVariableName, const TArray<double>& InValue);
+
+	UFUNCTION(BlueprintCallable, Category="Deformer", meta=(DisplayName="Set Variable (Vector2)"))
+	bool SetVector2Variable(FName InVariableName, const FVector2D& InValue);
+
+	UFUNCTION(BlueprintCallable, Category="Deformer", meta=(DisplayName="Set Variable (Vector2 Array)"))
+	bool SetVector2ArrayVariable(FName InVariableName, const TArray<FVector2D>& InValue);
+
 	UFUNCTION(BlueprintCallable, Category="Deformer", meta=(DisplayName="Set Variable (Vector)"))
 	bool SetVectorVariable(FName InVariableName, const FVector& InValue);
 
-	/** Set the value of a 4-vector variable. */
+	UFUNCTION(BlueprintCallable, Category="Deformer", meta=(DisplayName="Set Variable (Vector Array)"))
+	bool SetVectorArrayVariable(FName InVariableName, const TArray<FVector>& InValue);
+
 	UFUNCTION(BlueprintCallable, Category="Deformer", meta=(DisplayName="Set Variable (Vector4)"))
 	bool SetVector4Variable(FName InVariableName, const FVector4& InValue);
 
+	UFUNCTION(BlueprintCallable, Category="Deformer", meta=(DisplayName="Set Variable (Vector4 Array)"))
+	bool SetVector4ArrayVariable(FName InVariableName, const TArray<FVector4>& InValue);
+
+	UFUNCTION(BlueprintCallable, Category="Deformer", meta=(DisplayName="Set Variable (LinearColor)"))
+	bool SetLinearColorVariable(FName InVariableName, const FLinearColor& InValue);
+
+	UFUNCTION(BlueprintCallable, Category="Deformer", meta=(DisplayName="Set Variable (LinearColor Array)"))
+	bool SetLinearColorArrayVariable(FName InVariableName, const TArray<FLinearColor>& InValue);
+
+	UFUNCTION(BlueprintCallable, Category="Deformer", meta=(DisplayName="Set Variable (Quat)"))
+	bool SetQuatVariable(FName InVariableName, const FQuat& InValue);
+
+	UFUNCTION(BlueprintCallable, Category="Deformer", meta=(DisplayName="Set Variable (Quat Array)"))
+	bool SetQuatArrayVariable(FName InVariableName, const TArray<FQuat>& InValue);
+
+	UFUNCTION(BlueprintCallable, Category="Deformer", meta=(DisplayName="Set Variable (Rotator)"))
+	bool SetRotatorVariable(FName InVariableName, const FRotator& InValue);
+
+	UFUNCTION(BlueprintCallable, Category="Deformer", meta=(DisplayName="Set Variable (Rotator Array)"))
+	bool SetRotatorArrayVariable(FName InVariableName, const TArray<FRotator>& InValue);
+	
 	/** Set the value of a transform variable. */
 	UFUNCTION(BlueprintCallable, Category = "Deformer", meta = (DisplayName = "Set Variable (Transform)"))
 	bool SetTransformVariable(FName InVariableName, const FTransform& InValue);
 
+	UFUNCTION(BlueprintCallable, Category = "Deformer", meta = (DisplayName = "Set Variable (Transform Array)"))
+	bool SetTransformArrayVariable(FName InVariableName, const TArray<FTransform>& InValue);
+	
 	/** Get an array containing all the variables. */
 	UFUNCTION(BlueprintGetter)
 	const TArray<UOptimusVariableDescription*>& GetVariables() const;
@@ -252,12 +312,20 @@ public:
 	FOptimusPersistentBufferPoolPtr GetBufferPool() const { return BufferPool; }
 
 	void SetCanBeActive(bool bInCanBeActive);
+	
+	/** Making sure compute graphs belong to this instance does not run before instances before it */
+	int32 GraphSortPriorityOffset = 0;
+	
+	/** Used to see which buffers have valid data produced by dispatched instances and are safe to access for the current instance */	
+	EMeshDeformerOutputBuffer OutputBuffersFromPreviousInstances = EMeshDeformerOutputBuffer::None;
+	
 protected:
-	/** Implementation of UMeshDeformerInstance. */
-	void AllocateResources() override;
-	void ReleaseResources() override;
-	void EnqueueWork(FEnqueueWorkDesc const& InDesc) override;
-	EMeshDeformerOutputBuffer GetOutputBuffers() const override;
+	friend class UOptimusDeformerDynamicInstanceManager;
+
+	void AllocateResources();
+	void ReleaseResources();
+	void EnqueueWork(UMeshDeformerInstance::FEnqueueWorkDesc const& InDesc);
+	EMeshDeformerOutputBuffer GetOutputBuffers() const;
 
 private:
 	/** The Mesh Component that owns this Mesh Deformer Instance. */
diff --git a/Engine/Plugins/Animation/DeformerGraph/Source/OptimusCore/Public/OptimusVariableDescription.h b/Engine/Plugins/Animation/DeformerGraph/Source/OptimusCore/Public/OptimusVariableDescription.h
--- a/Engine/Plugins/Animation/DeformerGraph/Source/OptimusCore/Public/OptimusVariableDescription.h
+++ b/Engine/Plugins/Animation/DeformerGraph/Source/OptimusCore/Public/OptimusVariableDescription.h
@@ -71,8 +71,8 @@ public:
 	UPROPERTY(EditAnywhere, Category = VariableDefinition)
 	FOptimusValueContainerStruct DefaultValueStruct;
 	
-	/** Cached shader value binary data. */
-	UPROPERTY()
+	/**	Runtime container for variable values in a deformer instance*/
+	UPROPERTY(Transient)
 	FShaderValueContainer CachedShaderValue;
 
 	void PostLoad();
diff --git a/Engine/Plugins/Animation/DeformerGraph/Source/OptimusSettings/Public/OptimusSettings.h b/Engine/Plugins/Animation/DeformerGraph/Source/OptimusSettings/Public/OptimusSettings.h
--- a/Engine/Plugins/Animation/DeformerGraph/Source/OptimusSettings/Public/OptimusSettings.h
+++ b/Engine/Plugins/Animation/DeformerGraph/Source/OptimusSettings/Public/OptimusSettings.h
@@ -30,7 +30,7 @@ class OPTIMUSSETTINGS_API UOptimusSettings : public UDeveloperSettings
 
 	/** Set when skinned meshes should have a default deformer applied. */
 	UPROPERTY(config, EditAnywhere, Category = DeformerGraph)
-	EOptimusDefaultDeformerMode DefaultMode = EOptimusDefaultDeformerMode::Never;
+	EOptimusDefaultDeformerMode DefaultMode = EOptimusDefaultDeformerMode::OptIn;
 
 	/** A default deformer that will be used on a skinned mesh if no other deformer has been set. */
 	UPROPERTY(config, EditAnywhere, Category = DeformerGraph, meta = (AllowedClasses = "/Script/OptimusCore.OptimusDeformer", EditCondition = "DefaultMode != EOptimusDefaultDeformerMode::Never"))
diff --git a/Engine/Plugins/Runtime/ComputeFramework/Source/ComputeFramework/Private/ComputeGraphInstance.cpp b/Engine/Plugins/Runtime/ComputeFramework/Source/ComputeFramework/Private/ComputeGraphInstance.cpp
--- a/Engine/Plugins/Runtime/ComputeFramework/Source/ComputeFramework/Private/ComputeGraphInstance.cpp
+++ b/Engine/Plugins/Runtime/ComputeFramework/Source/ComputeFramework/Private/ComputeGraphInstance.cpp
@@ -26,7 +26,7 @@ void FComputeGraphInstance::DestroyDataProviders()
 	DataProviders.Reset();
 }
 
-bool FComputeGraphInstance::EnqueueWork(UComputeGraph* InComputeGraph, FSceneInterface const* InScene, FName InExecutionGroupName, FName InOwnerName, FSimpleDelegate InFallbackDelegate, UObject* InOwnerPointer)
+bool FComputeGraphInstance::EnqueueWork(UComputeGraph* InComputeGraph, FSceneInterface const* InScene, FName InExecutionGroupName, FName InOwnerName, FSimpleDelegate InFallbackDelegate, UObject* InOwnerPointer, uint8 InGraphSortPriorityOffset)
 {
 	if (InComputeGraph == nullptr || InScene == nullptr)
 	{
@@ -71,7 +71,7 @@ bool FComputeGraphInstance::EnqueueWork(UComputeGraph* InComputeGraph, FSceneInt
 	}
 
 	ENQUEUE_RENDER_COMMAND(ComputeFrameworkEnqueueExecutionCommand)(
-		[ComputeGraphWorker, InExecutionGroupName, InOwnerName, SortPriority = GraphSortPriority, GraphRenderProxy, MovedDataProviderRenderProxies = MoveTemp(DataProviderRenderProxies), InFallbackDelegate, InOwnerPointer](FRHICommandListImmediate& RHICmdList)
+		[ComputeGraphWorker, InExecutionGroupName, InOwnerName, SortPriority = InGraphSortPriorityOffset + GraphSortPriority, GraphRenderProxy, MovedDataProviderRenderProxies = MoveTemp(DataProviderRenderProxies), InFallbackDelegate, InOwnerPointer](FRHICommandListImmediate& RHICmdList)
 		{
 			// Compute graph scheduler will take ownership of the provider proxies.
 			ComputeGraphWorker->Enqueue(InExecutionGroupName, InOwnerName, SortPriority, GraphRenderProxy, MovedDataProviderRenderProxies, InFallbackDelegate, InOwnerPointer);
diff --git a/Engine/Plugins/Runtime/ComputeFramework/Source/ComputeFramework/Public/ComputeFramework/ComputeGraphInstance.h b/Engine/Plugins/Runtime/ComputeFramework/Source/ComputeFramework/Public/ComputeFramework/ComputeGraphInstance.h
--- a/Engine/Plugins/Runtime/ComputeFramework/Source/ComputeFramework/Public/ComputeFramework/ComputeGraphInstance.h
+++ b/Engine/Plugins/Runtime/ComputeFramework/Source/ComputeFramework/Public/ComputeFramework/ComputeGraphInstance.h
@@ -38,7 +38,7 @@ public:
 	TArray< TObjectPtr<UComputeDataProvider> >& GetDataProviders() { return DataProviders; }
 
 	/** Enqueue the ComputeGraph work. */
-	bool EnqueueWork(UComputeGraph* InComputeGraph, FSceneInterface const* InScene, FName InExecutionGroupName, FName InOwnerName, FSimpleDelegate InFallbackDelegate, UObject* InOwnerPointer = nullptr);
+	bool EnqueueWork(UComputeGraph* InComputeGraph, FSceneInterface const* InScene, FName InExecutionGroupName, FName InOwnerName, FSimpleDelegate InFallbackDelegate, UObject* InOwnerPointer = nullptr, uint8 InGraphSortPriorityOffset = 0);
 
 private:
 	/** The currently bound Data Provider objects. */
diff --git a/Engine/Plugins/Runtime/RigVM/Source/RigVM/Public/RigVMCore/RigVMExecuteContext.h b/Engine/Plugins/Runtime/RigVM/Source/RigVM/Public/RigVMCore/RigVMExecuteContext.h
--- a/Engine/Plugins/Runtime/RigVM/Source/RigVM/Public/RigVMCore/RigVMExecuteContext.h
+++ b/Engine/Plugins/Runtime/RigVM/Source/RigVM/Public/RigVMCore/RigVMExecuteContext.h
@@ -339,6 +339,9 @@ struct FRigVMExecuteContext
 
 	/** The current component this VM is owned by */
 	const USceneComponent* GetOwningComponent() const { return OwningComponent; }
+	
+	/** The current component this VM is owned by */
+	USceneComponent* GetMutableOwningComponent() const { return const_cast<USceneComponent*>(OwningComponent); }
 
 	/** The current actor this VM is owned by */
 	const AActor* GetOwningActor() const { return OwningActor; }
diff --git a/Engine/Plugins/Runtime/RigVM/Source/RigVMDeveloper/Private/RigVMModel/RigVMController.cpp b/Engine/Plugins/Runtime/RigVM/Source/RigVMDeveloper/Private/RigVMModel/RigVMController.cpp
--- a/Engine/Plugins/Runtime/RigVM/Source/RigVMDeveloper/Private/RigVMModel/RigVMController.cpp
+++ b/Engine/Plugins/Runtime/RigVM/Source/RigVMDeveloper/Private/RigVMModel/RigVMController.cpp
@@ -9214,6 +9214,7 @@ bool URigVMController::SetPinDefaultValue(URigVMPin* InPin, const FString& InDef
 	const FString ClampedDefaultValue = InPin->IsRootPin() ? InPin->ClampDefaultValueFromMetaData(InDefaultValue) : InDefaultValue;
 
 	bool bSetPinDefaultValueSucceeded = false;
+	bool bFoundSkippedDefaultValue = false;
 	if (InPin->IsArray())
 	{
 		if (GetSchema()->CanUnfoldPin(this, InPin))
@@ -9269,7 +9270,11 @@ bool URigVMController::SetPinDefaultValue(URigVMPin* InPin, const FString& InDef
 			if (MemberValuePair.Split(TEXT("="), &MemberName, &MemberValue))
 			{
 				URigVMPin* SubPin = InPin->FindSubPin(MemberName);
-				if (SubPin && !MemberValue.IsEmpty())
+				if (!SubPin)
+				{
+					bFoundSkippedDefaultValue = true;
+				}
+				else if( !MemberValue.IsEmpty())
 				{
 					PostProcessDefaultValue(SubPin, MemberValue);
 					if (!MemberValue.IsEmpty())
@@ -9282,31 +9287,28 @@ bool URigVMController::SetPinDefaultValue(URigVMPin* InPin, const FString& InDef
 		}
 	}
 	
-	if(!bSetPinDefaultValueSucceeded)
+	// set default value on the current pin if there is no subpin to store some of the default values
+	if(!bSetPinDefaultValueSucceeded || bFoundSkippedDefaultValue)
 	{
+		// always mark the pin's default value type as user provided
+		// even if the value didn't change. this is done to remember
+		// the value when switching versions of nodes / functions.
+		const ERigVMPinDefaultValueType NewDefaultValueType = GetDefaultValueType(InPin, ClampedDefaultValue);
+		if(InPin->DefaultValueType != NewDefaultValueType)
+		{
+			InPin->DefaultValueType = NewDefaultValueType;
+			bSetPinDefaultValueSucceeded = true;
+		}
 		// no need to send notifications if not changing the value
-		if (InPin->GetSubPins().IsEmpty())
+		if(InPin->DefaultValue != ClampedDefaultValue)
 		{
-			// always mark the pin's default value type as user provided
-			// even if the value didn't change. this is done to remember
-			// the value when switching versions of nodes / functions.
-			const ERigVMPinDefaultValueType NewDefaultValueType = GetDefaultValueType(InPin, ClampedDefaultValue);
-			if(InPin->DefaultValueType != NewDefaultValueType)
-			{
-				InPin->DefaultValueType = NewDefaultValueType;
-				bSetPinDefaultValueSucceeded = true;
-			}
-
-			if(InPin->DefaultValue != ClampedDefaultValue)
+			InPin->DefaultValue = ClampedDefaultValue;
+			Notify(ERigVMGraphNotifType::PinDefaultValueChanged, InPin);
+			if (!bSuspendNotifications)
 			{
-				InPin->DefaultValue = ClampedDefaultValue;
-				Notify(ERigVMGraphNotifType::PinDefaultValueChanged, InPin);
-				if (!bSuspendNotifications)
-				{
-					(void)Graph->MarkPackageDirty();
-				}
-				bSetPinDefaultValueSucceeded = true;
+				(void)Graph->MarkPackageDirty();
 			}
+			bSetPinDefaultValueSucceeded = true;
 		}
 	}
 
@@ -15701,17 +15703,19 @@ FName URigVMController::AddTrait(URigVMNode* InNode, UScriptStruct* InTraitScrip
 	TSharedPtr<FStructOnScope> TraitScope(new FStructOnScope(InTraitScriptStruct));
 	FRigVMTrait* Trait = (FRigVMTrait*)TraitScope->GetStructMemory();
 
-	if(!InDefaultValue.IsEmpty())
+	FString DefaultValue = InDefaultValue;
+	// Combine user overrides with struct defaults into a single default value string
+	CreateDefaultValueForStructIfRequired(InTraitScriptStruct, DefaultValue);
+	OverrideDefaultValueMember(TEXT("Name"), TEXT("\"")+ValidTraitName.ToString()+TEXT("\""), DefaultValue);
+	
+	FRigVMPinDefaultValueImportErrorContext ErrorPipe(ELogVerbosity::Verbose);
 	{
-		FRigVMPinDefaultValueImportErrorContext ErrorPipe(ELogVerbosity::Verbose);
-		{
-			// force logging to the error pipe for error detection
-			LOG_SCOPE_VERBOSITY_OVERRIDE(LogExec, ErrorPipe.GetMaxVerbosity()); 
-			InTraitScriptStruct->ImportText(*InDefaultValue, Trait, nullptr, PPF_None, &ErrorPipe, InTraitScriptStruct->GetName()); 
-		}
+		// force logging to the error pipe for error detection
+		LOG_SCOPE_VERBOSITY_OVERRIDE(LogExec, ErrorPipe.GetMaxVerbosity()); 
+		InTraitScriptStruct->ImportText(*DefaultValue, Trait, nullptr, PPF_None, &ErrorPipe, InTraitScriptStruct->GetName()); 
 	}
 
-	Trait->Name = ValidTraitName.ToString();
+	check(Trait->Name == ValidTraitName.ToString());
 
 	FString FailureReason;
 	if(!Trait->CanBeAddedToNode(InNode, &FailureReason))
@@ -15725,7 +15729,7 @@ FName URigVMController::AddTrait(URigVMNode* InNode, UScriptStruct* InTraitScrip
 	{
 		Action.SetTitle(FString::Printf(TEXT("Add Trait")));
 		GetActionStack()->BeginAction(Action);
-		GetActionStack()->AddAction(FRigVMAddTraitAction(this, InNode, ValidTraitName, InTraitScriptStruct, InDefaultValue, InPinIndex == INDEX_NONE ? InNode->GetPins().Num() : InPinIndex));
+		GetActionStack()->AddAction(FRigVMAddTraitAction(this, InNode, ValidTraitName, InTraitScriptStruct, DefaultValue, InPinIndex == INDEX_NONE ? InNode->GetPins().Num() : InPinIndex));
 	}
 
 	InNode->TraitRootPinNames.Add(ValidTraitName.ToString());
@@ -15737,14 +15741,15 @@ FName URigVMController::AddTrait(URigVMNode* InNode, UScriptStruct* InTraitScrip
 	TraitPin->CPPTypeObject = InTraitScriptStruct;
 	TraitPin->CPPTypeObjectPath = *TraitPin->CPPTypeObject->GetPathName();
 	TraitPin->Direction = ERigVMPinDirection::Input;
+	TraitPin->DefaultValue = DefaultValue;
 
 	AddNodePin(InNode, TraitPin);
 	Notify(ERigVMGraphNotifType::PinAdded, TraitPin);
 	
-	AddPinsForStruct(InTraitScriptStruct, InNode, TraitPin, TraitPin->GetDirection(), InDefaultValue, true);
+	AddPinsForStruct(InTraitScriptStruct, InNode, TraitPin, TraitPin->GetDirection(), DefaultValue, true);
 
 	FRigVMPinInfoArray ProgrammaticPins;
-	Trait->GetProgrammaticPins(this, INDEX_NONE, InDefaultValue, ProgrammaticPins);
+	Trait->GetProgrammaticPins(this, INDEX_NONE, DefaultValue, ProgrammaticPins);
 
 	const FRigVMRegistry& Registry = FRigVMRegistry::Get();
 	const FRigVMPinInfoArray PreviousPins;
@@ -18483,6 +18488,34 @@ void URigVMController::PostProcessDefaultValue(const URigVMPin* Pin, FString& Ou
 	}
 }
 
+void URigVMController::OverrideDefaultValueMember(const FString& InMemberName, const FString& InMemberValue, FString& InOutDefaultValue)
+{
+	FString NewMemberValuePair = FString::Printf(TEXT("%s=%s"), *InMemberName, *InMemberValue);
+	TArray<FString> DefaultValues = URigVMPin::SplitDefaultValue(InOutDefaultValue);
+	
+	bool bFound = false;
+	for (FString& MemberNameValuePair : DefaultValues)
+	{
+		FString MemberName, MemberValue;
+		if (MemberNameValuePair.Split(TEXT("="), &MemberName, &MemberValue))
+		{
+			if (MemberName == InMemberName)
+			{
+				MemberNameValuePair = NewMemberValuePair;
+				bFound = true;
+				break;
+			}
+		}
+	}
+
+	if (!bFound)
+	{
+		DefaultValues.Add(NewMemberValuePair);
+	}
+	
+	InOutDefaultValue = FString::Printf(TEXT("(%s)"), *FString::Join(DefaultValues, TEXT(",")));	
+}
+
 void URigVMController::ResolveTemplateNodeMetaData(URigVMTemplateNode* InNode, bool bSetupUndoRedo)
 {
 #if WITH_EDITOR
diff --git a/Engine/Plugins/Runtime/RigVM/Source/RigVMDeveloper/Private/RigVMModel/RigVMPin.cpp b/Engine/Plugins/Runtime/RigVM/Source/RigVMDeveloper/Private/RigVMModel/RigVMPin.cpp
--- a/Engine/Plugins/Runtime/RigVM/Source/RigVMDeveloper/Private/RigVMModel/RigVMPin.cpp
+++ b/Engine/Plugins/Runtime/RigVM/Source/RigVMDeveloper/Private/RigVMModel/RigVMPin.cpp
@@ -849,13 +849,15 @@ FString URigVMPin::GetDefaultValue(const URigVMPin::FPinOverride& InOverride, bo
 	}
 	else if (IsStruct())
 	{
+		static const FString EmptyStructDefaultValue = TEXT("()");
+		
+		// for trait pins, there are cases where a pin is not created for a property (see ShouldCreatePinForProperty())
+		// so we store the value of that property in the default value of the struct pin containing that property
+		// as a result, to retrieve the default value we need to combine the default value on the struct pin, with additional overrides in the available sub pins
 		if (SubPins.Num() > 0 || IsTraitPin())
 		{
-			TArray<FString> MemberDefaultValues;
-			if(IsTraitPin())
-			{
-				MemberDefaultValues.Add(FString::Printf(TEXT("Name=\"%s\""), *GetName()));
-			}
+			FString FinalDefaultValue = DefaultValue;
+
 			for (const URigVMPin* SubPin : SubPins)
 			{
 				FString MemberDefaultValue = SubPin->GetDefaultValue(InOverride, bAdaptValueForPinType);
@@ -867,18 +869,16 @@ FString URigVMPin::GetDefaultValue(const URigVMPin::FPinOverride& InOverride, bo
 				{
 					continue;
 				}
-				MemberDefaultValues.Add(FString::Printf(TEXT("%s=%s"), *SubPin->GetName(), *MemberDefaultValue));
-			}
-			if (MemberDefaultValues.Num() == 0)
-			{
-				return TEXT("()");
+
+				URigVMController::OverrideDefaultValueMember(SubPin->GetName(), MemberDefaultValue, FinalDefaultValue);
 			}
-			return FString::Printf(TEXT("(%s)"), *FString::Join(MemberDefaultValues, TEXT(",")));
+
+			return !FinalDefaultValue.IsEmpty() ? FinalDefaultValue : EmptyStructDefaultValue;
 		}
 
 		// special case certain pin types to adapt their values from
 		// alternative representations.
-		static const FString EmptyStructDefaultValue = TEXT("()");
+		
 		if(bAdaptValueForPinType && !DefaultValue.IsEmpty() && DefaultValue != EmptyStructDefaultValue)
 		{
 			if(GetScriptStruct() == TBaseStructure<FQuat>::Get())
diff --git a/Engine/Plugins/Runtime/RigVM/Source/RigVMDeveloper/Public/RigVMModel/RigVMController.h b/Engine/Plugins/Runtime/RigVM/Source/RigVMDeveloper/Public/RigVMModel/RigVMController.h
--- a/Engine/Plugins/Runtime/RigVM/Source/RigVMDeveloper/Public/RigVMModel/RigVMController.h
+++ b/Engine/Plugins/Runtime/RigVM/Source/RigVMDeveloper/Public/RigVMModel/RigVMController.h
@@ -1303,7 +1303,7 @@ private:
 	void RefreshFunctionPins(URigVMNode* InNode, bool bSetupUndoRedo = false);
 
 	void ReportRemovedLink(const FString& InSourcePinPath, const FString& InTargetPinPath, const FString& Reason = FString());
-
+public:
 	struct FPinState
 	{
 		ERigVMPinDirection Direction;
@@ -1322,10 +1322,12 @@ private:
 	void ApplyPinState(URigVMPin* InPin, const FPinState& InPinState, bool bSetupUndoRedo = false);
 	void ApplyPinStates(URigVMNode* InNode, const TMap<FString, FPinState>& InPinStates, const TMap<FString, FString>& InRedirectedPinPaths = TMap<FString, FString>(), bool bSetupUndoRedo = false);
 
+private:
 
 	static FLinearColor GetColorFromMetadata(const FString& InMetadata);
 	static void CreateDefaultValueForStructIfRequired(UScriptStruct* InStruct, FString& InOutDefaultValue);
 	static void PostProcessDefaultValue(const URigVMPin* Pin, FString& OutDefaultValue);
+	static void OverrideDefaultValueMember(const FString& InMemberName, const FString& InMemberValue, FString& InOutDefaultValue);
 
 	void ResolveTemplateNodeMetaData(URigVMTemplateNode* InNode, bool bSetupUndoRedo);
 
diff --git a/Engine/Plugins/Runtime/RigVM/Source/RigVMEditor/Private/Editor/RigVMGraphDetailCustomization.cpp b/Engine/Plugins/Runtime/RigVM/Source/RigVMEditor/Private/Editor/RigVMGraphDetailCustomization.cpp
--- a/Engine/Plugins/Runtime/RigVM/Source/RigVMEditor/Private/Editor/RigVMGraphDetailCustomization.cpp
+++ b/Engine/Plugins/Runtime/RigVM/Source/RigVMEditor/Private/Editor/RigVMGraphDetailCustomization.cpp
@@ -2103,17 +2103,20 @@ FText FRigVMWrappedNodeDetailCustomization::GetNameListText(const FNameProperty*
 	FText FirstText;
 	for(TWeakObjectPtr<URigVMDetailsViewWrapperObject> ObjectBeingCustomized : ObjectsBeingCustomized)
 	{
-		if (FName* Value = InProperty->ContainerPtrToValuePtr<FName>(ObjectBeingCustomized.Get()))
+		if (ObjectBeingCustomized.IsValid())
 		{
-			FText Text = FText::FromName(*Value);
-			if(FirstText.IsEmpty())
+			if (FName* Value = InProperty->ContainerPtrToValuePtr<FName>(ObjectBeingCustomized.Get()))
 			{
-				FirstText = Text;
-			}
-			else if(!FirstText.EqualTo(Text))
-			{
-				return RigVMGraphDetailCustomizationMultipleValues;
-			}
+				FText Text = FText::FromName(*Value);
+				if(FirstText.IsEmpty())
+				{
+					FirstText = Text;
+				}
+				else if(!FirstText.EqualTo(Text))
+				{
+					return RigVMGraphDetailCustomizationMultipleValues;
+				}
+			}	
 		}
 	}
 	return FirstText;
diff --git a/Engine/Source/Runtime/Engine/Classes/Components/SkinnedMeshComponent.h b/Engine/Source/Runtime/Engine/Classes/Components/SkinnedMeshComponent.h
--- a/Engine/Source/Runtime/Engine/Classes/Components/SkinnedMeshComponent.h
+++ b/Engine/Source/Runtime/Engine/Classes/Components/SkinnedMeshComponent.h
@@ -283,10 +283,14 @@ protected:
 	UPROPERTY(EditAnywhere, BlueprintReadOnly, Category = "Deformer", meta = (InlineEditConditionToggle))
 	bool bSetMeshDeformer = false;
 
-	/** The mesh deformer to use. If no mesh deformer is set from here or the SkeletalMesh, then we fall back to the fixed function deformation. */
+	/** The mesh deformer to use. If no mesh deformer is set from here or the SkeletalMesh, then we fall back to the fixed function deformation, unless AlwaysUseMeshDeformer is turned on. */
 	UPROPERTY(EditAnywhere, BlueprintReadOnly, Category = "Deformer", meta = (editcondition = "bSetMeshDeformer"))
 	TObjectPtr<UMeshDeformer> MeshDeformer;
 
+	/** If true, and if no mesh deformer is set from here or the SkeletalMesh, fall back to the default deformer specified in the project settings, unless DefaultMode is set to "Never" in project settings*/
+	UPROPERTY(EditAnywhere, BlueprintReadOnly, Category = "Deformer")
+	bool bAlwaysUseMeshDeformer = false;
+
 	/** Set the MeshDeformer and update the internal MeshDeformerInstance. */
 	ENGINE_API void SetMeshDeformer(bool bInSetMeshDeformer, UMeshDeformer* InMeshDeformer);
 
@@ -322,7 +326,7 @@ public:
 	 * 
 	 * This function takes GetMeshDeformerMaxLOD() into account, so there's no need to call both.
 	 */
-	UMeshDeformerInstance* GetMeshDeformerInstanceForLOD(int32 LODIndex) const;
+	ENGINE_API UMeshDeformerInstance* GetMeshDeformerInstanceForLOD(int32 LODIndex) const;
 
 	/** Max LOD at which to update or apply the MeshDeformer. */
 	ENGINE_API int32 GetMeshDeformerMaxLOD() const;
@@ -1078,6 +1082,20 @@ public:
 	UFUNCTION(BlueprintCallable, Category = "Components|SkinnedMesh")
 	ENGINE_API void UnsetMeshDeformer();
 
+	/**
+	 * Always use a MeshDeformer as long as one can be found in the project settings
+	 *
+	 * @param bShouldAlwaysUseMeshDeformer Always use mesh deformer for this component
+	 */
+	UFUNCTION(BlueprintCallable, Category = "Components|SkinnedMesh")
+	ENGINE_API void SetAlwaysUseMeshDeformer(bool bShouldAlwaysUseMeshDeformer);
+
+	/**
+	 * Returns whether the component is set to always use a mesh deformer if one can be found in the project settings
+	 */
+	UFUNCTION(BlueprintCallable, Category = "Components|SkinnedMesh")
+	ENGINE_API bool GetAlwaysUseMeshDeformer() const;
+	
 	/** 
 	 * Get Parent Bone of the input bone
 	 * 
diff --git a/Engine/Source/Runtime/Engine/Private/Components/SkinnedMeshComponent.cpp b/Engine/Source/Runtime/Engine/Private/Components/SkinnedMeshComponent.cpp
--- a/Engine/Source/Runtime/Engine/Private/Components/SkinnedMeshComponent.cpp
+++ b/Engine/Source/Runtime/Engine/Private/Components/SkinnedMeshComponent.cpp
@@ -1341,7 +1341,8 @@ void USkinnedMeshComponent::PostEditChangeProperty(FPropertyChangedEvent& Proper
 	if (const FProperty* Property = PropertyChangedEvent.Property)
 	{
 		if (Property->GetFName() == GET_MEMBER_NAME_CHECKED(USkinnedMeshComponent, MeshDeformer) ||
-			Property->GetFName() == GET_MEMBER_NAME_CHECKED(USkinnedMeshComponent, bSetMeshDeformer))
+			Property->GetFName() == GET_MEMBER_NAME_CHECKED(USkinnedMeshComponent, bSetMeshDeformer) ||
+			Property->GetFName() == GET_MEMBER_NAME_CHECKED(USkinnedMeshComponent, bAlwaysUseMeshDeformer))
 		{
 			const FMeshDeformerSet ActiveDeformers = GetActiveMeshDeformers();
 
@@ -2381,19 +2382,34 @@ FMeshDeformerSet USkinnedMeshComponent::GetActiveMeshDeformers() const
 	// If there's no user-specified deformer, find out if we need to set a default deformer for
 	// Unlimited Bone Influences.
 	const bool bIsDeformerRequiredForUBI = FGPUBaseSkinVertexFactory::GetAlwaysUseDeformerForUnlimitedBoneInfluences(GetScene()->GetShaderPlatform());
-	if (bIsDeformerRequiredForUBI && !ActiveDeformer)
+
+	if (!bIsDeformerRequestedByUser)
 	{
-		bool bMeshUsesUBI = false;
-		for (const FSkeletalMeshLODRenderData& LODRenderData : RenderData->LODRenderData)
+		bool bShouldStillUseMeshDeformer = false;
+		
+		if (bAlwaysUseMeshDeformer)
+		{
+			bShouldStillUseMeshDeformer = true;
+		}
+		else if (bIsDeformerRequiredForUBI)
 		{
-			if (LODRenderData.SkinWeightVertexBuffer.GetBoneInfluenceType() == UnlimitedBoneInfluence)
+			bool bMeshUsesUBI = false;
+			for (const FSkeletalMeshLODRenderData& LODRenderData : RenderData->LODRenderData)
 			{
-				bMeshUsesUBI = true;
-				break;
+				if (LODRenderData.SkinWeightVertexBuffer.GetBoneInfluenceType() == UnlimitedBoneInfluence)
+				{
+					bMeshUsesUBI = true;
+					break;
+				}
+			}
+
+			if (bMeshUsesUBI)
+			{
+				bShouldStillUseMeshDeformer = true;
 			}
 		}
 
-		if (bMeshUsesUBI)
+		if (bShouldStillUseMeshDeformer)
 		{
 			static IMeshDeformerProvider* MeshDeformerProvider = IMeshDeformerProvider::Get();
 
@@ -2438,7 +2454,8 @@ FMeshDeformerSet USkinnedMeshComponent::GetActiveMeshDeformers() const
 		// There should be a LODInfo entry for this LOD, but if not, default to allowing the deformer
 		const bool bAllowedByLODInfo = Index >= NumLODs || GetSkinnedAsset()->GetLODInfo(Index)->bAllowMeshDeformer;
 
-		const bool bDeformerEnabledForThisLOD = bRequiredForUBI || (bIsDeformerRequestedByUser && bAllowedByMaxLOD && bAllowedByLODInfo);
+		// Always enable for UBI requests, conditional if it is either user requested or triggered by AlwaysUseMeshDeformer 
+		const bool bDeformerEnabledForThisLOD = bRequiredForUBI || ((bIsDeformerRequestedByUser || bAlwaysUseMeshDeformer) && bAllowedByMaxLOD && bAllowedByLODInfo);
 		Result.DeformerIndexForLOD[Index] = bDeformerEnabledForThisLOD ? 0 : INDEX_NONE;
 	}
 
@@ -2514,6 +2531,18 @@ void USkinnedMeshComponent::UnsetMeshDeformer()
 	SetMeshDeformer(false, nullptr);
 }
 
+void USkinnedMeshComponent::SetAlwaysUseMeshDeformer(bool bShouldAlwaysUseMeshDeformer)
+{
+	bAlwaysUseMeshDeformer = bShouldAlwaysUseMeshDeformer;
+	// Refresh active deformer 
+	SetMeshDeformer(bSetMeshDeformer, MeshDeformer);
+}
+
+bool USkinnedMeshComponent::GetAlwaysUseMeshDeformer() const
+{
+	return bAlwaysUseMeshDeformer;
+}
+
 static TAutoConsoleVariable<int32> CVarMeshDeformerMaxLod(
 	TEXT("r.MeshDeformerMaxLOD"),
 	-1,
