[Deformer Graph] fixed compilation error caused by marking FRigVMTrait_OptimusVariableBase as abstract

diff --git a/Engine/Plugins/Animation/DeformerGraph/Source/OptimusCore/Private/ControlRig/RigUnit_Optimus.h b/Engine/Plugins/Animation/DeformerGraph/Source/OptimusCore/Private/ControlRig/RigUnit_Optimus.h
--- a/Engine/Plugins/Animation/DeformerGraph/Source/OptimusCore/Private/ControlRig/RigUnit_Optimus.h
+++ b/Engine/Plugins/Animation/DeformerGraph/Source/OptimusCore/Private/ControlRig/RigUnit_Optimus.h
@@ -64,7 +64,7 @@ struct OPTIMUSCORE_API FRigUnit_AddOptimusDeformer: public FRigUnitMutable
 
 
 
-USTRUCT(meta=(Abstract))
+USTRUCT()
 struct OPTIMUSCORE_API FRigVMTrait_OptimusVariableBase: public FRigVMTrait
 {
 	GENERATED_BODY()
