[Deformer Graph] 2nd try at fixing the compilation error

diff --git a/Engine/Plugins/Animation/DeformerGraph/Source/OptimusCore/Private/ControlRig/RigUnit_Optimus.h b/Engine/Plugins/Animation/DeformerGraph/Source/OptimusCore/Private/ControlRig/RigUnit_Optimus.h
--- a/Engine/Plugins/Animation/DeformerGraph/Source/OptimusCore/Private/ControlRig/RigUnit_Optimus.h
+++ b/Engine/Plugins/Animation/DeformerGraph/Source/OptimusCore/Private/ControlRig/RigUnit_Optimus.h
@@ -69,7 +69,7 @@ struct OPTIMUSCORE_API FRigVMTrait_OptimusVariableBase: public FRigVMTrait
 {
 	GENERATED_BODY()
 
-	virtual void SetValue(UOptimusDeformerInstance* InInstance) const PURE_VIRTUAL(FRigVMTrait_OptimusVariableBase::SetValue, {});
+	virtual void SetValue(UOptimusDeformerInstance* InInstance) const {};
 };
 
 USTRUCT(BlueprintType)
