[Deformer Graph] fixed incorrect num of permutations for skinnedmeshread DI

diff --git a/Engine/Plugins/Animation/DeformerGraph/Source/OptimusCore/Private/DataInterfaces/OptimusDataInterfaceSkinnedMeshRead.cpp b/Engine/Plugins/Animation/DeformerGraph/Source/OptimusCore/Private/DataInterfaces/OptimusDataInterfaceSkinnedMeshRead.cpp
--- a/Engine/Plugins/Animation/DeformerGraph/Source/OptimusCore/Private/DataInterfaces/OptimusDataInterfaceSkinnedMeshRead.cpp
+++ b/Engine/Plugins/Animation/DeformerGraph/Source/OptimusCore/Private/DataInterfaces/OptimusDataInterfaceSkinnedMeshRead.cpp
@@ -133,7 +133,7 @@ enum class EOptimusSkinnedMeshReadReadableOutputBuffer: uint32
 	Position = 1<<0,
 	Tangents = 1<<1,
 	VertexColor = 1<<2,
-	NumPermutations = Position | Tangents | VertexColor
+	NumPermutations = 1 << 3
 };
 ENUM_CLASS_FLAGS(EOptimusSkinnedMeshReadReadableOutputBuffer);
 
