[Deformer Graph] fixed GetMeshDeformerInstance API returning the instance manager instead of the instance matching the source deformer such that we maintain backcompat

diff --git a/Engine/Plugins/Animation/DeformerGraph/Source/OptimusCore/Private/OptimusDeformerDynamicInstanceManager.cpp b/Engine/Plugins/Animation/DeformerGraph/Source/OptimusCore/Private/OptimusDeformerDynamicInstanceManager.cpp
--- a/Engine/Plugins/Animation/DeformerGraph/Source/OptimusCore/Private/OptimusDeformerDynamicInstanceManager.cpp
+++ b/Engine/Plugins/Animation/DeformerGraph/Source/OptimusCore/Private/OptimusDeformerDynamicInstanceManager.cpp
@@ -120,6 +120,11 @@ EMeshDeformerOutputBuffer UOptimusDeformerDynamicInstanceManager::GetOutputBuffe
 	return Result;
 }
 
+UMeshDeformerInstance* UOptimusDeformerDynamicInstanceManager::GetInstanceForSourceDeformer()
+{
+	return DefaultInstance;
+}
+
 
 void UOptimusDeformerDynamicInstanceManager::OnControlRigBeginDestroy(UControlRig* InControlRig)
 {
diff --git a/Engine/Plugins/Animation/DeformerGraph/Source/OptimusCore/Public/OptimusDeformerDynamicInstanceManager.h b/Engine/Plugins/Animation/DeformerGraph/Source/OptimusCore/Public/OptimusDeformerDynamicInstanceManager.h
--- a/Engine/Plugins/Animation/DeformerGraph/Source/OptimusCore/Public/OptimusDeformerDynamicInstanceManager.h
+++ b/Engine/Plugins/Animation/DeformerGraph/Source/OptimusCore/Public/OptimusDeformerDynamicInstanceManager.h
@@ -39,6 +39,9 @@ public:
 	/** Return the buffers that this deformer can potentially write to */
 	EMeshDeformerOutputBuffer GetOutputBuffers() const override;
 
+	/** InstanceManager is an intermediate instance, call this function to get the instance for the deformer that created this instance manager */
+	UMeshDeformerInstance* GetInstanceForSourceDeformer() override;
+
 	/** Remove associated deformer instances when the rig is removed */
 	void OnControlRigBeginDestroy(UControlRig* InControlRig);
 
diff --git a/Engine/Plugins/Animation/DeformerGraph/Source/OptimusCore/Public/OptimusDeformerInstance.h b/Engine/Plugins/Animation/DeformerGraph/Source/OptimusCore/Public/OptimusDeformerInstance.h
--- a/Engine/Plugins/Animation/DeformerGraph/Source/OptimusCore/Public/OptimusDeformerInstance.h
+++ b/Engine/Plugins/Animation/DeformerGraph/Source/OptimusCore/Public/OptimusDeformerInstance.h
@@ -186,7 +186,7 @@ protected:
  * It contains the per instance deformer variable state and local state for each of the graphs in the deformer.
  */
 UCLASS(BlueprintType)
-class UOptimusDeformerInstance : public UObject
+class UOptimusDeformerInstance : public UMeshDeformerInstance
 {
 	GENERATED_BODY()
 
@@ -322,10 +322,11 @@ public:
 protected:
 	friend class UOptimusDeformerDynamicInstanceManager;
 
-	void AllocateResources();
-	void ReleaseResources();
-	void EnqueueWork(UMeshDeformerInstance::FEnqueueWorkDesc const& InDesc);
-	EMeshDeformerOutputBuffer GetOutputBuffers() const;
+	void AllocateResources() override;
+	void ReleaseResources() override;
+	void EnqueueWork(FEnqueueWorkDesc const& InDesc) override;
+	EMeshDeformerOutputBuffer GetOutputBuffers() const override;
+	UMeshDeformerInstance* GetInstanceForSourceDeformer() override {return this;};
 
 private:
 	/** The Mesh Component that owns this Mesh Deformer Instance. */
diff --git a/Engine/Source/Runtime/Engine/Classes/Animation/MeshDeformerInstance.h b/Engine/Source/Runtime/Engine/Classes/Animation/MeshDeformerInstance.h
--- a/Engine/Source/Runtime/Engine/Classes/Animation/MeshDeformerInstance.h
+++ b/Engine/Source/Runtime/Engine/Classes/Animation/MeshDeformerInstance.h
@@ -79,4 +79,9 @@ public:
 	
 	/** Return the buffers that this deformer can potentially write to */
 	ENGINE_API virtual EMeshDeformerOutputBuffer GetOutputBuffers() const PURE_VIRTUAL(, return EMeshDeformerOutputBuffer::None; );
+
+	/** Returns the specific instance that directly represents the source deformer, this is needed as a deformer may create intermediate instances that aren't
+	 * necessarily user-facing.
+	 */
+	ENGINE_API virtual UMeshDeformerInstance* GetInstanceForSourceDeformer() PURE_VIRTUAL(, return this; ); 
 };
diff --git a/Engine/Source/Runtime/Engine/Private/Components/SkinnedMeshComponent.cpp b/Engine/Source/Runtime/Engine/Private/Components/SkinnedMeshComponent.cpp
--- a/Engine/Source/Runtime/Engine/Private/Components/SkinnedMeshComponent.cpp
+++ b/Engine/Source/Runtime/Engine/Private/Components/SkinnedMeshComponent.cpp
@@ -2482,7 +2482,7 @@ void USkinnedMeshComponent::CreateMeshDeformerInstances(const FMeshDeformerSet&
 
 UMeshDeformerInstance* USkinnedMeshComponent::GetMeshDeformerInstance() const
 {
-	return MeshDeformerInstances.DeformerInstances.Num() > 0 ? MeshDeformerInstances.DeformerInstances[0] : nullptr;
+	return MeshDeformerInstances.DeformerInstances.Num() > 0 ? MeshDeformerInstances.DeformerInstances[0]->GetInstanceForSourceDeformer() : nullptr;
 }
 
 UMeshDeformerInstance* USkinnedMeshComponent::GetMeshDeformerInstanceForLOD(int32 LODIndex) const
