[Deformer Graph] Fixed assertion when loading old asset with float variables

diff --git a/Engine/Plugins/Animation/DeformerGraph/Source/OptimusCore/Private/OptimusVariableDescription.cpp b/Engine/Plugins/Animation/DeformerGraph/Source/OptimusCore/Private/OptimusVariableDescription.cpp
--- a/Engine/Plugins/Animation/DeformerGraph/Source/OptimusCore/Private/OptimusVariableDescription.cpp
+++ b/Engine/Plugins/Animation/DeformerGraph/Source/OptimusCore/Private/OptimusVariableDescription.cpp
@@ -46,14 +46,6 @@ void UOptimusVariableDescription::PostLoad()
 {
 	Super::PostLoad();
 
-	// 32-bit float data type is not supported for variables although they were allowed before. Do an in-place upgrade here. 
-	const FOptimusDataTypeHandle FloatDataType = FOptimusDataTypeRegistry::Get().FindType(*FFloatProperty::StaticClass());
-	const FOptimusDataTypeHandle DoubleDataType = FOptimusDataTypeRegistry::Get().FindType(*FDoubleProperty::StaticClass());
-	
-	if (DataType == FloatDataType)
-	{
-		SetDataType(DoubleDataType);
-	}
 	
 PRAGMA_DISABLE_DEPRECATION_WARNINGS
 
@@ -68,6 +60,20 @@ PRAGMA_DISABLE_DEPRECATION_WARNINGS
 	}
 PRAGMA_ENABLE_DEPRECATION_WARNINGS
 
+	// 32-bit float data type is not supported for variables although they were allowed before. Do an in-place upgrade here. 
+	const FOptimusDataTypeHandle FloatDataType = FOptimusDataTypeRegistry::Get().FindType(*FFloatProperty::StaticClass());
+	const FOptimusDataTypeHandle DoubleDataType = FOptimusDataTypeRegistry::Get().FindType(*FDoubleProperty::StaticClass());
+
+	if (DataType == FloatDataType)
+	{
+		TValueOrError<float, EPropertyBagResult> SavedValue = DefaultValueStruct.Value.GetValueFloat(FOptimusValueContainerStruct::ValuePropertyName);
+		SetDataType(DoubleDataType);
+		if (SavedValue.HasValue())
+		{
+			DefaultValueStruct.Value.SetValueDouble(FOptimusValueContainerStruct::ValuePropertyName, SavedValue.GetValue());
+		}
+	}
+
 	if (!DefaultValueStruct.IsInitialized())
 	{
 		SetDataType(DataType);
