[Deformer Graph] avoid OOB access when there isn't a vert color buffer using a color index mask

diff --git a/Engine/Plugins/Animation/DeformerGraph/Shaders/Private/DataInterfaceSkinnedMesh.ush b/Engine/Plugins/Animation/DeformerGraph/Shaders/Private/DataInterfaceSkinnedMesh.ush
--- a/Engine/Plugins/Animation/DeformerGraph/Shaders/Private/DataInterfaceSkinnedMesh.ush
+++ b/Engine/Plugins/Animation/DeformerGraph/Shaders/Private/DataInterfaceSkinnedMesh.ush
@@ -12,6 +12,7 @@ Buffer<uint> {DataInterfaceName}_IndexBuffer;
 Buffer<float> {DataInterfaceName}_PositionInputBuffer;
 Buffer<SNORM float4> {DataInterfaceName}_TangentInputBuffer;
 Buffer<float4> {DataInterfaceName}_ColorInputBuffer;
+uint {DataInterfaceName}_ColorIndexMask;
 Buffer<float2> {DataInterfaceName}_UVInputBuffer;
 
 uint ReadNumVertices_{DataInterfaceName}()
@@ -55,7 +56,8 @@ float4 ReadTangentZ_{DataInterfaceName}(uint VertexIndex)
 
 float4 ReadColor_{DataInterfaceName}(uint VertexIndex)
 {
-	uint BufferIndex = {DataInterfaceName}_InputStreamStart + VertexIndex;
+	// Index Mask is used here in case ColorStaticBuffer is bound to the GWhiteVertexBufferWithSRV that is only a few bytes
+	uint BufferIndex = ({DataInterfaceName}_InputStreamStart + VertexIndex) & ({DataInterfaceName}_ColorIndexMask);
 	return {DataInterfaceName}_ColorInputBuffer[BufferIndex] FMANUALFETCH_COLOR_COMPONENT_SWIZZLE;
 }
 
diff --git a/Engine/Plugins/Animation/DeformerGraph/Source/OptimusCore/Private/DataInterfaces/OptimusDataInterfaceSkinnedMesh.cpp b/Engine/Plugins/Animation/DeformerGraph/Source/OptimusCore/Private/DataInterfaces/OptimusDataInterfaceSkinnedMesh.cpp
--- a/Engine/Plugins/Animation/DeformerGraph/Source/OptimusCore/Private/DataInterfaces/OptimusDataInterfaceSkinnedMesh.cpp
+++ b/Engine/Plugins/Animation/DeformerGraph/Source/OptimusCore/Private/DataInterfaces/OptimusDataInterfaceSkinnedMesh.cpp
@@ -102,6 +102,7 @@ BEGIN_SHADER_PARAMETER_STRUCT(FSkinnedMeshDataInterfaceParameters, )
 	SHADER_PARAMETER_SRV(Buffer<SNORM float4>, TangentInputBuffer)
 	SHADER_PARAMETER_SRV(Buffer<float2>, UVInputBuffer)
 	SHADER_PARAMETER_SRV(Buffer<float4>, ColorInputBuffer)
+	SHADER_PARAMETER(uint32, ColorIndexMask)
 END_SHADER_PARAMETER_STRUCT()
 
 void UOptimusSkinnedMeshDataInterface::GetShaderParameters(TCHAR const* UID, FShaderParametersMetadataBuilder& InOutBuilder, FShaderParametersMetadataAllocations& InOutAllocations) const
@@ -200,5 +201,13 @@ void FOptimusSkinnedMeshDataProviderProxy::GatherDispatchData(FDispatchData cons
 		Parameters.TangentInputBuffer = MeshTangentBufferSRV != nullptr ? MeshTangentBufferSRV : NullSRVBinding;
 		Parameters.UVInputBuffer = MeshUVBufferSRV != nullptr ? MeshUVBufferSRV : NullSRVBinding;
 		Parameters.ColorInputBuffer = MeshColorBufferSRV != nullptr ? MeshColorBufferSRV : NullSRVBinding;
+		
+		// Basically when we are accessing GWhiteVertexBufferWithSRV(NullSRVBinding),
+		// we should not access beyond index 0 since the buffer is only a few bytes
+		
+		// See FGPUSkinPassthroughVertexFactory::UpdateUniformBuffer() and LocalVertexFactory.ush :: GetVertexFactoryIntermediates()
+		// Ideally we should be getting this value from the GPUBaseSkinVertexFactory but the need for
+		// section index make it tricky when we are doing unified dispatch
+		Parameters.ColorIndexMask = MeshColorBufferSRV != nullptr ? ~0u : 0;
 	}
 }
