[Deformer Graph] fixed uninitialized rig unit struct member

diff --git a/Engine/Plugins/Animation/DeformerGraph/Source/OptimusCore/Private/ControlRig/RigUnit_Optimus.h b/Engine/Plugins/Animation/DeformerGraph/Source/OptimusCore/Private/ControlRig/RigUnit_Optimus.h
--- a/Engine/Plugins/Animation/DeformerGraph/Source/OptimusCore/Private/ControlRig/RigUnit_Optimus.h
+++ b/Engine/Plugins/Animation/DeformerGraph/Source/OptimusCore/Private/ControlRig/RigUnit_Optimus.h
@@ -80,7 +80,7 @@ struct OPTIMUSCORE_API FRigVMTrait_SetDeformerIntVariable: public FRigVMTrait_Op
 	void SetValue(UOptimusDeformerInstance* InInstance) const override;
 	
 	UPROPERTY(EditAnywhere, Category = "Trait")
-	int32 Value;	
+	int32 Value = 0;	
 };
 
 USTRUCT(BlueprintType)
@@ -102,7 +102,7 @@ struct OPTIMUSCORE_API FRigVMTrait_SetDeformerInt2Variable: public FRigVMTrait_O
 	void SetValue(UOptimusDeformerInstance* InInstance) const override;
 	
 	UPROPERTY(EditAnywhere, Category = "Trait")
-	FIntPoint Value;	
+	FIntPoint Value = FIntPoint::ZeroValue;	
 };
 
 USTRUCT(BlueprintType)
@@ -124,7 +124,7 @@ struct OPTIMUSCORE_API FRigVMTrait_SetDeformerInt3Variable: public FRigVMTrait_O
 	void SetValue(UOptimusDeformerInstance* InInstance) const override;
 	
 	UPROPERTY(EditAnywhere, Category = "Trait")
-	FIntVector Value;	
+	FIntVector Value = FIntVector::ZeroValue;	
 };
 
 USTRUCT(BlueprintType)
@@ -146,7 +146,7 @@ struct OPTIMUSCORE_API FRigVMTrait_SetDeformerInt4Variable: public FRigVMTrait_O
 	void SetValue(UOptimusDeformerInstance* InInstance) const override;
 	
 	UPROPERTY(EditAnywhere, Category = "Trait")
-	FIntVector4 Value;	
+	FIntVector4 Value = FIntVector4::ZeroValue;	
 };
 
 USTRUCT(BlueprintType)
@@ -171,7 +171,7 @@ struct OPTIMUSCORE_API FRigVMTrait_SetDeformerFloatVariable: public FRigVMTrait_
 	void SetValue(UOptimusDeformerInstance* InInstance) const override;
 	
 	UPROPERTY(EditAnywhere, Category = "Trait")
-	double Value;	
+	double Value = 0.0;	
 };
 
 USTRUCT(BlueprintType)
@@ -193,7 +193,7 @@ struct OPTIMUSCORE_API FRigVMTrait_SetDeformerVector2Variable: public FRigVMTrai
 	void SetValue(UOptimusDeformerInstance* InInstance) const override;
 	
 	UPROPERTY(EditAnywhere, Category = "Trait")
-	FVector2D Value;	
+	FVector2D Value = FVector2D::ZeroVector;	
 };
 
 
@@ -216,7 +216,7 @@ struct OPTIMUSCORE_API FRigVMTrait_SetDeformerVectorVariable: public FRigVMTrait
 	void SetValue(UOptimusDeformerInstance* InInstance) const override;
 	
 	UPROPERTY(EditAnywhere, Category = "Trait")
-	FVector Value;	
+	FVector Value = FVector::ZeroVector; 
 };
 
 USTRUCT(BlueprintType)
@@ -239,7 +239,7 @@ struct OPTIMUSCORE_API FRigVMTrait_SetDeformerVector4Variable: public FRigVMTrai
 	void SetValue(UOptimusDeformerInstance* InInstance) const override;
 	
 	UPROPERTY(EditAnywhere, Category = "Trait")
-	FVector4 Value;	
+	FVector4 Value = FVector4::Zero();	
 };
 
 USTRUCT(BlueprintType)
@@ -261,7 +261,7 @@ struct OPTIMUSCORE_API FRigVMTrait_SetDeformerLinearColorVariable: public FRigVM
 	void SetValue(UOptimusDeformerInstance* InInstance) const override;
 	
 	UPROPERTY(EditAnywhere, Category = "Trait")
-	FLinearColor Value;	
+	FLinearColor Value = FLinearColor::Black;	
 };
 
 USTRUCT(BlueprintType)
@@ -284,7 +284,7 @@ struct OPTIMUSCORE_API FRigVMTrait_SetDeformerQuatVariable: public FRigVMTrait_O
 	void SetValue(UOptimusDeformerInstance* InInstance) const override;
 	
 	UPROPERTY(EditAnywhere, Category = "Trait")
-	FQuat Value;	
+	FQuat Value = FQuat::Identity;	
 };
 
 
@@ -307,7 +307,7 @@ struct OPTIMUSCORE_API FRigVMTrait_SetDeformerRotatorVariable: public FRigVMTrai
 	void SetValue(UOptimusDeformerInstance* InInstance) const override;
 	
 	UPROPERTY(EditAnywhere, Category = "Trait")
-	FRotator Value;	
+	FRotator Value = FRotator::ZeroRotator;	
 };
 
 USTRUCT(BlueprintType)
@@ -329,7 +329,7 @@ struct OPTIMUSCORE_API FRigVMTrait_SetDeformerTransformVariable: public FRigVMTr
 	void SetValue(UOptimusDeformerInstance* InInstance) const override;
 	
 	UPROPERTY(EditAnywhere, Category = "Trait")
-	FTransform Value;	
+	FTransform Value = FTransform::Identity;	
 };
 
 USTRUCT(BlueprintType)
