[Deformer Graph] minor cosmetic change to skeleton/advanced skeleton data interface nodes

- better display name
- marked old skeleton data interface as outdated

diff --git a/Engine/Plugins/Animation/DeformerGraph/Source/OptimusCore/Private/DataInterfaces/OptimusDataInterfaceAdvancedSkeleton.cpp b/Engine/Plugins/Animation/DeformerGraph/Source/OptimusCore/Private/DataInterfaces/OptimusDataInterfaceAdvancedSkeleton.cpp
--- a/Engine/Plugins/Animation/DeformerGraph/Source/OptimusCore/Private/DataInterfaces/OptimusDataInterfaceAdvancedSkeleton.cpp
+++ b/Engine/Plugins/Animation/DeformerGraph/Source/OptimusCore/Private/DataInterfaces/OptimusDataInterfaceAdvancedSkeleton.cpp
@@ -253,12 +253,14 @@ void UOptimusAdvancedSkeletonDataInterface::PostEditChangeChainProperty(FPropert
 
 FString UOptimusAdvancedSkeletonDataInterface::GetDisplayName() const
 {
+	FString WeightProfileName = TEXT("Default Skin Weights");
+	
 	if (SkinWeightProfile != NAME_None)
 	{
-		return TEXT("Advanced Skeleton (") + SkinWeightProfile.ToString() + TEXT(")");
+		WeightProfileName = SkinWeightProfile.ToString();
 	}
-	
-	return TEXT("Advanced Skeleton");
+
+	return FString::Printf(TEXT("Skeleton - %s"), *WeightProfileName);
 }
 
 TArray<FOptimusCDIPinDefinition> UOptimusAdvancedSkeletonDataInterface::GetPinDefinitions() const
diff --git a/Engine/Plugins/Animation/DeformerGraph/Source/OptimusCore/Private/DataInterfaces/OptimusDataInterfaceAdvancedSkeleton.h b/Engine/Plugins/Animation/DeformerGraph/Source/OptimusCore/Private/DataInterfaces/OptimusDataInterfaceAdvancedSkeleton.h
--- a/Engine/Plugins/Animation/DeformerGraph/Source/OptimusCore/Private/DataInterfaces/OptimusDataInterfaceAdvancedSkeleton.h
+++ b/Engine/Plugins/Animation/DeformerGraph/Source/OptimusCore/Private/DataInterfaces/OptimusDataInterfaceAdvancedSkeleton.h
@@ -140,10 +140,15 @@ public:
 	
 	UPROPERTY(EditAnywhere, Category = "Data Interface")
 	FName SkinWeightProfile = NAME_None;
-	
+
+	// If turned on, another set of bone matrices are computed per-frame to allow for layered skinning.
+	// It is typically used with a secondary skin weight profile storing the weights of a subset of bones like tweaker bones.
+	// The bind matrices for these bones are dynamic and computed based on their parent's current transform instead of initial transform
 	UPROPERTY(EditAnywhere, Category = "Data Interface")
-	bool bEnableLayeredSkinning = true;
-	
+	bool bEnableLayeredSkinning = false;
+
+	// Per-bone animation attributes, allows for custom bone data to be used in kernels, one of the places you can create
+	// animation attributes is Control Rig
 	UPROPERTY(EditAnywhere, Category = "Data Interface", meta = (ShowOnlyInnerProperties))
 	FOptimusAnimAttributeBufferArray AttributeBufferArray;
 
diff --git a/Engine/Plugins/Animation/DeformerGraph/Source/OptimusCore/Private/DataInterfaces/OptimusDataInterfaceSkeleton.cpp b/Engine/Plugins/Animation/DeformerGraph/Source/OptimusCore/Private/DataInterfaces/OptimusDataInterfaceSkeleton.cpp
--- a/Engine/Plugins/Animation/DeformerGraph/Source/OptimusCore/Private/DataInterfaces/OptimusDataInterfaceSkeleton.cpp
+++ b/Engine/Plugins/Animation/DeformerGraph/Source/OptimusCore/Private/DataInterfaces/OptimusDataInterfaceSkeleton.cpp
@@ -16,7 +16,7 @@
 
 FString UOptimusSkeletonDataInterface::GetDisplayName() const
 {
-	return TEXT("Skeleton");
+	return TEXT("Skeleton (Outdated)");
 }
 
 TArray<FOptimusCDIPinDefinition> UOptimusSkeletonDataInterface::GetPinDefinitions() const
diff --git a/Engine/Plugins/Animation/DeformerGraph/Source/OptimusCore/Private/DataInterfaces/OptimusDataInterfaceSkeleton.h b/Engine/Plugins/Animation/DeformerGraph/Source/OptimusCore/Private/DataInterfaces/OptimusDataInterfaceSkeleton.h
--- a/Engine/Plugins/Animation/DeformerGraph/Source/OptimusCore/Private/DataInterfaces/OptimusDataInterfaceSkeleton.h
+++ b/Engine/Plugins/Animation/DeformerGraph/Source/OptimusCore/Private/DataInterfaces/OptimusDataInterfaceSkeleton.h
@@ -11,7 +11,7 @@ class FSkeletalMeshObject;
 class FSkeletonDataInterfaceParameters;
 class USkinnedMeshComponent;
 
-/** Compute Framework Data Interface for skeletal data. */
+/** Compute Framework Data Interface for skeletal data. (Outdated)*/
 UCLASS(Category = ComputeFramework)
 class OPTIMUSCORE_API UOptimusSkeletonDataInterface : public UOptimusComputeDataInterface
 {
@@ -21,6 +21,8 @@ public:
 	//~ Begin UOptimusComputeDataInterface Interface
 	FString GetDisplayName() const override;
 	TArray<FOptimusCDIPinDefinition> GetPinDefinitions() const override;
+	// UOptimusAdvancedSkeletonDataInterface is now a better version of this data interface
+	bool IsVisible() const override {return false;};
 	TSubclassOf<UActorComponent> GetRequiredComponentClass() const override;
 	//~ End UOptimusComputeDataInterface Interface
 	
diff --git a/Engine/Plugins/Animation/DeformerGraph/Source/OptimusCore/Private/Nodes/OptimusNode_DataInterface.cpp b/Engine/Plugins/Animation/DeformerGraph/Source/OptimusCore/Private/Nodes/OptimusNode_DataInterface.cpp
--- a/Engine/Plugins/Animation/DeformerGraph/Source/OptimusCore/Private/Nodes/OptimusNode_DataInterface.cpp
+++ b/Engine/Plugins/Animation/DeformerGraph/Source/OptimusCore/Private/Nodes/OptimusNode_DataInterface.cpp
@@ -111,6 +111,9 @@ void UOptimusNode_DataInterface::PostLoadNodeSpecificData()
 	{
 		CreateComponentPin();
 	}
+
+	// Make sure to use the latest display name
+	UpdateDisplayNameFromDataInterface();
 }
 
 void UOptimusNode_DataInterface::OnDataTypeChanged(FName InTypeName)
