[Deformer Graph] use a better way to show displayname for outdated data interfaces based on their IsVisible override

diff --git a/Engine/Plugins/Animation/DeformerGraph/Source/OptimusCore/Private/DataInterfaces/OptimusDataInterfaceSkeleton.cpp b/Engine/Plugins/Animation/DeformerGraph/Source/OptimusCore/Private/DataInterfaces/OptimusDataInterfaceSkeleton.cpp
--- a/Engine/Plugins/Animation/DeformerGraph/Source/OptimusCore/Private/DataInterfaces/OptimusDataInterfaceSkeleton.cpp
+++ b/Engine/Plugins/Animation/DeformerGraph/Source/OptimusCore/Private/DataInterfaces/OptimusDataInterfaceSkeleton.cpp
@@ -16,7 +16,7 @@
 
 FString UOptimusSkeletonDataInterface::GetDisplayName() const
 {
-	return TEXT("Skeleton (Outdated)");
+	return TEXT("Skeleton");
 }
 
 TArray<FOptimusCDIPinDefinition> UOptimusSkeletonDataInterface::GetPinDefinitions() const
diff --git a/Engine/Plugins/Animation/DeformerGraph/Source/OptimusCore/Private/Nodes/OptimusNode_DataInterface.cpp b/Engine/Plugins/Animation/DeformerGraph/Source/OptimusCore/Private/Nodes/OptimusNode_DataInterface.cpp
--- a/Engine/Plugins/Animation/DeformerGraph/Source/OptimusCore/Private/Nodes/OptimusNode_DataInterface.cpp
+++ b/Engine/Plugins/Animation/DeformerGraph/Source/OptimusCore/Private/Nodes/OptimusNode_DataInterface.cpp
@@ -111,9 +111,6 @@ void UOptimusNode_DataInterface::PostLoadNodeSpecificData()
 	{
 		CreateComponentPin();
 	}
-
-	// Make sure to use the latest display name
-	UpdateDisplayNameFromDataInterface();
 }
 
 void UOptimusNode_DataInterface::OnDataTypeChanged(FName InTypeName)
@@ -361,6 +358,20 @@ void UOptimusNode_DataInterface::ConstructNode()
 	}
 }
 
+FText UOptimusNode_DataInterface::GetDisplayName() const
+{
+	FText SerializedDisplayName = Super::GetDisplayName();
+
+	if (DataInterfaceData && !DataInterfaceData->IsVisible())
+	{
+		FText OutdatedSuffix = LOCTEXT("OutdatedSuffix", "(Outdated)");
+
+		return FText::Join(FText::GetEmpty(), SerializedDisplayName, OutdatedSuffix);
+	}
+
+	return SerializedDisplayName;
+}
+
 
 void UOptimusNode_DataInterface::PostDuplicate(EDuplicateMode::Type DuplicateMode)
 {
diff --git a/Engine/Plugins/Animation/DeformerGraph/Source/OptimusCore/Private/Nodes/OptimusNode_DataInterface.h b/Engine/Plugins/Animation/DeformerGraph/Source/OptimusCore/Private/Nodes/OptimusNode_DataInterface.h
--- a/Engine/Plugins/Animation/DeformerGraph/Source/OptimusCore/Private/Nodes/OptimusNode_DataInterface.h
+++ b/Engine/Plugins/Animation/DeformerGraph/Source/OptimusCore/Private/Nodes/OptimusNode_DataInterface.h
@@ -58,6 +58,7 @@ public:
 protected:
 	// -- UOptimusNode overrides
 	void ConstructNode() override;
+	FText GetDisplayName() const override;
 	bool ValidateConnection(const UOptimusNodePin& InThisNodesPin, const UOptimusNodePin& InOtherNodesPin, FString* OutReason) const override;
 	TOptional<FText> ValidateForCompile(const FOptimusPinTraversalContext& InContext) const override;
 	void PostLoadNodeSpecificData() override;
