[Deformer Graph] minor, fix category for ReadSkinnedMesh DI

diff --git a/Engine/Plugins/Animation/DeformerGraph/Source/OptimusCore/Private/DataInterfaces/OptimusDataInterfaceSkinnedMeshRead.cpp b/Engine/Plugins/Animation/DeformerGraph/Source/OptimusCore/Private/DataInterfaces/OptimusDataInterfaceSkinnedMeshRead.cpp
--- a/Engine/Plugins/Animation/DeformerGraph/Source/OptimusCore/Private/DataInterfaces/OptimusDataInterfaceSkinnedMeshRead.cpp
+++ b/Engine/Plugins/Animation/DeformerGraph/Source/OptimusCore/Private/DataInterfaces/OptimusDataInterfaceSkinnedMeshRead.cpp
@@ -29,7 +29,7 @@ FString UOptimusSkinnedMeshReadDataInterface::GetDisplayName() const
 
 FName UOptimusSkinnedMeshReadDataInterface::GetCategory() const
 {
-	return CategoryName::OutputDataInterfaces;
+	return CategoryName::DataInterfaces;
 }
 
 TArray<FOptimusCDIPinDefinition> UOptimusSkinnedMeshReadDataInterface::GetPinDefinitions() const
