[Deformer Graph] Fix a nullptr access when using advanced skeleton DI with layered turned off

diff --git a/Engine/Plugins/Animation/DeformerGraph/Source/OptimusCore/Private/DataInterfaces/OptimusDataInterfaceAdvancedSkeleton.cpp b/Engine/Plugins/Animation/DeformerGraph/Source/OptimusCore/Private/DataInterfaces/OptimusDataInterfaceAdvancedSkeleton.cpp
--- a/Engine/Plugins/Animation/DeformerGraph/Source/OptimusCore/Private/DataInterfaces/OptimusDataInterfaceAdvancedSkeleton.cpp
+++ b/Engine/Plugins/Animation/DeformerGraph/Source/OptimusCore/Private/DataInterfaces/OptimusDataInterfaceAdvancedSkeleton.cpp
@@ -487,6 +487,11 @@ void FOptimusBoneTransformBuffer::SetData(FSkeletalMeshLODRenderData const* InLo
 	}		
 }
 
+bool FOptimusBoneTransformBuffer::HasData() const
+{
+	return NumBones.Num() > 0;
+}
+
 void FOptimusBoneTransformBuffer::AllocateResources(FRDGBuilder& GraphBuilder)
 {
 	if (BufferRefPerSection.Num() == 0)
@@ -1019,7 +1024,9 @@ void FOptimusAdvancedSkeletonDataProviderProxy::GatherDispatchData(FDispatchData
 		Parameters->BoneMatrices = BoneBufferSRV != nullptr ? BoneBufferSRV : NullSRVBinding;
 		Parameters->InputWeightStream = SkinWeightBufferSRV != nullptr ? SkinWeightBufferSRV : NullSRVBinding;
 		Parameters->InputWeightLookupStream = InputWeightLookupStreamSRV != nullptr ? InputWeightLookupStreamSRV : NullSRVBinding;
-		Parameters->LayeredBoneMatrices = LayeredBoneMatrixBuffer.BufferSRVPerSection[InvocationIndex] != nullptr ? LayeredBoneMatrixBuffer.BufferSRVPerSection[InvocationIndex] : FallbackSRV;
+		Parameters->LayeredBoneMatrices =
+			LayeredBoneMatrixBuffer.HasData() && LayeredBoneMatrixBuffer.BufferSRVPerSection[InvocationIndex] != nullptr ?
+				LayeredBoneMatrixBuffer.BufferSRVPerSection[InvocationIndex] : FallbackSRV;
 
 		for(int32 AttributeIndex = 0 ; AttributeIndex < AttributeBufferOffsets.Num(); AttributeIndex++)
 		{
diff --git a/Engine/Plugins/Animation/DeformerGraph/Source/OptimusCore/Private/DataInterfaces/OptimusDataInterfaceAdvancedSkeleton.h b/Engine/Plugins/Animation/DeformerGraph/Source/OptimusCore/Private/DataInterfaces/OptimusDataInterfaceAdvancedSkeleton.h
--- a/Engine/Plugins/Animation/DeformerGraph/Source/OptimusCore/Private/DataInterfaces/OptimusDataInterfaceAdvancedSkeleton.h
+++ b/Engine/Plugins/Animation/DeformerGraph/Source/OptimusCore/Private/DataInterfaces/OptimusDataInterfaceAdvancedSkeleton.h
@@ -192,6 +192,7 @@ struct FOptimusBoneTransformBuffer
 	TArray<int32> NumBones;
 
 	void SetData(FSkeletalMeshLODRenderData const* InLodRenderData, const TArray<FTransform>& InBoneTransforms);
+	bool HasData() const;
 
 	void AllocateResources(FRDGBuilder& GraphBuilder);
 };
