[Deformer Graph] Fixed unexpected flicker when a control is moved via anim details while multiple deformer rigs are present. The execution order was unintentionally affected by user scrubbing in anim details

diff --git a/Engine/Plugins/Animation/DeformerGraph/Source/OptimusCore/Private/OptimusDeformerDynamicInstanceManager.cpp b/Engine/Plugins/Animation/DeformerGraph/Source/OptimusCore/Private/OptimusDeformerDynamicInstanceManager.cpp
--- a/Engine/Plugins/Animation/DeformerGraph/Source/OptimusCore/Private/OptimusDeformerDynamicInstanceManager.cpp
+++ b/Engine/Plugins/Animation/DeformerGraph/Source/OptimusCore/Private/OptimusDeformerDynamicInstanceManager.cpp
@@ -210,7 +210,15 @@ void UOptimusDeformerDynamicInstanceManager::EnqueueRigDeformer(FGuid InInstance
 {
 	// Typically called from anim thread, but there shouldn't be concurrent access to this queue. All rigs running on the current mesh should run sequentially.
 	
-	// There are cases where control rig can run multiple times in one tick(like when moving controls), but we only need to enqueue once per instance
-	InstanceQueueMap.FindOrAdd(InExecutionPhase).FindOrAdd(InExecutionGroup).AddUnique(InInstanceGuid);
+	
+	TArray<FGuid>& InstanceQueueRef = InstanceQueueMap.FindOrAdd(InExecutionPhase).FindOrAdd(InExecutionGroup);
+	
+	// If we ever get duplicates, it means extra unnecessary instances were added via extra control rig evaluations triggered by user actions like moving a control.
+	// So let's invalidate those instances.
+	if (FGuid* BadInstance = InstanceQueueRef.FindByKey(InInstanceGuid))
+	{
+		*BadInstance = FGuid();
+	}
+	InstanceQueueRef.Add(InInstanceGuid);
 }
 
