Fixed OptimusDataTypeRegistry attempting to load assets removed from the registry

diff --git a/Engine/Plugins/Animation/DeformerGraph/Source/OptimusCore/Private/OptimusDataTypeRegistry.cpp b/Engine/Plugins/Animation/DeformerGraph/Source/OptimusCore/Private/OptimusDataTypeRegistry.cpp
--- a/Engine/Plugins/Animation/DeformerGraph/Source/OptimusCore/Private/OptimusDataTypeRegistry.cpp
+++ b/Engine/Plugins/Animation/DeformerGraph/Source/OptimusCore/Private/OptimusDataTypeRegistry.cpp
@@ -509,6 +509,11 @@ FName FOptimusDataTypeRegistry::GetTypeName(UScriptStruct* InStruct)
 	return Optimus::GetTypeName(InStruct, true);
 }
 
+FName FOptimusDataTypeRegistry::GetTypeName(const FAssetData& InStructAsset)
+{
+	return Optimus::GetTypeName(InStructAsset);
+}
+
 FName FOptimusDataTypeRegistry::GetArrayTypeName(const FFieldClass& InFieldClass)
 {
 	return GetArrayTypeName(GetTypeName(InFieldClass));
@@ -1603,11 +1608,8 @@ void FOptimusDataTypeRegistry::OnAssetRemoved(const FAssetData& InAssetData)
 {
 	if (InAssetData.AssetClassPath == UUserDefinedStruct::StaticClass()->GetClassPathName())
 	{
-		if (UUserDefinedStruct* UserDefinedStruct = Cast<UUserDefinedStruct>(InAssetData.GetAsset()))
-		{
-			FName TypeName = GetTypeName(UserDefinedStruct);
-			UnregisterType(TypeName);
-		}
+		// Avoid using InAssetData.GetAsset() here, the asset's module/plugin may have just been unloaded
+		UnregisterType(GetTypeName(InAssetData));
 	}
 }
 
diff --git a/Engine/Plugins/Animation/DeformerGraph/Source/OptimusCore/Private/OptimusHelpers.cpp b/Engine/Plugins/Animation/DeformerGraph/Source/OptimusCore/Private/OptimusHelpers.cpp
--- a/Engine/Plugins/Animation/DeformerGraph/Source/OptimusCore/Private/OptimusHelpers.cpp
+++ b/Engine/Plugins/Animation/DeformerGraph/Source/OptimusCore/Private/OptimusHelpers.cpp
@@ -158,19 +158,41 @@ FName Optimus::GetMemberPropertyShaderName(UScriptStruct* InStruct, const FPrope
 	return InMemberProperty->GetFName();
 }
 
+namespace Optimus::Private
+{
+	FName GetTypeNameForGuid(const FGuid& Guid)
+	{
+		return FName(*FString::Printf(TEXT("FUserDefinedStruct_%s"), *Guid.ToString()));
+	}
+}
+
 FName Optimus::GetTypeName(UScriptStruct* InStructType, bool bInShouldGetUniqueNameForUserDefinedStruct)
 {
 	if (UUserDefinedStruct* UserDefinedStruct = Cast<UUserDefinedStruct>(InStructType))
 	{
 		if (bInShouldGetUniqueNameForUserDefinedStruct)
 		{
-			return FName(*FString::Printf(TEXT("FUserDefinedStruct_%s"), *UserDefinedStruct->GetCustomGuid().ToString()));
+			return Optimus::Private::GetTypeNameForGuid(UserDefinedStruct->GetCustomGuid());
 		}
 	}
-	
+
 	return FName(*InStructType->GetStructCPPName());
 }
 
+FName Optimus::GetTypeName(const FAssetData& InStructAsset)
+{
+	check(InStructAsset.AssetClassPath == UUserDefinedStruct::StaticClass()->GetClassPathName());
+
+	FGuid Guid;
+
+	// UUserDefinedStruct::Guid is asset registry searchable, we can find it without loading the actual asset
+	static const FName NAME_Guid = GET_MEMBER_NAME_CHECKED(UUserDefinedStruct, Guid);
+
+	ensure(InStructAsset.GetTagValue(NAME_Guid, Guid));
+
+	return Optimus::Private::GetTypeNameForGuid(Guid);
+}
+
 void Optimus::ConvertObjectPathToShaderFilePath(FString& InOutPath)
 {
 	// Shader compiler recognizes "/Engine/Generated/..." path as special. 
diff --git a/Engine/Plugins/Animation/DeformerGraph/Source/OptimusCore/Public/OptimusDataTypeRegistry.h b/Engine/Plugins/Animation/DeformerGraph/Source/OptimusCore/Public/OptimusDataTypeRegistry.h
--- a/Engine/Plugins/Animation/DeformerGraph/Source/OptimusCore/Public/OptimusDataTypeRegistry.h
+++ b/Engine/Plugins/Animation/DeformerGraph/Source/OptimusCore/Public/OptimusDataTypeRegistry.h
@@ -44,6 +44,7 @@ public:
 
 	static FName GetTypeName(const FFieldClass& InFieldClass);
 	static FName GetTypeName(UScriptStruct* InStruct);
+	static FName GetTypeName(const FAssetData& InStructAsset);
 	
 	static FName GetArrayTypeName(const FFieldClass& InFieldClass);
 	static FName GetArrayTypeName(UScriptStruct* InStruct);
diff --git a/Engine/Plugins/Animation/DeformerGraph/Source/OptimusCore/Public/OptimusHelpers.h b/Engine/Plugins/Animation/DeformerGraph/Source/OptimusCore/Public/OptimusHelpers.h
--- a/Engine/Plugins/Animation/DeformerGraph/Source/OptimusCore/Public/OptimusHelpers.h
+++ b/Engine/Plugins/Animation/DeformerGraph/Source/OptimusCore/Public/OptimusHelpers.h
@@ -3,6 +3,7 @@
 #pragma once
 
 #include "CoreMinimal.h"
+#include "AssetRegistry/AssetData.h"
 #include "Math/TransformCalculus.h"
 #include "Serialization/ObjectReader.h"
 #include "Serialization/ObjectWriter.h"
@@ -108,6 +109,10 @@ namespace Optimus
 		Otherise, it returns the friendly name for user-facing shader text display*/
 	OPTIMUSCORE_API FName GetTypeName(UScriptStruct* InStruct, bool bInShouldGetUniqueNameForUserDefinedStruct = true);
 	
+	/** Return the unique type name for registry and kernel code generation.
+	    Similar to the UScriptStruct version (with bInShouldGetUniqueNameForUserDefinedStruct true) but doesn't require loading the asset. */
+	OPTIMUSCORE_API FName GetTypeName(const FAssetData& InStructAsset);
+
 	/** Return the display name for the struct to be shown in type pickers*/
 	OPTIMUSCORE_API FText GetTypeDisplayName(UScriptStruct* InStruct);
 
