[Deformer Graph] minor change to detect corner cases around variable nodes

diff --git a/Engine/Plugins/Animation/DeformerGraph/Source/OptimusCore/Private/OptimusDeformer.cpp b/Engine/Plugins/Animation/DeformerGraph/Source/OptimusCore/Private/OptimusDeformer.cpp
--- a/Engine/Plugins/Animation/DeformerGraph/Source/OptimusCore/Private/OptimusDeformer.cpp
+++ b/Engine/Plugins/Animation/DeformerGraph/Source/OptimusCore/Private/OptimusDeformer.cpp
@@ -2128,7 +2128,12 @@ TArray<FOptimusComputeGraphInfo> UOptimusDeformer::CompileNodeGraphToComputeGrap
 		}
 		else if (Cast<const IOptimusValueProvider>(ConnectedNode.Node))
 		{
-			ValueNodes.AddUnique(ConnectedNode.Node);
+			// Ideally we shouldn't hit this ensure, but in case we do, it is safe to proceed as we have this If check here
+			// However, the cause of the ensure should be investigated
+			if (ensure(!ValueNodes.Contains(ConnectedNode.Node)))
+			{
+				ValueNodes.Add(ConnectedNode.Node);
+			}
 		}
 		else if (const UOptimusNode_LoopTerminal* LoopTerminal = Cast<const UOptimusNode_LoopTerminal>(ConnectedNode.Node))
 		{
