[Deformer Graph] added support for using variables/constants to set properties on data interfaces

diff --git a/Engine/Plugins/Animation/DeformerGraph/Source/OptimusCore/Private/ControlRig/RigUnit_Optimus.cpp b/Engine/Plugins/Animation/DeformerGraph/Source/OptimusCore/Private/ControlRig/RigUnit_Optimus.cpp
--- a/Engine/Plugins/Animation/DeformerGraph/Source/OptimusCore/Private/ControlRig/RigUnit_Optimus.cpp
+++ b/Engine/Plugins/Animation/DeformerGraph/Source/OptimusCore/Private/ControlRig/RigUnit_Optimus.cpp
@@ -188,6 +188,14 @@ TArray<FRigVMUserWorkflow> FRigUnit_AddOptimusDeformer::GetSupportedWorkflows(co
 						{
 							TraitName = Controller->AddTrait(Node->GetFName(), *FRigVMTrait_SetDeformerTransformArrayVariable::StaticStruct()->GetPathName(), Variable->VariableName);
 						}
+						else if (Variable->DataType == FOptimusDataTypeRegistry::Get().FindType(*FNameProperty::StaticClass()))
+						{
+							TraitName = Controller->AddTrait(Node->GetFName(), *FRigVMTrait_SetDeformerNameVariable::StaticStruct()->GetPathName(), Variable->VariableName);
+						}
+						else if (Variable->DataType == FOptimusDataTypeRegistry::Get().FindArrayType(*FNameProperty::StaticClass()))
+						{
+							TraitName = Controller->AddTrait(Node->GetFName(), *FRigVMTrait_SetDeformerNameArrayVariable::StaticStruct()->GetPathName(), Variable->VariableName);
+						}
 						
 						Controller->SetPinExpansion(Node->FindTrait(TraitName)->GetPinPath(), true, true);
 					}
@@ -461,3 +469,13 @@ void FRigVMTrait_SetDeformerTransformArrayVariable::SetValue(UOptimusDeformerIns
 {
 	InInstance->SetTransformArrayVariable(*GetName(), Value);
 }
+
+void FRigVMTrait_SetDeformerNameVariable::SetValue(UOptimusDeformerInstance* InInstance) const
+{
+	InInstance->SetNameVariable(*GetName(), Value);
+}
+
+void FRigVMTrait_SetDeformerNameArrayVariable::SetValue(UOptimusDeformerInstance* InInstance) const
+{
+	InInstance->SetNameArrayVariable(*GetName(), Value);
+}
diff --git a/Engine/Plugins/Animation/DeformerGraph/Source/OptimusCore/Private/ControlRig/RigUnit_Optimus.h b/Engine/Plugins/Animation/DeformerGraph/Source/OptimusCore/Private/ControlRig/RigUnit_Optimus.h
--- a/Engine/Plugins/Animation/DeformerGraph/Source/OptimusCore/Private/ControlRig/RigUnit_Optimus.h
+++ b/Engine/Plugins/Animation/DeformerGraph/Source/OptimusCore/Private/ControlRig/RigUnit_Optimus.h
@@ -343,3 +343,25 @@ struct OPTIMUSCORE_API FRigVMTrait_SetDeformerTransformArrayVariable: public FRi
 	TArray<FTransform> Value;	
 };
 
+USTRUCT(BlueprintType)
+struct OPTIMUSCORE_API FRigVMTrait_SetDeformerNameVariable: public FRigVMTrait_OptimusVariableBase
+{
+	GENERATED_BODY()
+
+	void SetValue(UOptimusDeformerInstance* InInstance) const override;
+	
+	UPROPERTY(EditAnywhere, Category = "Trait")
+	FName Value = NAME_None;	
+};
+
+USTRUCT(BlueprintType)
+struct OPTIMUSCORE_API FRigVMTrait_SetDeformerNameArrayVariable: public FRigVMTrait_OptimusVariableBase
+{
+	GENERATED_BODY()
+
+	void SetValue(UOptimusDeformerInstance* InInstance) const override;
+	
+	UPROPERTY(EditAnywhere, Category = "Trait")
+	TArray<FName> Value;	
+};
+
diff --git a/Engine/Plugins/Animation/DeformerGraph/Source/OptimusCore/Private/DataInterfaces/OptimusDataInterfaceAdvancedSkeleton.cpp b/Engine/Plugins/Animation/DeformerGraph/Source/OptimusCore/Private/DataInterfaces/OptimusDataInterfaceAdvancedSkeleton.cpp
--- a/Engine/Plugins/Animation/DeformerGraph/Source/OptimusCore/Private/DataInterfaces/OptimusDataInterfaceAdvancedSkeleton.cpp
+++ b/Engine/Plugins/Animation/DeformerGraph/Source/OptimusCore/Private/DataInterfaces/OptimusDataInterfaceAdvancedSkeleton.cpp
@@ -7,6 +7,7 @@
 #include "OptimusHelpers.h"
 #include "OptimusDataDomain.h"
 #include "OptimusDataTypeRegistry.h"
+#include "OptimusDeformerInstance.h"
 #include "RenderGraphBuilder.h"
 #include "Rendering/SkeletalMeshLODRenderData.h"
 #include "Rendering/SkeletalMeshRenderData.h"
@@ -268,6 +269,19 @@ TArray<FOptimusCDIPinDefinition> UOptimusAdvancedSkeletonDataInterface::GetPinDe
 	return GetPinDefinitions_Internal();
 }
 
+TArray<FOptimusCDIPropertyPinDefinition> UOptimusAdvancedSkeletonDataInterface::GetPropertyPinDefinitions() const
+{
+	TArray<FOptimusCDIPropertyPinDefinition> PropertyPinDefinitions;
+
+	const FOptimusDataTypeHandle NameType = FOptimusDataTypeRegistry::Get().FindType(*FNameProperty::StaticClass());
+	
+	PropertyPinDefinitions.Add(
+		{SkinWeightProfilePropertyName, NameType}
+	);
+
+	return PropertyPinDefinitions;
+}
+
 
 TSubclassOf<UActorComponent> UOptimusAdvancedSkeletonDataInterface::GetRequiredComponentClass() const
 {
@@ -459,18 +473,18 @@ UComputeDataProvider* UOptimusAdvancedSkeletonDataInterface::CreateDataProvider(
 	Provider->Init(this, Cast<USkeletalMeshComponent>(InBinding));
 	return Provider;
 }
-void FOptimusBoneTransformBuffer::SetData(FSkeletalMeshLODRenderData const* InLodRenderData, const TArray<FTransform>& InBoneTransforms)
+void FOptimusBoneTransformBuffer::SetData(FSkeletalMeshLODRenderData const& InLodRenderData, const TArray<FTransform>& InBoneTransforms)
 {
 	const FName Matrix34TypeName = FOptimusDataTypeRegistry::Matrix34TypeName;
 	const FOptimusDataTypeHandle Matrix34TypeHandle = FOptimusDataTypeRegistry::Get().FindType(Matrix34TypeName);
 	int32 Matrix34ShaderSize = Matrix34TypeHandle->ShaderValueSize;
 
-	BufferData.AddDefaulted(InLodRenderData->RenderSections.Num());
-	NumBones.AddDefaulted(InLodRenderData->RenderSections.Num());
+	BufferData.AddDefaulted(InLodRenderData.RenderSections.Num());
+	NumBones.AddDefaulted(InLodRenderData.RenderSections.Num());
 				
-	for (int32 SectionIndex = 0; SectionIndex < InLodRenderData->RenderSections.Num(); ++SectionIndex)
+	for (int32 SectionIndex = 0; SectionIndex < InLodRenderData.RenderSections.Num(); ++SectionIndex)
 	{
-		FSkelMeshRenderSection const& RenderSection = InLodRenderData->RenderSections[SectionIndex];
+		FSkelMeshRenderSection const& RenderSection = InLodRenderData.RenderSections[SectionIndex];
 		NumBones[SectionIndex] = RenderSection.BoneMap.Num();
 
 		BufferData[SectionIndex].AddDefaulted(NumBones[SectionIndex] * Matrix34ShaderSize);
@@ -514,8 +528,14 @@ void FOptimusBoneTransformBuffer::AllocateResources(FRDGBuilder& GraphBuilder)
 	}
 }
 
+void UOptimusAdvancedSkeletonDataProvider::SetDeformerInstance(UOptimusDeformerInstance* InInstance)
+{
+	DeformerInstance = InInstance;
+}
+
 void UOptimusAdvancedSkeletonDataProvider::Init(const UOptimusAdvancedSkeletonDataInterface* InDataInterface, USkeletalMeshComponent* InSkeletalMesh)
 {
+	WeakDataInterface = InDataInterface;
 	bEnableLayeredSkinning = InDataInterface->bEnableLayeredSkinning;
 	SkeletalMesh = InSkeletalMesh;
 	SkinWeightProfile = InDataInterface->SkinWeightProfile;
@@ -567,22 +587,26 @@ void UOptimusAdvancedSkeletonDataProvider::Init(const UOptimusAdvancedSkeletonDa
 
 }
 
-void UOptimusAdvancedSkeletonDataProvider::ComputeBoneTransformsForLayeredSkinning(FOptimusBoneTransformBuffer& OutBoneBuffer, FSkeletalMeshLODRenderData const* InLodRenderData, const FReferenceSkeleton& InRefSkeleton)
+void UOptimusAdvancedSkeletonDataProvider::ComputeBoneTransformsForLayeredSkinning(FOptimusBoneTransformBuffer& OutBoneBuffer, FSkeletalMeshLODRenderData const& InLodRenderData, const FReferenceSkeleton& InRefSkeleton)
 {
 	if (!bIsLayeredSkinInitialized)
 	{
 		bIsLayeredSkinInitialized = true;
+
+		CachedWeightedBoneIndices.Reset();
+		CachedBoundaryBoneIndex.Reset();
+		CachedLayerSpaceInitialBoneTransform.Reset();
 		
 		const TArray<FTransform>& InitialBoneSpaceTransforms = InRefSkeleton.GetRefBonePose();
 		
 		// 1. Look for all bones with non-zero weights in this skin weight profile
 		CachedWeightedBoneIndices.Reserve(InRefSkeleton.GetNum());
 		
-		FSkinWeightVertexBuffer const* WeightBuffer = InLodRenderData->GetSkinWeightVertexBuffer();
-		if (InLodRenderData->SkinWeightProfilesData.ContainsProfile(SkinWeightProfile))
+		FSkinWeightVertexBuffer const* WeightBuffer = InLodRenderData.GetSkinWeightVertexBuffer();
+		if (InLodRenderData.SkinWeightProfilesData.ContainsProfile(SkinWeightProfile))
 		{
 			const FSkinWeightProfileStack ProfileStack{SkinWeightProfile};
-			WeightBuffer = InLodRenderData->SkinWeightProfilesData.GetOverrideBuffer(ProfileStack);
+			WeightBuffer = InLodRenderData.SkinWeightProfilesData.GetOverrideBuffer(ProfileStack);
 		}
 
 		for (uint32 VertexIndex = 0; VertexIndex < WeightBuffer->GetNumVertices(); VertexIndex++)
@@ -590,7 +614,7 @@ void UOptimusAdvancedSkeletonDataProvider::ComputeBoneTransformsForLayeredSkinni
 			// Find the render section, which we need to find the final bone index.
 			int32 SectionIndex = INDEX_NONE;
 			int32 SectionVertexIndex = INDEX_NONE;
-			InLodRenderData->GetSectionFromVertexIndex(VertexIndex, SectionIndex, SectionVertexIndex);
+			InLodRenderData.GetSectionFromVertexIndex(VertexIndex, SectionIndex, SectionVertexIndex);
 
 			uint32 VertexWeightOffset = 0;
 			uint32 VertexInfluenceCount = 0;
@@ -601,7 +625,7 @@ void UOptimusAdvancedSkeletonDataProvider::ComputeBoneTransformsForLayeredSkinni
 				if (Weight > 0)
 				{
 					int32 SectionBoneIndex = WeightBuffer->GetBoneIndex(VertexIndex, InfluenceIndex);
-					int32 FinalBoneIndex = InLodRenderData->RenderSections[SectionIndex].BoneMap[SectionBoneIndex];
+					int32 FinalBoneIndex = InLodRenderData.RenderSections[SectionIndex].BoneMap[SectionBoneIndex];
 					CachedWeightedBoneIndices.Add(FinalBoneIndex);
 				}
 			}
@@ -685,79 +709,64 @@ FComputeDataProviderRenderProxy* UOptimusAdvancedSkeletonDataProvider::GetRender
 
 	if (SkeletalMesh && SkeletalMesh->MeshObject)
 	{
-		FSkeletalMeshObject* SkeletalMeshObject = SkeletalMesh->MeshObject;	
-		if (!bSkinWeightBufferReady)
+		const UOptimusAdvancedSkeletonDataInterface* DataInterface = WeakDataInterface.Get();
+		FOptimusValueContainerStruct ValueContainer =
+			DeformerInstance->GetDataInterfacePropertyOverride(
+				WeakDataInterface.Get(),
+				DataInterface->SkinWeightProfilePropertyName
+				);
+
+		TValueOrError<FName, EPropertyBagResult> Value = ValueContainer.Value.GetValueName(FOptimusValueContainerStruct::ValuePropertyName);
+		if (Value.HasValue())
 		{
-			const int32 LodIndex = SkeletalMeshObject->GetLOD();
-			FSkeletalMeshRenderData const& SkeletalMeshRenderData = SkeletalMeshObject->GetSkeletalMeshRenderData();
-			FSkeletalMeshLODRenderData const& LodRenderData = SkeletalMeshRenderData.LODRenderData[LodIndex];
-
-			if (SkinWeightProfile.IsNone())
+			if (SkinWeightProfile != Value.GetValue())
 			{
-				bSkinWeightBufferReady = true;
+				SkinWeightProfile = Value.GetValue();
+				if (bEnableLayeredSkinning)
+				{
+					bIsLayeredSkinInitialized = false;
+				}
 			}
-			else
+		}
+		FSkeletalMeshObject* SkeletalMeshObject = SkeletalMesh->MeshObject;
+		const int32 LodIndex = SkeletalMeshObject->GetLOD();
+		FSkeletalMeshRenderData const& SkeletalMeshRenderData = SkeletalMeshObject->GetSkeletalMeshRenderData();
+		FSkeletalMeshLODRenderData const& LodRenderData = SkeletalMeshRenderData.LODRenderData[LodIndex];
+
+		bool bSkinWeightBufferReady = false;
+
+		if (SkinWeightProfile.IsNone())
+		{
+			bSkinWeightBufferReady = true;
+		}
+		else if (LodRenderData.SkinWeightProfilesData.ContainsProfile(SkinWeightProfile))
+		{
+			// Retrieve this profile's skin weight buffer
+			const FSkinWeightProfileStack ProfileStack{SkinWeightProfile};
+			FSkinWeightVertexBuffer* Buffer = LodRenderData.SkinWeightProfilesData.GetOverrideBuffer(ProfileStack);
+			bSkinWeightBufferReady = Buffer ? true : false;
+		}
+
+		if (!bSkinWeightBufferReady)
+		{
+			if (LodRenderData.SkinWeightProfilesData.ContainsProfile(SkinWeightProfile))
 			{
-				if (LodRenderData.SkinWeightProfilesData.ContainsProfile(SkinWeightProfile))
+				const FSkinWeightProfileStack ProfileStack{SkinWeightProfile};
+				// Put in a skin weight profile request
+				if (FSkinWeightProfileManager* Manager = FSkinWeightProfileManager::Get(SkeletalMesh->GetWorld()))
 				{
-					// Retrieve this profile's skin weight buffer
-					const FSkinWeightProfileStack ProfileStack{SkinWeightProfile};
-					FSkinWeightVertexBuffer* Buffer = LodRenderData.SkinWeightProfilesData.GetOverrideBuffer(ProfileStack);
-						
-					if (Buffer)
-					{
-						bSkinWeightBufferReady = true;
-					}
-					else
-					{
-						TWeakObjectPtr<USkinnedMeshComponent> WeakComponent = SkeletalMesh;
-						TWeakObjectPtr<UOptimusAdvancedSkeletonDataProvider> WeakThis = this;
-						FRequestFinished Callback = [WeakComponent, WeakThis](TWeakObjectPtr<USkeletalMesh> WeakMesh, FSkinWeightProfileStack ProfileStack)
-						{
-							// Ensure that the request objects are still valid
-							if (WeakMesh.IsValid() && WeakComponent.IsValid() && WeakThis.IsValid())
-							{
-								if (FSkeletalMeshRenderData * RenderData = WeakMesh->GetResourceForRendering())
-								{
-									const int32 NumLODs = RenderData->LODRenderData.Num();
-									for (int32 Index = 0; Index < NumLODs; ++Index)
-									{
-										FSkeletalMeshLODRenderData& LODRenderData = RenderData->LODRenderData[Index];
-										FSkinWeightProfilesData& SkinweightData = LODRenderData.SkinWeightProfilesData;
-
-										// Retrieve this profile's skin weight buffer
-										FSkinWeightVertexBuffer* Buffer = SkinweightData.GetOverrideBuffer(ProfileStack);
-										if (ensure(Buffer))
-										{
-											WeakThis->bSkinWeightBufferReady = true;
-										}
-									}
-								}
-							}
-						};
-
-						// Put in a skin weight profile request
-						if (FSkinWeightProfileManager* Manager = FSkinWeightProfileManager::Get(SkeletalMesh->GetWorld()))
-						{
-							Manager->RequestSkinWeightProfileStack(ProfileStack, SkeletalMesh->GetSkinnedAsset(), this, Callback);
-						}
-					}
-				}	
+					FRequestFinished DummyCallback = [](TWeakObjectPtr<USkeletalMesh> WeakMesh, FSkinWeightProfileStack ProfileStack){};
+					Manager->RequestSkinWeightProfileStack(ProfileStack, SkeletalMesh->GetSkinnedAsset(), this, DummyCallback);
+				}
 			}
 		}
-	
-		if (bSkinWeightBufferReady)
+		else
 		{
-			const int32 LodIndex = SkeletalMeshObject->GetLOD();
-			FSkeletalMeshRenderData const& SkeletalMeshRenderData = SkeletalMeshObject->GetSkeletalMeshRenderData();
-			FSkeletalMeshLODRenderData const* LodRenderData = &SkeletalMeshRenderData.LODRenderData[LodIndex];
-			const FReferenceSkeleton& RefSkeleton = SkeletalMesh->GetSkinnedAsset()->GetRefSkeleton();
-			
-
 			FOptimusBoneTransformBuffer LayeredBoneMatrixBuffer;
 			
 			if (bEnableLayeredSkinning)
 			{
+				const FReferenceSkeleton& RefSkeleton = SkeletalMesh->GetSkinnedAsset()->GetRefSkeleton();
 				ComputeBoneTransformsForLayeredSkinning(LayeredBoneMatrixBuffer, LodRenderData, RefSkeleton);
 			}
 			
@@ -766,10 +775,10 @@ FComputeDataProviderRenderProxy* UOptimusAdvancedSkeletonDataProvider::GetRender
 			// Section, Attribute, BoneData
 			TArray<TArray<TArray<uint8>>> AttributeBuffers;
 
-			AttributeBuffers.AddDefaulted(LodRenderData->RenderSections.Num());
-			for (int32 SectionIndex = 0; SectionIndex < LodRenderData->RenderSections.Num(); ++SectionIndex)
+			AttributeBuffers.AddDefaulted(LodRenderData.RenderSections.Num());
+			for (int32 SectionIndex = 0; SectionIndex < LodRenderData.RenderSections.Num(); ++SectionIndex)
 			{
-				FSkelMeshRenderSection const& RenderSection = LodRenderData->RenderSections[SectionIndex];
+				FSkelMeshRenderSection const& RenderSection = LodRenderData.RenderSections[SectionIndex];
 				AttributeBuffers[SectionIndex].AddDefaulted(AttributeBufferRuntimeData.Num());
 				for (int32 AttributeIndex = 0; AttributeIndex < AttributeBufferRuntimeData.Num(); ++AttributeIndex)
 				{
diff --git a/Engine/Plugins/Animation/DeformerGraph/Source/OptimusCore/Private/DataInterfaces/OptimusDataInterfaceAdvancedSkeleton.h b/Engine/Plugins/Animation/DeformerGraph/Source/OptimusCore/Private/DataInterfaces/OptimusDataInterfaceAdvancedSkeleton.h
--- a/Engine/Plugins/Animation/DeformerGraph/Source/OptimusCore/Private/DataInterfaces/OptimusDataInterfaceAdvancedSkeleton.h
+++ b/Engine/Plugins/Animation/DeformerGraph/Source/OptimusCore/Private/DataInterfaces/OptimusDataInterfaceAdvancedSkeleton.h
@@ -2,6 +2,7 @@
 
 #pragma once
 
+#include "IOptimusDeformerInstanceAccessor.h"
 #include "OptimusComputeDataInterface.h"
 #include "OptimusDataType.h"
 #include "OptimusDataTypeRegistry.h"
@@ -116,6 +117,7 @@ public:
 	//~ Begin UOptimusComputeDataInterface Interface
 	FString GetDisplayName() const override;
 	TArray<FOptimusCDIPinDefinition> GetPinDefinitions() const override;
+	TArray<FOptimusCDIPropertyPinDefinition> GetPropertyPinDefinitions() const override;
 	TSubclassOf<UActorComponent> GetRequiredComponentClass() const override;
 
 	void Initialize() override;
@@ -153,6 +155,9 @@ public:
 	FOptimusAnimAttributeBufferArray AttributeBufferArray;
 
 private:
+	friend class UOptimusAdvancedSkeletonDataProvider;
+	FName SkinWeightProfilePropertyName = GET_MEMBER_NAME_CHECKED(UOptimusAdvancedSkeletonDataInterface, SkinWeightProfile);
+	
 	FString GetUnusedAttributeName(int32 CurrentAttributeIndex ,const FString& InName) const;
 	void UpdateAttributePinNamesAndHlslIds();
 	TArray<FOptimusCDIPinDefinition> GetPinDefinitions_Internal(bool bGetAllPossiblePins = false, int32 AttributeIndexToExclude = INDEX_NONE) const;	
@@ -191,7 +196,7 @@ struct FOptimusBoneTransformBuffer
 	TArray<TArray<uint8>> BufferData;
 	TArray<int32> NumBones;
 
-	void SetData(FSkeletalMeshLODRenderData const* InLodRenderData, const TArray<FTransform>& InBoneTransforms);
+	void SetData(FSkeletalMeshLODRenderData const& InLodRenderData, const TArray<FTransform>& InBoneTransforms);
 	bool HasData() const;
 
 	void AllocateResources(FRDGBuilder& GraphBuilder);
@@ -199,11 +204,15 @@ struct FOptimusBoneTransformBuffer
 
 /** Compute Framework Data Provider for reading skeletal mesh. */
 UCLASS(BlueprintType, editinlinenew, Category = ComputeFramework)
-class UOptimusAdvancedSkeletonDataProvider : public UComputeDataProvider
+class UOptimusAdvancedSkeletonDataProvider :
+	public UComputeDataProvider,
+	public IOptimusDeformerInstanceAccessor
+
 {
 	GENERATED_BODY()
 
 public:
+	void SetDeformerInstance(UOptimusDeformerInstance* InInstance) override;
 	
 	UPROPERTY(EditAnywhere, BlueprintReadWrite, Category = Binding)
 	TObjectPtr<USkeletalMeshComponent> SkeletalMesh = nullptr;
@@ -213,7 +222,6 @@ public:
 	int32 ParameterBufferSize = 0;
 
 	FName SkinWeightProfile = NAME_None;
-	bool bSkinWeightBufferReady = false;
 	
 	bool bEnableLayeredSkinning = true;
 	bool bIsLayeredSkinInitialized = false;
@@ -230,7 +238,12 @@ public:
 		USkeletalMeshComponent* InSkeletalMesh
 	);
 
-	void ComputeBoneTransformsForLayeredSkinning(FOptimusBoneTransformBuffer& OutBoneBuffer, FSkeletalMeshLODRenderData const* InLodRenderData, const FReferenceSkeleton& InRefSkeleton);
+	void ComputeBoneTransformsForLayeredSkinning(FOptimusBoneTransformBuffer& OutBoneBuffer, FSkeletalMeshLODRenderData const& InLodRenderData, const FReferenceSkeleton& InRefSkeleton);
+private:
+	UPROPERTY()
+	TObjectPtr<UOptimusDeformerInstance> DeformerInstance;
+
+	TWeakObjectPtr<const UOptimusAdvancedSkeletonDataInterface> WeakDataInterface;
 };
 
 class FOptimusAdvancedSkeletonDataProviderProxy : public FComputeDataProviderRenderProxy
diff --git a/Engine/Plugins/Animation/DeformerGraph/Source/OptimusCore/Private/DataInterfaces/OptimusDataInterfaceGraph.cpp b/Engine/Plugins/Animation/DeformerGraph/Source/OptimusCore/Private/DataInterfaces/OptimusDataInterfaceGraph.cpp
--- a/Engine/Plugins/Animation/DeformerGraph/Source/OptimusCore/Private/DataInterfaces/OptimusDataInterfaceGraph.cpp
+++ b/Engine/Plugins/Animation/DeformerGraph/Source/OptimusCore/Private/DataInterfaces/OptimusDataInterfaceGraph.cpp
@@ -47,6 +47,21 @@ void UOptimusGraphDataInterface::Init(TArray<FOptimusGraphVariableDescription> c
 	}
 }
 
+int32 UOptimusGraphDataInterface::FindFunctionIndex(const FOptimusValueIdentifier& InValueId) const
+{
+	for (int32 Index = 0 ; Index < Variables.Num(); Index++)
+	{
+		if (Variables[Index].ValueId == InValueId)
+		{
+			return Index;
+		}
+	}
+
+	return INDEX_NONE;
+}
+
+
+
 void UOptimusGraphDataInterface::GetSupportedInputs(TArray<FShaderFunctionDefinition>& OutFunctions) const
 {
 	OutFunctions.Reserve(OutFunctions.Num() + Variables.Num());
@@ -112,9 +127,9 @@ PRAGMA_DISABLE_DEPRECATION_WARNINGS
 		if (!VariableDescription.Value_DEPRECATED.IsEmpty())
 		{
 			FOptimusDataTypeHandle DataType = FOptimusDataTypeRegistry::Get().FindType(VariableDescription.ValueType);
-			VariableDescription.ShaderValue = DataType->MakeShaderValue();
-			check(VariableDescription.ShaderValue.ArrayList.Num() == 0);
-			VariableDescription.ShaderValue.ShaderValue = VariableDescription.Value_DEPRECATED;
+			VariableDescription.ShaderValue_DEPRECATED = DataType->MakeShaderValue();
+			check(VariableDescription.ShaderValue_DEPRECATED.ArrayList.Num() == 0);
+			VariableDescription.ShaderValue_DEPRECATED.ShaderValue = VariableDescription.Value_DEPRECATED;
 			VariableDescription.Value_DEPRECATED.Reset();
 		}
 	}
@@ -128,16 +143,6 @@ void UOptimusGraphDataProvider::Init(UMeshComponent* InMeshComponent, const TArr
 	Variables = InVariables;
 	ParameterBufferSize = InParameterBufferSize;
 
-	for (FOptimusGraphVariableDescription& Variable : Variables)
-	{
-		// When source object was introduced, we also appended a unique index to the value name provided by each value provider
-		// so instead of using the name directly, we need to do this extra step
-		if (!Variable.SourceObject.IsNull())
-		{
-			Variable.CachedSourceValueName = Optimus::ExtractSourceValueName(Variable.Name);
-		}
-	}
-	
 	int32 TotalNumArrays = 0;
 	for (FOptimusGraphVariableDescription& Variable : Variables)
 	{
@@ -165,21 +170,6 @@ void UOptimusGraphDataProvider::Init(UMeshComponent* InMeshComponent, const TArr
 	}
 }
 
-void UOptimusGraphDataProvider::SetConstant(TSoftObjectPtr<UObject> InSourceObject, FShaderValueContainer const& InValue)
-{
-	for (int32 VariableIndex = 0; VariableIndex < Variables.Num(); ++VariableIndex)
-	{
-		if (Variables[VariableIndex].SourceObject == InSourceObject)
-		{
-			if (ensure(FShaderValueContainer::IsSameType(Variables[VariableIndex].ShaderValue, InValue)))
-			{
-				Variables[VariableIndex].ShaderValue = InValue;
-				break;
-			}
-		}
-	}
-}
-
 FComputeDataProviderRenderProxy* UOptimusGraphDataProvider::GetRenderProxy()
 {
 	return new FOptimusGraphDataProviderProxy(DeformerInstance, Variables, ParameterBufferSize, ParameterArrayMetadata);
@@ -224,61 +214,12 @@ FOptimusGraphDataProviderProxy::FOptimusGraphDataProviderProxy(
     	}	
 	};
 
-	TArray<UOptimusVariableDescription*> const& VariableValues = DeformerInstance->GetVariables();
+	
 	for (FOptimusGraphVariableDescription const& Variable : Variables)
 	{
-		if (Variable.ShaderValue.IsValid())
-		{
-			// Use the constant value.
-			CopyVariableToBuffer(Variable.Offset, Variable.CachedArrayIndexStart, Variable.ShaderValue);
-		}
-		else
-		{
-			// Find value from variables on the deformer instance.
-			// todo[CF]: Use a map for more efficient look up? Or something even faster like having a fixed location per variable?
-			for (UOptimusVariableDescription const* VariableValue : VariableValues)
-			{
-				if (VariableValue != nullptr)
-				{
-					if (Variable.ValueType != VariableValue->DataType->ShaderValueType)
-					{
-						continue;
-					}
-
-					bool bNameMatch = false;
-					
-					// Once upon a time when these values had no source objects, they also just have simple names
-					// so we can directly use the name to find the matching variable
-					if (Variable.SourceObject.IsNull())
-					{
-						// Using GetPlainNameString here because back then variables in the graph data interface were
-						// also generated using GetPlainNameString. This certainly creates an issue where
-						// multiple variables were sharing the same name, but at least the first matching variable
-						// would still work
-						if (Variable.Name == VariableValue->VariableName.GetPlainNameString())
-						{
-							bNameMatch = true;
-						}	
-					}
-					else
-					{
-						// When source object was introduced, we also appended a unique index to the value name
-						// so instead of using the name directly we use the source value name
-						if (Variable.CachedSourceValueName == VariableValue->VariableName)
-						{
-							bNameMatch = true;
-						}	
-					}
-
-					if (bNameMatch)
-					{
-						CopyVariableToBuffer(Variable.Offset, Variable.CachedArrayIndexStart, VariableValue->CachedShaderValue);
-						
-						break;
-					}
-				}
-			}
-		}
+		const FShaderValueContainer& ShaderValue = DeformerInstance->GetShaderValue(Variable.ValueId);
+
+		CopyVariableToBuffer(Variable.Offset, Variable.CachedArrayIndexStart, ShaderValue);
 	}
 }
 
diff --git a/Engine/Plugins/Animation/DeformerGraph/Source/OptimusCore/Private/DataInterfaces/OptimusDataInterfaceGraph.h b/Engine/Plugins/Animation/DeformerGraph/Source/OptimusCore/Private/DataInterfaces/OptimusDataInterfaceGraph.h
--- a/Engine/Plugins/Animation/DeformerGraph/Source/OptimusCore/Private/DataInterfaces/OptimusDataInterfaceGraph.h
+++ b/Engine/Plugins/Animation/DeformerGraph/Source/OptimusCore/Private/DataInterfaces/OptimusDataInterfaceGraph.h
@@ -4,9 +4,12 @@
 
 #include "IOptimusDeformerInstanceAccessor.h"
 #include "OptimusDataTypeRegistry.h"
+#include "OptimusValue.h"
 #include "ComputeFramework/ComputeDataInterface.h"
 #include "ComputeFramework/ComputeDataProvider.h"
 #include "ComputeFramework/ShaderParamTypeDefinition.h"
+
+
 #include "OptimusDataInterfaceGraph.generated.h"
 
 class UOptimusDeformerInstance;
@@ -19,7 +22,7 @@ USTRUCT()
 struct FOptimusGraphVariableDescription
 {
 	GENERATED_BODY()
-
+	
 	UPROPERTY()
 	FString	Name;
 
@@ -27,24 +30,24 @@ struct FOptimusGraphVariableDescription
 	FShaderValueTypeHandle ValueType;
 
 	UPROPERTY()
-    FShaderValueContainer ShaderValue;
-	
+	FOptimusValueIdentifier ValueId;
+
 	UPROPERTY()
 	int32 Offset = 0;
 
-	UPROPERTY()
-	TSoftObjectPtr<UObject> SourceObject;
-	
 	// Cache below are set by the data provider, computed from serialized data
-	
-	// SourceValueName extracted from Name, Initialized by the DataProvider for when SourceObject is not null
-	FString CachedSourceValueName;
-
 	int32 CachedArrayIndexStart = 0;
+
 	
 	// Deprecated 
 	UPROPERTY(meta = (DeprecatedProperty, DeprecationMessage = "Use ShaderValue instead"))
 	TArray<uint8> Value_DEPRECATED;
+	
+	UPROPERTY(meta = (DeprecatedProperty, DeprecationMessage = "Values are now store on the deformer instance"))
+	FShaderValueContainer ShaderValue_DEPRECATED;
+	
+	UPROPERTY(meta = (DeprecatedProperty, DeprecationMessage = "Node-Value mapping is now handled by deformer instance directly"))
+	TSoftObjectPtr<UObject> SourceObject_DEPRECATED;
 };
 
 /** Compute Framework Data Interface used for marshaling compute graph parameters and variables. */
@@ -55,6 +58,7 @@ class UOptimusGraphDataInterface : public UComputeDataInterface
 
 public:
 	void Init(TArray<FOptimusGraphVariableDescription> const& InVariables);
+	int32 FindFunctionIndex(const FOptimusValueIdentifier& InValueId) const;
 
 	//~ Begin UComputeDataInterface Interface
 	TCHAR const* GetClassName() const override { return TEXT("Graph"); }
@@ -67,6 +71,10 @@ public:
 
 	void PostLoad() override;
 private:
+	
+	// For PostLoad fixup
+	friend class UOptimusDeformer;
+	
 	UPROPERTY()
 	TArray<FOptimusGraphVariableDescription> Variables;
 
@@ -103,8 +111,6 @@ public:
 	
 	void Init(UMeshComponent* InMeshComponent, const TArray<FOptimusGraphVariableDescription>& InVariables, int32 InParameterBufferSize);
 	
-	void SetConstant(TSoftObjectPtr<UObject> InSourceObject, FShaderValueContainer const& InValue);
-	
 	//~ Begin UComputeDataProvider Interface
 	FComputeDataProviderRenderProxy* GetRenderProxy() override;
 	//~ End UComputeDataProvider Interface
diff --git a/Engine/Plugins/Animation/DeformerGraph/Source/OptimusCore/Private/Nodes/OptimusNode_ComputeKernelBase.cpp b/Engine/Plugins/Animation/DeformerGraph/Source/OptimusCore/Private/Nodes/OptimusNode_ComputeKernelBase.cpp
--- a/Engine/Plugins/Animation/DeformerGraph/Source/OptimusCore/Private/Nodes/OptimusNode_ComputeKernelBase.cpp
+++ b/Engine/Plugins/Animation/DeformerGraph/Source/OptimusCore/Private/Nodes/OptimusNode_ComputeKernelBase.cpp
@@ -57,9 +57,8 @@ static void CopyValueType(FShaderValueTypeHandle InValueType,  FShaderParamTypeD
 FOptimus_ComputeKernelResult UOptimusNode_ComputeKernelBase::CreateComputeKernel(
 	UObject *InKernelSourceOuter,
 	const FOptimusPinTraversalContext& InTraversalContext,
-	const FOptimus_KernelInputMap InKernelInputs,
-	const FOptimus_KernelOutputMap InKernelOutputs,
-	const TArray<const UOptimusNode *>& InValueNodes,
+	const FOptimus_KernelInputMap& InKernelInputs,
+	const FOptimus_KernelOutputMap& InKernelOutputs,
 	UComputeDataInterface* InOutKernelDataInterface,
 	FOptimus_InterfaceBindingMap& OutInputDataBindings,
 	FOptimus_InterfaceBindingMap& OutOutputDataBindings
@@ -161,8 +160,7 @@ FOptimus_ComputeKernelResult UOptimusNode_ComputeKernelBase::CreateComputeKernel
 			{
 				TOptional<FText> Result = ProcessInputPinForComputeKernel(
 					InTraversalContext, SubPin, GroupIndex == 0 ? FString() : Pin->GetName(),
-					InKernelInputs, InValueNodes,
-					KernelSource, GeneratedFunctions, OutInputDataBindings
+					InKernelInputs, KernelSource, GeneratedFunctions, OutInputDataBindings
 					);
 				if (Result.IsSet())
 				{
@@ -369,7 +367,6 @@ TOptional<FText> UOptimusNode_ComputeKernelBase::ProcessInputPinForComputeKernel
 	const UOptimusNodePin* InInputPin,
 	const FString& InGroupName,
 	const FOptimus_KernelInputMap& InKernelInputs,
-	const TArray<const UOptimusNode*>& InValueNodes,
 	UOptimusKernelSource* InKernelSource,
 	TArray<FString>& OutGeneratedFunctions,
 	FOptimus_InterfaceBindingMap& OutInputDataBindings
@@ -386,20 +383,19 @@ TOptional<FText> UOptimusNode_ComputeKernelBase::ProcessInputPinForComputeKernel
 		const UOptimusNodePin* OutputPin = InputConnection->ConnectedPin;
 		const UOptimusNode* OutputNode = OutputPin->GetOwningNode();
 		UComputeDataInterface* DataInterface = InputConnection->DataInterface;
+
+		if (!DataInterface->CanSupportUnifiedDispatch() && !InGroupName.IsEmpty())
+		{
+			return FText::Format(LOCTEXT("SecondaryGroupOnlySupportsDataInterfaceUnifiedDispatch", "Cannot connect Secondary group input binding pin '{0}' to a data interface that does not support unified dispatch, consider inserting a kernel inbetween"), FText::FromName(InInputPin->GetUniqueName()));
+		}
 		
-		const UOptimusComponentSourceBinding* ComponentBinding = nullptr;
 		int32 DataInterfaceFuncIndex = INDEX_NONE;
-		FString DataFunctionName;
 
 		if (Cast<const IOptimusComputeKernelProvider>(OutputNode))
 		{
 			if (ensure(Cast<UOptimusRawBufferDataInterface>(DataInterface)))
 			{
 				DataInterfaceFuncIndex = UOptimusRawBufferDataInterface::GetReadValueInputIndex(EOptimusBufferReadType::Default);
-				
-				TArray<FShaderFunctionDefinition> ReadFunctions;
-				DataInterface->GetSupportedInputs(ReadFunctions);
-				DataFunctionName = ReadFunctions[DataInterfaceFuncIndex].Name;	
 			}
 		}
 		else if (const IOptimusDataInterfaceProvider* InterfaceProvider = Cast<const IOptimusDataInterfaceProvider>(OutputNode)) 
@@ -407,20 +403,15 @@ TOptional<FText> UOptimusNode_ComputeKernelBase::ProcessInputPinForComputeKernel
 			if (ensure(Cast<UOptimusComputeDataInterface>(DataInterface)))
 			{
 				DataInterfaceFuncIndex = InterfaceProvider->GetDataFunctionIndexFromPin(OutputPin);
-				
-				TArray<FShaderFunctionDefinition> ReadFunctions;
-				DataInterface->GetSupportedInputs(ReadFunctions);
-				DataFunctionName = ReadFunctions[DataInterfaceFuncIndex].Name;
 			}
 		}
-		else if (const IOptimusValueProvider * ValueProvider = Cast<const IOptimusValueProvider>(OutputNode))
+		else if (const IOptimusValueProvider* ValueProvider = Cast<const IOptimusValueProvider>(OutputNode))
 		{
-			if (ensure(Cast<UOptimusGraphDataInterface>(DataInterface)))
+			UOptimusGraphDataInterface* GraphDataInterface = Cast<UOptimusGraphDataInterface>(DataInterface);
+			if (ensure(GraphDataInterface))
 			{
-				// Value nodes bind the single graph data interface.
-				DataInterfaceFuncIndex = InValueNodes.Find(OutputNode);
+				DataInterfaceFuncIndex = GraphDataInterface->FindFunctionIndex(ValueProvider->GetValueIdentifier());
 				check(DataInterfaceFuncIndex != INDEX_NONE);
-				DataFunctionName = Optimus::MakeUniqueValueName(ValueProvider->GetValueName(), DataInterfaceFuncIndex);
 			}
 		}
 		else if (Cast<const UOptimusNode_LoopTerminal>(OutputNode))
@@ -428,17 +419,13 @@ TOptional<FText> UOptimusNode_ComputeKernelBase::ProcessInputPinForComputeKernel
 			if (ensure(Cast<const UOptimusLoopTerminalDataInterface>(DataInterface)))
 			{
 				DataInterfaceFuncIndex = UOptimusNode_LoopTerminal::GetDataFunctionIndexFromPin(OutputPin);
-				TArray<FShaderFunctionDefinition> ReadFunctions;
-				DataInterface->GetSupportedInputs(ReadFunctions);
-				DataFunctionName = ReadFunctions[DataInterfaceFuncIndex].Name;
 			}
 		}
 
-		if (!DataInterface->CanSupportUnifiedDispatch() && !InGroupName.IsEmpty())
-		{
-			return FText::Format(LOCTEXT("SecondaryGroupOnlySupportsDataInterfaceUnifiedDispatch", "Cannot connect Secondary group input binding pin '{0}' to a data interface that does not support unified dispatch, consider inserting a kernel inbetween"), FText::FromName(InInputPin->GetUniqueName()));
-		}
-
+		TArray<FShaderFunctionDefinition> ReadFunctions;
+		DataInterface->GetSupportedInputs(ReadFunctions);
+		FString DataFunctionName = ReadFunctions[DataInterfaceFuncIndex].Name;
+		
 		// The shader function definition that exposes the function that we use to
 		// read values to input into the kernel.
 		FShaderFunctionDefinition FuncDef;
diff --git a/Engine/Plugins/Animation/DeformerGraph/Source/OptimusCore/Private/Nodes/OptimusNode_ComputeKernelBase.h b/Engine/Plugins/Animation/DeformerGraph/Source/OptimusCore/Private/Nodes/OptimusNode_ComputeKernelBase.h
--- a/Engine/Plugins/Animation/DeformerGraph/Source/OptimusCore/Private/Nodes/OptimusNode_ComputeKernelBase.h
+++ b/Engine/Plugins/Animation/DeformerGraph/Source/OptimusCore/Private/Nodes/OptimusNode_ComputeKernelBase.h
@@ -39,9 +39,8 @@ public:
 	FOptimus_ComputeKernelResult CreateComputeKernel(
 		UObject* InKernelSourceOuter,
 		const FOptimusPinTraversalContext& InTraversalContext,
-		const FOptimus_KernelInputMap InKernelInputs,
-		const FOptimus_KernelOutputMap InKernelOutputs,
-		const TArray<const UOptimusNode*>& InValueNodes,
+		const FOptimus_KernelInputMap& InKernelInputs,
+		const FOptimus_KernelOutputMap& InKernelOutputs,
 		UComputeDataInterface* InOutKernelDataInterface,
 		FOptimus_InterfaceBindingMap& OutInputDataBindings,
 		FOptimus_InterfaceBindingMap& OutOutputDataBindings
@@ -102,7 +101,6 @@ private:
 		const UOptimusNodePin* InInputPin,
 		const FString& InGroupName,
 		const FOptimus_KernelInputMap& InKernelInputs,
-		const TArray<const UOptimusNode*>& InValueNodes,
 		UOptimusKernelSource* InKernelSource,
 		TArray<FString>& OutGeneratedFunctions,
 		FOptimus_InterfaceBindingMap& OutInputDataBindings
diff --git a/Engine/Plugins/Animation/DeformerGraph/Source/OptimusCore/Private/Nodes/OptimusNode_ConstantValue.cpp b/Engine/Plugins/Animation/DeformerGraph/Source/OptimusCore/Private/Nodes/OptimusNode_ConstantValue.cpp
--- a/Engine/Plugins/Animation/DeformerGraph/Source/OptimusCore/Private/Nodes/OptimusNode_ConstantValue.cpp
+++ b/Engine/Plugins/Animation/DeformerGraph/Source/OptimusCore/Private/Nodes/OptimusNode_ConstantValue.cpp
@@ -3,10 +3,12 @@
 #include "Nodes/OptimusNode_ConstantValue.h"
 
 #include "OptimusDataTypeRegistry.h"
+#include "OptimusValue.h"
 #include "OptimusNodePin.h"
 #include "OptimusNodeGraph.h"
 
 #include "OptimusHelpers.h"
+#include "OptimusValueContainerStruct.h"
 
 #include UE_INLINE_GENERATED_CPP_BY_NAME(OptimusNode_ConstantValue)
 
@@ -83,6 +85,27 @@ void UOptimusNode_ConstantValue::PostLoadNodeSpecificData()
 	}
 }
 
+FOptimusValueContainerStruct UOptimusNode_ConstantValue::GetValue() const
+{
+	FOptimusValueContainerStruct ValueContainer;
+	
+	const UOptimusNodePin *ValuePin = FindPinFromPath({TEXT("Value")});
+	if (ensure(ValuePin))
+	{
+		const FProperty *ValueProperty = ValuePin->GetPropertyFromPin();
+		FOptimusDataTypeRef DataType = GetValueDataType();
+		if (ensure(ValueProperty) && ensure(DataType.IsValid()))
+		{
+			TArrayView<const uint8> ValueData(ValueProperty->ContainerPtrToValuePtr<uint8>(this), ValueProperty->GetSize());
+
+			ValueContainer.SetType(DataType);
+			ValueContainer.SetValue(DataType, ValueData);
+		}
+	}
+
+	return ValueContainer;
+}
+
 
 #if WITH_EDITOR
 
@@ -101,13 +124,12 @@ void UOptimusNode_ConstantValue::PostEditChangeProperty(FPropertyChangedEvent& P
 #endif // WITH_EDITOR
 
 
-FString UOptimusNode_ConstantValue::GetValueName() const
+FOptimusValueIdentifier UOptimusNode_ConstantValue::GetValueIdentifier() const
 {
-	return GetName();
+	return {EOptimusValueType::Constant, Optimus::GetSanitizedNameForHlsl(*GetNodePath())};
 }
 
-
-FOptimusDataTypeRef UOptimusNode_ConstantValue::GetValueType() const
+FOptimusDataTypeRef UOptimusNode_ConstantValue::GetValueDataType() const
 {
 	const UOptimusNode_ConstantValueGeneratorClass* Class = GetGeneratorClass();
 	if (ensure(Class))
@@ -117,30 +139,6 @@ FOptimusDataTypeRef UOptimusNode_ConstantValue::GetValueType() const
 	return {};
 }
 
-
-FShaderValueContainer UOptimusNode_ConstantValue::GetShaderValue() const
-{
-	const UOptimusNodePin *ValuePin = FindPinFromPath({TEXT("Value")});
-	if (ensure(ValuePin))
-	{
-		const FProperty *ValueProperty = ValuePin->GetPropertyFromPin();
-		FOptimusDataTypeRef DataType = GetValueType();
-		if (ensure(ValueProperty) && ensure(DataType.IsValid()))
-		{
-			TArrayView<const uint8> ValueData(ValueProperty->ContainerPtrToValuePtr<uint8>(this), ValueProperty->GetSize());
-			FShaderValueContainer ValueResult = DataType->MakeShaderValue();
-
-			if (DataType->ConvertPropertyValueToShader(ValueData, ValueResult))
-			{
-				return ValueResult;
-			}
-		}
-	}
-	
-	return {};
-}
-
-
 FTopLevelAssetPath UOptimusNode_ConstantValue::GetAssetPathForClassDefiner() const
 {
 	return StaticClass()->GetClassPathName();
@@ -180,7 +178,7 @@ UClass* UOptimusNode_ConstantValue::GetClassFromCreationString(
 
 void UOptimusNode_ConstantValue::ConstructNode()
 {
-	SetDisplayName(FText::Format(FText::FromString(TEXT("{0} Constant")), GetValueType()->DisplayName));
+	SetDisplayName(FText::Format(FText::FromString(TEXT("{0} Constant")), GetValueDataType()->DisplayName));
 
 	UOptimusNode::ConstructNode();
 }
diff --git a/Engine/Plugins/Animation/DeformerGraph/Source/OptimusCore/Private/Nodes/OptimusNode_ConstantValue.h b/Engine/Plugins/Animation/DeformerGraph/Source/OptimusCore/Private/Nodes/OptimusNode_ConstantValue.h
--- a/Engine/Plugins/Animation/DeformerGraph/Source/OptimusCore/Private/Nodes/OptimusNode_ConstantValue.h
+++ b/Engine/Plugins/Animation/DeformerGraph/Source/OptimusCore/Private/Nodes/OptimusNode_ConstantValue.h
@@ -64,9 +64,9 @@ public:
 #endif
 
 	// IOptimusValueProvider overrides 
-	FString GetValueName() const override;
-	FOptimusDataTypeRef GetValueType() const override;
-	FShaderValueContainer GetShaderValue() const override;
+	FOptimusValueIdentifier GetValueIdentifier() const override;
+	FOptimusDataTypeRef GetValueDataType() const override;
+	FOptimusValueContainerStruct GetValue() const override;
 
 	// IOptimusGeneratedClassDefiner implementation
 	FTopLevelAssetPath GetAssetPathForClassDefiner() const override;
diff --git a/Engine/Plugins/Animation/DeformerGraph/Source/OptimusCore/Private/Nodes/OptimusNode_DataInterface.cpp b/Engine/Plugins/Animation/DeformerGraph/Source/OptimusCore/Private/Nodes/OptimusNode_DataInterface.cpp
--- a/Engine/Plugins/Animation/DeformerGraph/Source/OptimusCore/Private/Nodes/OptimusNode_DataInterface.cpp
+++ b/Engine/Plugins/Animation/DeformerGraph/Source/OptimusCore/Private/Nodes/OptimusNode_DataInterface.cpp
@@ -111,6 +111,12 @@ void UOptimusNode_DataInterface::PostLoadNodeSpecificData()
 	{
 		CreateComponentPin();
 	}
+
+	// Add in the property pins if there are any.
+	if (GetLinkerCustomVersion(FOptimusObjectVersion::GUID) < FOptimusObjectVersion::PropertyPinSupport)
+	{
+		CreatePropertyPinsFromDataInterface(DataInterfaceData, false);
+	}
 }
 
 void UOptimusNode_DataInterface::OnDataTypeChanged(FName InTypeName)
@@ -165,7 +171,8 @@ void UOptimusNode_DataInterface::RecreatePinsFromPinDefinitions()
 		RemovePin(Pin);
 	}
 
-	CreatePinsFromDataInterface(DataInterfaceData, true);
+	CreatePropertyPinsFromDataInterface(DataInterfaceData, true);
+	CreateShaderPinsFromDataInterface(DataInterfaceData, true);
 
 	for (UOptimusNodePin* Pin : GetPins())
 	{
@@ -341,6 +348,23 @@ EOptimusPinMutability UOptimusNode_DataInterface::GetOutputPinMutability(const U
 	return PinDefinitions[PinDefinitionIndex].bMutable ? EOptimusPinMutability::Mutable : EOptimusPinMutability::Immutable;
 }
 
+TArray<UOptimusNodePin*> UOptimusNode_DataInterface::GetPropertyPins() const
+{
+	const TArray<FOptimusCDIPropertyPinDefinition> PropertyPinDefinitions = DataInterfaceData->GetPropertyPinDefinitions();
+
+	TArray<UOptimusNodePin*> Result;
+	for (const FOptimusCDIPropertyPinDefinition& Definition : PropertyPinDefinitions)
+	{
+		UOptimusNodePin* Pin = FindPinFromPath({Definition.PinName});
+		if (ensure(Pin) && ensure(Pin->GetDirection() == EOptimusNodePinDirection::Input))
+		{
+			Result.Add(Pin);
+		}
+	}
+
+	return Result;
+}
+
 
 void UOptimusNode_DataInterface::ConstructNode()
 {
@@ -354,7 +378,8 @@ void UOptimusNode_DataInterface::ConstructNode()
 		}
 		SetDisplayName(FText::FromString(DataInterfaceData->GetDisplayName()));
 		CreateComponentPin();
-		CreatePinsFromDataInterface(DataInterfaceData, false);
+		CreatePropertyPinsFromDataInterface(DataInterfaceData, false);
+		CreateShaderPinsFromDataInterface(DataInterfaceData, false);
 	}
 }
 
@@ -380,7 +405,7 @@ void UOptimusNode_DataInterface::PostDuplicate(EDuplicateMode::Type DuplicateMod
 }
 
 
-void UOptimusNode_DataInterface::CreatePinsFromDataInterface(const UOptimusComputeDataInterface* InDataInterface, bool bSupportUndo)
+void UOptimusNode_DataInterface::CreateShaderPinsFromDataInterface(const UOptimusComputeDataInterface* InDataInterface, bool bSupportUndo)
 {
 	// A data interface provides read and write functions. A data interface node exposes
 	// the read functions as output pins to be fed into kernel nodes (or into other interface
@@ -546,6 +571,31 @@ void UOptimusNode_DataInterface::CreatePinFromDefinition(const FOptimusCDIPinDef
 	}
 }
 
+void UOptimusNode_DataInterface::CreatePropertyPinsFromDataInterface(const UOptimusComputeDataInterface* InDataInterface, bool bSupportUndo)
+{
+	const TArray<FOptimusCDIPropertyPinDefinition> PropertyPinDefinitions = InDataInterface->GetPropertyPinDefinitions();
+
+	// Property pins should go before any shader pins
+	UOptimusNodePin* BeforePin =nullptr;
+	if (GetPins().Num() > 1)
+	{
+		BeforePin = GetPins()[1];
+	}
+	
+	for (const FOptimusCDIPropertyPinDefinition& Definition : PropertyPinDefinitions)
+	{
+		if (bSupportUndo)
+		{
+			AddPin(Definition.PinName, EOptimusNodePinDirection::Input, FOptimusDataDomain(), Definition.DataType, BeforePin);
+		}
+		else
+		{
+			AddPinDirect(Definition.PinName, EOptimusNodePinDirection::Input, FOptimusDataDomain(), Definition.DataType, BeforePin);
+		}
+	}
+}
+
+
 void UOptimusNode_DataInterface::CreateComponentPin()
 {
 	const FOptimusDataTypeRegistry& TypeRegistry = FOptimusDataTypeRegistry::Get();
diff --git a/Engine/Plugins/Animation/DeformerGraph/Source/OptimusCore/Private/Nodes/OptimusNode_DataInterface.h b/Engine/Plugins/Animation/DeformerGraph/Source/OptimusCore/Private/Nodes/OptimusNode_DataInterface.h
--- a/Engine/Plugins/Animation/DeformerGraph/Source/OptimusCore/Private/Nodes/OptimusNode_DataInterface.h
+++ b/Engine/Plugins/Animation/DeformerGraph/Source/OptimusCore/Private/Nodes/OptimusNode_DataInterface.h
@@ -5,6 +5,7 @@
 #include "IOptimusComponentBindingProvider.h"
 #include "IOptimusDataInterfaceProvider.h"
 #include "IOptimusPinMutabilityDefiner.h"
+#include "IOptimusPropertyPinProvider.h"
 #include "OptimusComputeDataInterface.h"
 #include "OptimusComponentSource.h"
 
@@ -21,7 +22,8 @@ class UOptimusNode_DataInterface :
 	public UOptimusNode,
 	public IOptimusDataInterfaceProvider,
 	public IOptimusComponentBindingProvider,
-	public IOptimusPinMutabilityDefiner
+	public IOptimusPinMutabilityDefiner,
+	public IOptimusPropertyPinProvider
 {
 	GENERATED_BODY()
 
@@ -54,7 +56,9 @@ public:
 
 	// -- IOptimusPinMutabilityDefiner
 	EOptimusPinMutability GetOutputPinMutability(const UOptimusNodePin* InPin) const override;
-	
+
+	// -- IOptimusPropertyPinProvider
+	TArray<UOptimusNodePin*> GetPropertyPins() const override;
 protected:
 	// -- UOptimusNode overrides
 	void ConstructNode() override;
@@ -71,9 +75,7 @@ protected:
 	void PostDuplicate(EDuplicateMode::Type DuplicateMode) override;
 	
 private:
-	void CreatePinsFromDataInterface(const UOptimusComputeDataInterface *InDataInterface, bool bSupportUndo);
-	
-
+	void CreateShaderPinsFromDataInterface(const UOptimusComputeDataInterface *InDataInterface, bool bSupportUndo);
 	void CreatePinFromDefinition(
 		const FOptimusCDIPinDefinition &InDefinition,
 		const TMap<FString, const FShaderFunctionDefinition *>& InReadFunctionMap,
@@ -81,6 +83,7 @@ private:
 		bool bSupportUndo
 		);
 
+	void CreatePropertyPinsFromDataInterface(const UOptimusComputeDataInterface *InDataInterface, bool bSupportUndo);
 	void CreateComponentPin();
 
 protected:
diff --git a/Engine/Plugins/Animation/DeformerGraph/Source/OptimusCore/Private/Nodes/OptimusNode_GetVariable.cpp b/Engine/Plugins/Animation/DeformerGraph/Source/OptimusCore/Private/Nodes/OptimusNode_GetVariable.cpp
--- a/Engine/Plugins/Animation/DeformerGraph/Source/OptimusCore/Private/Nodes/OptimusNode_GetVariable.cpp
+++ b/Engine/Plugins/Animation/DeformerGraph/Source/OptimusCore/Private/Nodes/OptimusNode_GetVariable.cpp
@@ -115,10 +115,9 @@ void UOptimusNode_GetVariable::ExportCustomProperties(FOutputDevice& Out, uint32
 		Out.Logf(TEXT("%sCustomProperties VariableDefinition Name=\"%s\" Type=%s"),
 			FCString::Spc(Indent), *Var->VariableName.ToString(), *Var->DataType->TypeName.ToString());
 
-		if (const FProperty* Property = Var->DefaultValueStruct.GetValueProperty())
+		if (Var->DefaultValueStruct.IsInitialized())
 		{
-			FString ValueStr;
-			Property->ExportTextItem_InContainer(ValueStr, Var->DefaultValueStruct.GetValueMemory(), nullptr, nullptr, PPF_None);
+			FString ValueStr = Var->DefaultValueStruct.GetValueAsString();
 			Out.Logf(TEXT(" DefaultValue=\"%s\""), *ValueStr.ReplaceCharWithEscapedChar());
 		}
 		Out.Logf(TEXT("\n"));
@@ -167,9 +166,9 @@ void UOptimusNode_GetVariable::SetVariableDescription(UOptimusVariableDescriptio
 		return;
 	}
 
-	if (!EnumHasAnyFlags(InVariableDesc->DataType->UsageFlags, EOptimusDataTypeUsageFlags::Variable))
+	if (!EnumHasAnyFlags(InVariableDesc->DataType->UsageFlags, EOptimusDataTypeUsageFlags::Variable | EOptimusDataTypeUsageFlags::Property))
 	{
-		UE_LOG(LogOptimusCore, Error, TEXT("Data type '%s' is not usable in a resource"),
+		UE_LOG(LogOptimusCore, Error, TEXT("Data type '%s' is not usable in a variable"),
 		    *InVariableDesc->DataType->TypeName.ToString());
 		return;
 	}
@@ -195,18 +194,18 @@ TOptional<FText> UOptimusNode_GetVariable::ValidateForCompile(const FOptimusPinT
 	return {};
 }
 
-FString UOptimusNode_GetVariable::GetValueName() const
+FOptimusValueIdentifier UOptimusNode_GetVariable::GetValueIdentifier() const
 {
 	if (const UOptimusVariableDescription* Var = VariableDesc.Get())
 	{
-		return Var->VariableName.ToString();
+		return {EOptimusValueType::Variable, Var->VariableName};
 	}
 
 	return {};
 }
 
 
-FOptimusDataTypeRef UOptimusNode_GetVariable::GetValueType() const
+FOptimusDataTypeRef UOptimusNode_GetVariable::GetValueDataType() const
 {
 	if (const UOptimusVariableDescription* Var = VariableDesc.Get())
 	{
@@ -216,11 +215,12 @@ FOptimusDataTypeRef UOptimusNode_GetVariable::GetValueType() const
 	return {};
 }
 
-
-FShaderValueContainer UOptimusNode_GetVariable::GetShaderValue() const
+FOptimusValueContainerStruct UOptimusNode_GetVariable::GetValue() const
 {
-	// Unused, Graph data interface is responsible for collecting value from UOptimusVariableDescription
-	checkNoEntry();
+	if (const UOptimusVariableDescription* Var = VariableDesc.Get())
+	{
+		return Var->DefaultValueStruct;
+	}
 
 	return {};
 }
diff --git a/Engine/Plugins/Animation/DeformerGraph/Source/OptimusCore/Private/Nodes/OptimusNode_GetVariable.h b/Engine/Plugins/Animation/DeformerGraph/Source/OptimusCore/Private/Nodes/OptimusNode_GetVariable.h
--- a/Engine/Plugins/Animation/DeformerGraph/Source/OptimusCore/Private/Nodes/OptimusNode_GetVariable.h
+++ b/Engine/Plugins/Animation/DeformerGraph/Source/OptimusCore/Private/Nodes/OptimusNode_GetVariable.h
@@ -54,9 +54,9 @@ public:
 	TOptional<FText> ValidateForCompile(const FOptimusPinTraversalContext& InContext) const override;
 	
 	// IOptimusValueProvider overrides 
-	FString GetValueName() const override;
-	FOptimusDataTypeRef GetValueType() const override;
-	FShaderValueContainer GetShaderValue() const override;
+	FOptimusValueIdentifier GetValueIdentifier() const override;
+	FOptimusDataTypeRef GetValueDataType() const override;
+	FOptimusValueContainerStruct GetValue() const override;
 
 	//IOptimusPinMutabilityDefiner overrides 
 	EOptimusPinMutability GetOutputPinMutability(const UOptimusNodePin* InPin) const override { return EOptimusPinMutability::Mutable; };
diff --git a/Engine/Plugins/Animation/DeformerGraph/Source/OptimusCore/Private/OptimusComputeGraph.cpp b/Engine/Plugins/Animation/DeformerGraph/Source/OptimusCore/Private/OptimusComputeGraph.cpp
--- a/Engine/Plugins/Animation/DeformerGraph/Source/OptimusCore/Private/OptimusComputeGraph.cpp
+++ b/Engine/Plugins/Animation/DeformerGraph/Source/OptimusCore/Private/OptimusComputeGraph.cpp
@@ -16,6 +16,7 @@
 #include "OptimusNode.h"
 #include "OptimusObjectVersion.h"
 #include "Animation/MeshDeformerInstance.h"
+#include "DataInterfaces/OptimusDataInterfaceGraph.h"
 #include "Misc/UObjectToken.h"
 
 #include UE_INLINE_GENERATED_CPP_BY_NAME(OptimusComputeGraph)
@@ -153,4 +154,17 @@ EMeshDeformerOutputBuffer UOptimusComputeGraph::GetOutputBuffers() const
 	return Result;
 }
 
+UOptimusGraphDataInterface* UOptimusComputeGraph::GetGraphDataInterfaceForPostLoadFixUp()
+{
+	for (UComputeDataInterface* DataInterface : DataInterfaces)
+	{
+		if (UOptimusGraphDataInterface* GraphDataInterface = Cast<UOptimusGraphDataInterface>(DataInterface))
+		{
+			// There should be only one graph data interface
+			return GraphDataInterface;
+		}
+	}
+	return nullptr;
+}
+
 #undef LOCTEXT_NAMESPACE
diff --git a/Engine/Plugins/Animation/DeformerGraph/Source/OptimusCore/Private/OptimusComputeGraph.h b/Engine/Plugins/Animation/DeformerGraph/Source/OptimusCore/Private/OptimusComputeGraph.h
--- a/Engine/Plugins/Animation/DeformerGraph/Source/OptimusCore/Private/OptimusComputeGraph.h
+++ b/Engine/Plugins/Animation/DeformerGraph/Source/OptimusCore/Private/OptimusComputeGraph.h
@@ -10,6 +10,7 @@
 enum class EMeshDeformerOutputBuffer : uint8;
 class UOptimusDeformer;
 class UOptimusNode;
+class UOptimusGraphDataInterface;
 
 UCLASS()
 class UOptimusComputeGraph :
@@ -31,5 +32,7 @@ protected:
 	UPROPERTY()
 	TArray<TWeakObjectPtr<const UOptimusNode>> KernelToNode;
 
+	UOptimusGraphDataInterface* GetGraphDataInterfaceForPostLoadFixUp();
+	
 	friend class UOptimusDeformer;
 };
diff --git a/Engine/Plugins/Animation/DeformerGraph/Source/OptimusCore/Private/OptimusDataTypeRegistry.cpp b/Engine/Plugins/Animation/DeformerGraph/Source/OptimusCore/Private/OptimusDataTypeRegistry.cpp
--- a/Engine/Plugins/Animation/DeformerGraph/Source/OptimusCore/Private/OptimusDataTypeRegistry.cpp
+++ b/Engine/Plugins/Animation/DeformerGraph/Source/OptimusCore/Private/OptimusDataTypeRegistry.cpp
@@ -475,7 +475,19 @@ void FOptimusDataTypeRegistry::RegisterBuiltinTypes()
 		EOptimusDataTypeUsageFlags::PinType
 		);
 
-	
+	// Name
+	Registry.RegisterType(
+		*FNameProperty::StaticClass(),
+		FText::FromString(TEXT("Name")),
+		FShaderValueTypeHandle(),
+		[](FFieldVariant InScope, FName InName) {
+			FNameProperty* Property = new FNameProperty(InScope, InName, RF_Public);
+			return Property;
+		},
+		{},
+		FName(TEXT("name")), {},
+		EOptimusDataTypeUsageFlags::Property | EOptimusDataTypeUsageFlags::PinType);
+
 	Registry.TypeWithAtomicSupport = {FIntProperty::StaticClass()->GetFName()};
 	
 	// Scan available built-in types to see if we can create array type for them
@@ -889,7 +901,7 @@ bool FOptimusDataTypeRegistry::RegisterArrayTypeIfApplicable(FOptimusDataTypeHan
 {
 	check(InElementDataType.IsValid());
 	// For now only allow array type for variables
-	if (!EnumHasAnyFlags(InElementDataType->UsageFlags, EOptimusDataTypeUsageFlags::Variable))
+	if (!EnumHasAnyFlags(InElementDataType->UsageFlags, EOptimusDataTypeUsageFlags::Variable | EOptimusDataTypeUsageFlags::Property))
 	{
 		return false;
 	}
@@ -908,59 +920,71 @@ bool FOptimusDataTypeRegistry::RegisterArrayTypeIfApplicable(FOptimusDataTypeHan
 		
 		return ArrayProperty;
 	};
-
-	// Making sure we are copying property value of a type into shader value of a equal or larger type, see comment for GetArrayElementDataTypeForStructuredBuffer
-	FOptimusDataTypeHandle InnerDataTypeForStructuredBuffer = GetArrayElementDataTypeForStructuredBuffer(InElementDataType);
-	// This should never happen theoretically, but somehow Linux complained about it, hope this helps the next time it complains
-	if (!ensureMsgf(InnerDataTypeForStructuredBuffer.IsValid(), TEXT("Cannot find matching element type for array type registration: %s"), *(InElementDataType->TypeName.ToString())))
-	{
-		return false;
-	}
-	check(InnerDataTypeForStructuredBuffer->ShaderValueSize >= InElementDataType->ShaderValueSize);
-	PropertyValueConvertFuncT ElementPropertyValueConvertFunc = FindPropertyValueConvertFunc(InnerDataTypeForStructuredBuffer->TypeName);
-	check(ElementPropertyValueConvertFunc);
-	int32 ElementShaderValueSize = InnerDataTypeForStructuredBuffer->ShaderValueSize;
-
-	// Nested array is not supported at the moment for array variables
-	check(InnerDataTypeForStructuredBuffer->GetNumArrays() == 0);
 	
-	PropertyValueConvertFuncT ArrayPropertyValueConvertFunc = [
-		ArrayPropertyCreateFunc,
-		ElementShaderValueSize,
-		ElementPropertyValueConvertFunc
-		] (
-		TArrayView<const uint8> InRawValue,
-		FShaderValueContainerView OutShaderValue
-		) -> bool
+	PropertyValueConvertFuncT ArrayPropertyValueConvertFunc;
+	TArray<FArrayMetadata> ArrayMetadata;
+	FShaderValueTypeHandle ArrayShaderType;
+	if (InElementDataType->ShaderValueType.IsValid())
 	{
-		const TUniquePtr<FArrayProperty> LocalArrayProperty(CastField<FArrayProperty>(ArrayPropertyCreateFunc(nullptr, NAME_None)));
-
-		const FArrayProperty* ArrayProperty = LocalArrayProperty.Get();
-		// Convert each element, store them in a separate buffer that is to be uploaded
-		FScriptArrayHelper ArrayHelper(ArrayProperty, InRawValue.GetData());
-		TArray<uint8>& Buffer = OutShaderValue.ArrayList[0].ArrayOfValues;
-		Buffer.AddZeroed(ElementShaderValueSize * ArrayHelper.Num());
-				
-		for (int32 Index = 0; Index < ArrayHelper.Num(); Index++)
+		// Making sure we are copying property value of a type into shader value of a equal or larger type, see comment for GetArrayElementDataTypeForStructuredBuffer
+		FOptimusDataTypeHandle InnerDataTypeForStructuredBuffer = GetArrayElementDataTypeForStructuredBuffer(InElementDataType);
+		// This should never happen theoretically, but somehow Linux complained about it, hope this helps the next time it complains
+		if (!ensureMsgf(InnerDataTypeForStructuredBuffer.IsValid(), TEXT("Cannot find matching element type for array type registration: %s"), *(InElementDataType->TypeName.ToString())))
 		{
-			uint8* ElementPtr = ArrayHelper.GetRawPtr(Index);
-			uint8* ShaderValuePtr = Buffer.GetData() + ElementShaderValueSize * Index;
+			return false;
+		}
+		check(InnerDataTypeForStructuredBuffer->ShaderValueSize >= InElementDataType->ShaderValueSize);
+		
+		ArrayShaderType = FShaderValueType::MakeDynamicArrayType(InnerDataTypeForStructuredBuffer->ShaderValueType);
+		
+		PropertyValueConvertFuncT ElementPropertyValueConvertFunc = FindPropertyValueConvertFunc(InnerDataTypeForStructuredBuffer->TypeName);
+		check(ElementPropertyValueConvertFunc);
+		int32 ElementShaderValueSize = InnerDataTypeForStructuredBuffer->ShaderValueSize;
 
-			// Nested buffer is not possible
-			TArray<FArrayShaderValue> DummyBufferList;
+		// Nested array is not supported at the moment for array variables
+		check(InnerDataTypeForStructuredBuffer->GetNumArrays() == 0);
+		
+		ArrayPropertyValueConvertFunc= [
+			ArrayPropertyCreateFunc,
+			ElementShaderValueSize,
+			ElementPropertyValueConvertFunc
+			] (
+			TArrayView<const uint8> InRawValue,
+			FShaderValueContainerView OutShaderValue
+			) -> bool
+		{
+			const TUniquePtr<FArrayProperty> LocalArrayProperty(CastField<FArrayProperty>(ArrayPropertyCreateFunc(nullptr, NAME_None)));
 
-			if (!ElementPropertyValueConvertFunc(
-					{ElementPtr, ArrayProperty->Inner->GetSize()},
-					{
-						{ShaderValuePtr, ElementShaderValueSize},
-						DummyBufferList
-					}))
+			const FArrayProperty* ArrayProperty = LocalArrayProperty.Get();
+			// Convert each element, store them in a separate buffer that is to be uploaded
+			FScriptArrayHelper ArrayHelper(ArrayProperty, InRawValue.GetData());
+			TArray<uint8>& Buffer = OutShaderValue.ArrayList[0].ArrayOfValues;
+			Buffer.AddZeroed(ElementShaderValueSize * ArrayHelper.Num());
+					
+			for (int32 Index = 0; Index < ArrayHelper.Num(); Index++)
 			{
-				return false;
-			}
-		}	
-		return true;
-	};
+				uint8* ElementPtr = ArrayHelper.GetRawPtr(Index);
+				uint8* ShaderValuePtr = Buffer.GetData() + ElementShaderValueSize * Index;
+
+				// Nested buffer is not possible
+				TArray<FArrayShaderValue> DummyBufferList;
+
+				if (!ElementPropertyValueConvertFunc(
+						{ElementPtr, ArrayProperty->Inner->GetSize()},
+						{
+							{ShaderValuePtr, ElementShaderValueSize},
+							DummyBufferList
+						}))
+				{
+					return false;
+				}
+			}	
+			return true;
+		};
+
+		ArrayMetadata = {{ElementShaderValueSize, 0}};
+	}
+	
 
 	TSharedRef<FOptimusDataType> ArrayDataType = MakeShared<FOptimusDataType>();
 	
@@ -969,14 +993,14 @@ bool FOptimusDataTypeRegistry::RegisterArrayTypeIfApplicable(FOptimusDataTypeHan
 		ArrayDataType,
 		ArrayPropertyCreateFunc,
 		ArrayPropertyValueConvertFunc,
-		{{ElementShaderValueSize, 0}}
+		ArrayMetadata
 	};
 
 	{
 		*ArrayDataType = *InElementDataType;
 		ArrayDataType->TypeName = GetArrayTypeName(InElementDataType->TypeName);
 		ArrayDataType->DisplayName = FText::FromString(InElementDataType->DisplayName.ToString() + TEXT(" Array"));
-		ArrayDataType->ShaderValueType = FShaderValueType::MakeDynamicArrayType(InnerDataTypeForStructuredBuffer->ShaderValueType);
+		ArrayDataType->ShaderValueType = ArrayShaderType;
 		
 		EnumRemoveFlags(ArrayDataType->UsageFlags, EOptimusDataTypeUsageFlags::Resource | EOptimusDataTypeUsageFlags::AnimAttributes | EOptimusDataTypeUsageFlags::PerBoneAnimAttribute);
 		
@@ -1016,7 +1040,7 @@ bool FOptimusDataTypeRegistry::RegisterType(
 		InDataType.TypeName = GetTypeName(InFieldType);
 		InDataType.DisplayName = InDisplayName;
 		InDataType.ShaderValueType = InShaderValueType;
-		InDataType.ShaderValueSize = InShaderValueType->GetResourceElementSize();
+		InDataType.ShaderValueSize = InShaderValueType.IsValid() ? InShaderValueType->GetResourceElementSize() : 0;
 		InDataType.TypeCategory = InPinCategory;
 		if (InPinColor.IsSet())
 		{
diff --git a/Engine/Plugins/Animation/DeformerGraph/Source/OptimusCore/Private/OptimusDeformer.cpp b/Engine/Plugins/Animation/DeformerGraph/Source/OptimusCore/Private/OptimusDeformer.cpp
--- a/Engine/Plugins/Animation/DeformerGraph/Source/OptimusCore/Private/OptimusDeformer.cpp
+++ b/Engine/Plugins/Animation/DeformerGraph/Source/OptimusCore/Private/OptimusDeformer.cpp
@@ -228,9 +228,9 @@ UOptimusVariableDescription* UOptimusDeformer::AddVariable(
 		InDataTypeRef.Set(FOptimusDataTypeRegistry::Get().FindType(*FDoubleProperty::StaticClass()));
 	}
 
-	// Is this data type compatible with resources?
+	// Is this data type compatible with variables?
 	FOptimusDataTypeHandle DataType = InDataTypeRef.Resolve();
-	if (!DataType.IsValid() || !EnumHasAnyFlags(DataType->UsageFlags, EOptimusDataTypeUsageFlags::Variable))
+	if (!DataType.IsValid() || !EnumHasAnyFlags(DataType->UsageFlags, EOptimusDataTypeUsageFlags::Variable | EOptimusDataTypeUsageFlags::Property))
 	{
 		UE_LOG(LogOptimusCore, Error, TEXT("Invalid data type for variables."));
 		return nullptr;
@@ -1497,7 +1497,14 @@ bool UOptimusDeformer::Compile()
 		return false;
 	}
 
-	ComputeGraphs.Reset();
+	auto ClearCompiledData = [&]()
+	{
+		ComputeGraphs.Reset();
+		DataInterfacePropertyOverrideMap.Reset();
+		ValueMap.Reset();
+	};
+
+	ClearCompiledData();
 	
 	CompileBeginDelegate.Broadcast(this);
 	
@@ -1524,8 +1531,37 @@ bool UOptimusDeformer::Compile()
 	{
 		if (Graph->GetGraphType() != EOptimusNodeGraphType::Function)
 		{
-			TArray<FOptimusComputeGraphInfo> ComputeGraphInfos = CompileNodeGraphToComputeGraphs(Graph, ErrorReporter);
-			ComputeGraphs.Append(ComputeGraphInfos);
+			FOptimusNodeGraphCompilationResult Result = CompileNodeGraphToComputeGraphs(Graph, ErrorReporter);
+			ComputeGraphs.Append(Result.ComputeGraphInfos);
+
+			DataInterfacePropertyOverrideMap.Append(Result.DataInterfacePropertyOverrideMap);
+
+			// Merge Value Maps
+			for (const TPair<FOptimusValueIdentifier, FOptimusValueDescription>& Pair : Result.ValueMap)
+			{
+				const FOptimusValueIdentifier& ValueId = Pair.Key;
+				const FOptimusValueDescription& ValueDescription = Pair.Value;
+				
+				if (FOptimusValueDescription* ExistingDescription = ValueMap.Find(ValueId))
+				{
+					if (EnumHasAnyFlags(ValueDescription.ValueUsage, EOptimusValueUsage::CPU) &&
+						(!EnumHasAnyFlags(ExistingDescription->ValueUsage, EOptimusValueUsage::CPU)))
+					{
+						ExistingDescription->ValueUsage |= EOptimusValueUsage::CPU;
+						ExistingDescription->Value = ValueDescription.Value;
+					}
+					else if (EnumHasAnyFlags(ValueDescription.ValueUsage, EOptimusValueUsage::GPU) &&
+						(!EnumHasAnyFlags(ExistingDescription->ValueUsage, EOptimusValueUsage::GPU)))
+					{
+						ExistingDescription->ValueUsage |= EOptimusValueUsage::GPU;
+						ExistingDescription->ShaderValue = ValueDescription.ShaderValue;
+					}
+				}
+				else
+				{
+					ValueMap.Add(Pair);
+				}
+			}
 		}
 	}
 
@@ -1533,7 +1569,7 @@ bool UOptimusDeformer::Compile()
 
 	if (Status == EOptimusDeformerStatus::HasErrors)
 	{
-		ComputeGraphs.Reset();
+		ClearCompiledData();
 		return false;
 	}
 	
@@ -1617,7 +1653,7 @@ struct FOptimusInstancedPin
 	}
 };
 
-TArray<FOptimusComputeGraphInfo> UOptimusDeformer::CompileNodeGraphToComputeGraphs(
+FOptimusNodeGraphCompilationResult UOptimusDeformer::CompileNodeGraphToComputeGraphs(
 	const UOptimusNodeGraph* InNodeGraph,
 	TFunction<void(EOptimusDiagnosticLevel, FText, const UObject*)> InErrorReporter
 	)
@@ -2091,7 +2127,14 @@ TArray<FOptimusComputeGraphInfo> UOptimusDeformer::CompileNodeGraphToComputeGrap
 				}
 				else if (Cast<const IOptimusValueProvider>(SourceNode))
 				{
-					LinksToInsertCopyKernel.FindOrAdd(SourcePin).Add(TargetPin);
+					if (const IOptimusPropertyPinProvider* PropertyPinProvider = Cast<const IOptimusPropertyPinProvider>(TargetNode))
+					{
+						// Property pins are CPU only, no need to involve a kernel
+						if (!PropertyPinProvider->GetPropertyPins().Contains(TargetPin.Pin))
+						{
+							LinksToInsertCopyKernel.FindOrAdd(SourcePin).Add(TargetPin);
+						}
+					}
 				}
 			}
 		}
@@ -2105,17 +2148,19 @@ TArray<FOptimusComputeGraphInfo> UOptimusDeformer::CompileNodeGraphToComputeGrap
 	TMap<const UComputeDataInterface*, int32> DataInterfaceToBindingIndexMap;
 	
 	// Find all data interface nodes and create their data interfaces.
-	FOptimus_NodeToDataInterfaceMap NodeDataInterfaceMap;
+	TMap<FOptimusRoutedConstNode, UOptimusComputeDataInterface*> NodeDataInterfaceMap;
 	
 	// Find all value nodes (constant and variable) 
 	TArray<const UOptimusNode *> ValueNodes;
-
+	TMap<const UOptimusNode*, EOptimusValueUsage> ValueNodeUsageMap;
+	
 	TMap<FOptimusRoutedConstNode, TArray<UOptimusComputeDataInterface*>> LoopEntryToLoopDataInterfaces;
 	
 	for (const FOptimusRoutedConstNode& ConnectedNode : ConnectedNodes)
 	{
 		if (const IOptimusDataInterfaceProvider* NodeDataInterfaceProvider = Cast<const IOptimusDataInterfaceProvider>(ConnectedNode.Node))
 		{
+			// Gets a copy of node's data interface
 			UOptimusComputeDataInterface* DataInterface = NodeDataInterfaceProvider->GetDataInterface(this);
 			if (!DataInterface)
 			{
@@ -2123,16 +2168,47 @@ TArray<FOptimusComputeGraphInfo> UOptimusDeformer::CompileNodeGraphToComputeGrap
 				return {};
 			}
 
-			NodeDataInterfaceMap.Add(ConnectedNode.Node, DataInterface);
+			NodeDataInterfaceMap.Add(ConnectedNode, DataInterface);
 			DataInterfaceToBindingIndexMap.Add(DataInterface) = NodeDataInterfaceProvider->GetComponentBinding(ConnectedNode.TraversalContext)->GetIndex();
 		}
 		else if (Cast<const IOptimusValueProvider>(ConnectedNode.Node))
 		{
-			// Ideally we shouldn't hit this ensure, but in case we do, it is safe to proceed as we have this If check here
-			// However, the cause of the ensure should be investigated
-			if (ensure(!ValueNodes.Contains(ConnectedNode.Node)))
+			const UOptimusNode* Node = ConnectedNode.Node;
+			EOptimusValueUsage ValueUsage = EOptimusValueUsage::None;
+			
+			TArray<UOptimusNodePin*> OutputPins = Node->GetPinsByDirection(EOptimusNodePinDirection::Output, true);
+			check(OutputPins.Num() == 1);
+
+			const UOptimusNodePin* OutputPin = OutputPins[0];
+			TArray<FOptimusRoutedNodePin> OtherPins = OutputPin->GetConnectedPinsWithRouting(ConnectedNode.TraversalContext);
+
+			for (const FOptimusRoutedNodePin& RoutedOtherPin : OtherPins)
 			{
-				ValueNodes.Add(ConnectedNode.Node);
+				const UOptimusNodePin* OtherPin = RoutedOtherPin.NodePin;
+				const UOptimusNode* OtherNode = OtherPin->GetOwningNode();
+				if (const IOptimusPropertyPinProvider* PropertyPinProvider = Cast<const IOptimusPropertyPinProvider>(OtherNode))
+				{
+					if (PropertyPinProvider->GetPropertyPins().Contains(OtherPin))
+					{
+						// At least one connection requesting value on CPU
+						ValueUsage |= EOptimusValueUsage::CPU;
+						continue;
+					}
+				}
+
+				// at least one connection requesting this value on GPU
+				ValueUsage |= EOptimusValueUsage::GPU;
+				break;
+			}
+
+			EOptimusValueUsage& ValueUsageRef = ValueNodeUsageMap.FindOrAdd(Node);
+			ValueUsageRef |= ValueUsage;
+
+			
+			// Node may get used in a function that has multiple instances in this graph
+			if (!ValueNodes.Contains(Node))
+			{
+				ValueNodes.Add(Node);
 			}
 		}
 		else if (const UOptimusNode_LoopTerminal* LoopTerminal = Cast<const UOptimusNode_LoopTerminal>(ConnectedNode.Node))
@@ -2153,36 +2229,56 @@ TArray<FOptimusComputeGraphInfo> UOptimusDeformer::CompileNodeGraphToComputeGrap
 		}
 	}
 
+	TMap<FOptimusValueIdentifier, FOptimusValueDescription> NodeGraphValueMap;
+
 	// Create the graph data interface and fill it with the value nodes.
 	UOptimusGraphDataInterface* GraphDataInterface = NewObject<UOptimusGraphDataInterface>(this);
-
+	DataInterfaceToBindingIndexMap.Add(GraphDataInterface) = GraphDataComponentBinding->GetIndex();
+	
 	TArray<FOptimusGraphVariableDescription> ValueNodeDescriptions;
-	ValueNodeDescriptions.Reserve(ValueNodes.Num());
-	for (int32 ValueNodeIndex = 0 ; ValueNodeIndex < ValueNodes.Num(); ValueNodeIndex++)
+	ValueNodeDescriptions.Reserve(ValueNodes.Num());	
+
+	for (const UOptimusNode* ValueNode : ValueNodes)
 	{
-		UOptimusNode const* ValueNode = ValueNodes[ValueNodeIndex];
-		if (IOptimusValueProvider const* ValueProvider = Cast<const IOptimusValueProvider>(ValueNode))
+		const IOptimusValueProvider* ValueProvider = CastChecked<IOptimusValueProvider>(ValueNode);
+		FOptimusValueIdentifier ValueId = ValueProvider->GetValueIdentifier();
+		EOptimusValueUsage ValueUsage = ValueNodeUsageMap[ValueNode];
+
+		FOptimusValueDescription Description;
+		Description.DataType = ValueProvider->GetValueDataType();
+		Description.ValueUsage = ValueUsage;
+
+		if (EnumHasAnyFlags(Description.ValueUsage, EOptimusValueUsage::CPU))
+		{
+			Description.Value = ValueProvider->GetValue();
+		}
+		else if (EnumHasAnyFlags(Description.ValueUsage, EOptimusValueUsage::GPU))
 		{
-			FOptimusGraphVariableDescription& ValueNodeDescription = ValueNodeDescriptions.AddDefaulted_GetRef();
-			ValueNodeDescription.Name = Optimus::MakeUniqueValueName(ValueProvider->GetValueName(), ValueNodeIndex);
-			ValueNodeDescription.ValueType = ValueProvider->GetValueType()->ShaderValueType;
-			ValueNodeDescription.SourceObject = const_cast<UOptimusNode*>(ValueNode);
+			Description.ShaderValue = ValueProvider->GetValue().GetShaderValue(ValueProvider->GetValueDataType());
+		}
+		
+		NodeGraphValueMap.Add(ValueId, MoveTemp(Description));	
 
-			if (UOptimusNode_ConstantValue const* ConstantNode = Cast<const UOptimusNode_ConstantValue>(ValueNode))
-			{
-				ValueNodeDescription.ShaderValue = ConstantNode->GetShaderValue();
-			}
+		if (EnumHasAnyFlags(ValueUsage, EOptimusValueUsage::GPU))
+		{
+			FOptimusGraphVariableDescription ValueNodeDescription;;
+			ValueNodeDescription.ValueId = ValueProvider->GetValueIdentifier();
+			ValueNodeDescription.Name = ValueNodeDescription.ValueId.Name.ToString();
+			ValueNodeDescription.ValueType = ValueProvider->GetValueDataType()->ShaderValueType;
+
+			ValueNodeDescriptions.Add(MoveTemp(ValueNodeDescription));
 		}
 	}
 	GraphDataInterface->Init(ValueNodeDescriptions);
 
-	DataInterfaceToBindingIndexMap.Add(GraphDataInterface) = GraphDataComponentBinding->GetIndex();
 	
 	TMap<FOptimusInstancedNode, UComputeDataInterface*> KernelDataInterfaceMap;
 	TMap<FOptimusInstancedNode, FOptimus_KernelInputMap> KernelInputMap;
 	TMap<FOptimusInstancedNode, FOptimus_KernelOutputMap> KernelOutputMap;
 	TMap<FOptimusInstancedPin, UOptimusComputeDataInterface*> KernelOutputDataInterfaceMap;
 
+	TMap<TWeakObjectPtr<const UComputeDataInterface>, FOptimusDataInterfacePropertyOverrideInfo> NodeGraphDataInterfacePropertyOverrideMap;
+	
 	for (const FOptimusInstancedNode& InstancedNode : InstancedNodes )
 	{
 		const FOptimusRoutedConstNode& RoutedNode = InstancedNode.RoutedNode;
@@ -2228,7 +2324,7 @@ TArray<FOptimusComputeGraphInfo> UOptimusDeformer::CompileNodeGraphToComputeGrap
 					{
 						KernelInputMap[InstancedNode].Add(Pin) = {GraphDataInterface, SourcePin};
 					}
-					else if (UOptimusComputeDataInterface** NodeDataInterface = NodeDataInterfaceMap.Find(SourcePin->GetOwningNode()))
+					else if (UOptimusComputeDataInterface** NodeDataInterface = NodeDataInterfaceMap.Find(SourceRoutedNode))
 					{
 						KernelInputMap[InstancedNode].Add(Pin) = {*NodeDataInterface, SourcePin};
 					}
@@ -2325,7 +2421,6 @@ TArray<FOptimusComputeGraphInfo> UOptimusDeformer::CompileNodeGraphToComputeGrap
 					for (const FOptimusInstancedPin& TargetInstancedPin : TargetDataInterfacePins)
 					{
 						const FOptimusRoutedConstNode TargetRoutedNode = TargetInstancedPin.InstancedNode.RoutedNode;
-						const UOptimusNode* TargetNode = TargetRoutedNode.Node;
 						const UOptimusNodePin* TargetPin = TargetInstancedPin.Pin;
 
 						if (bShouldCopyToDataInterface)
@@ -2335,7 +2430,7 @@ TArray<FOptimusComputeGraphInfo> UOptimusDeformer::CompileNodeGraphToComputeGrap
 						}
 						else
 						{
-							if (UOptimusComputeDataInterface** NodeDataInterface = NodeDataInterfaceMap.Find(TargetNode))
+							if (UOptimusComputeDataInterface** NodeDataInterface = NodeDataInterfaceMap.Find(TargetRoutedNode))
 							{
 								KernelOutputMap[InstancedNode].FindOrAdd(Pin).Add({*NodeDataInterface, TargetPin});
 							}
@@ -2344,6 +2439,40 @@ TArray<FOptimusComputeGraphInfo> UOptimusDeformer::CompileNodeGraphToComputeGrap
 				}
 			}
 		}
+		
+		if (const IOptimusPropertyPinProvider* PropertyPinProvider = Cast<const IOptimusPropertyPinProvider>(Node))
+		{
+			const UComputeDataInterface* ProviderDataInterface = nullptr;
+			if (Cast<const IOptimusDataInterfaceProvider>(Node))
+			{
+				ProviderDataInterface = NodeDataInterfaceMap[RoutedNode];
+			}
+			else
+			{
+				// To be implemented when we have actual use cases
+				check(false);
+			}
+		
+			if (ensure(ProviderDataInterface))
+			{
+				for (UOptimusNodePin* Pin : PropertyPinProvider->GetPropertyPins())
+				{
+					FOptimusInstancedPin InstancedPin = {InstancedNode, Pin};
+					if (FOptimusInstancedPin* SourceInstancedPin = TargetPinToSourcePin.Find(InstancedPin))
+					{
+						const UOptimusNode* SourceNode = SourceInstancedPin->Pin->GetOwningNode();
+
+						if (const IOptimusValueProvider* ValueProvider = Cast<IOptimusValueProvider>(SourceNode))
+						{
+							NodeGraphDataInterfacePropertyOverrideMap.FindOrAdd(ProviderDataInterface)
+									.PinNameToValueIdMap
+									.Add(Pin->GetFName(),
+										ValueProvider->GetValueIdentifier());	
+						}
+					}
+				}
+			}
+		}
 	}
 
 	if (!LinksToInsertCopyKernel.IsEmpty())
@@ -2367,18 +2496,19 @@ TArray<FOptimusComputeGraphInfo> UOptimusDeformer::CompileNodeGraphToComputeGrap
 		const TArray<FOptimusInstancedPin>& TargetInstancedPins = OutputLink.Value;
 		
 		const UOptimusNode* SourceNode = SourceInstancedPin.Pin->GetOwningNode();
+		const FOptimusRoutedConstNode& SourceRoutedNode = SourceInstancedPin.InstancedNode.RoutedNode;
 		if (const IOptimusDataInterfaceProvider* InterfaceProvider = Cast<const IOptimusDataInterfaceProvider>(SourceNode))
 		{
 			FDataInterfaceFunctionBinding DataInterfaceBinding;
-			DataInterfaceBinding.DataInterface = NodeDataInterfaceMap[SourceNode];
+			DataInterfaceBinding.DataInterface = NodeDataInterfaceMap[SourceRoutedNode];
 			DataInterfaceBinding.FunctionIndex = InterfaceProvider->GetDataFunctionIndexFromPin(SourceInstancedPin.Pin);
 			CopyFromDataInterfaceMap.Add(SourceInstancedPin) = DataInterfaceBinding;
 		}
-		else if (Cast<const IOptimusValueProvider>(SourceNode))
+		else if (const IOptimusValueProvider* ValueProvider = Cast<const IOptimusValueProvider>(SourceNode))
 		{
 			FDataInterfaceFunctionBinding DataInterfaceBinding;
 			DataInterfaceBinding.DataInterface = GraphDataInterface;
-			DataInterfaceBinding.FunctionIndex = ValueNodes.Find(SourceNode);
+			DataInterfaceBinding.FunctionIndex = GraphDataInterface->FindFunctionIndex(ValueProvider->GetValueIdentifier());
 			CopyFromDataInterfaceMap.Add(SourceInstancedPin) = DataInterfaceBinding;	
 		}
 		else if (Cast<const IOptimusComputeKernelProvider>(SourceNode))
@@ -2394,12 +2524,13 @@ TArray<FOptimusComputeGraphInfo> UOptimusDeformer::CompileNodeGraphToComputeGrap
 		for (const FOptimusInstancedPin& TargetInstancedPin : TargetInstancedPins)
 		{
 			const UOptimusNode* TargetNode = TargetInstancedPin.Pin->GetOwningNode();
+			const FOptimusRoutedConstNode& TargetRoutedNode = TargetInstancedPin.InstancedNode.RoutedNode;
 
 			if (const IOptimusDataInterfaceProvider* InterfaceProvider = Cast<const IOptimusDataInterfaceProvider>(TargetNode);
 				ensure(InterfaceProvider))
 			{
 				// One-time Initialization of the copy kernel based on the first target pin, because if source is a value provider, it does not
-				// have a meaningful data domain and a meaning component source binding
+				// have a meaningful data domain and a meaningful component source binding
 				if (!bIsCopyKernelDataInterfaceCreated)
 				{
 					bIsCopyKernelDataInterfaceCreated = true;
@@ -2415,15 +2546,17 @@ TArray<FOptimusComputeGraphInfo> UOptimusDeformer::CompileNodeGraphToComputeGrap
 				
 				
 				FDataInterfaceFunctionBinding DataInterfaceBinding;
-				DataInterfaceBinding.DataInterface = NodeDataInterfaceMap[TargetNode];
+				DataInterfaceBinding.DataInterface = NodeDataInterfaceMap[TargetRoutedNode];
 				DataInterfaceBinding.FunctionIndex = InterfaceProvider->GetDataFunctionIndexFromPin(TargetInstancedPin.Pin);
 				CopyToDataInterfaceMap.Add(TargetInstancedPin) = DataInterfaceBinding;
 			}
 		}
 	}
 
-
-	TArray<FOptimusComputeGraphInfo> GraphInfos;
+	FOptimusNodeGraphCompilationResult Result;
+	Result.DataInterfacePropertyOverrideMap = NodeGraphDataInterfacePropertyOverrideMap;
+	Result.ValueMap = NodeGraphValueMap;
+	TArray<FOptimusComputeGraphInfo>& GraphInfos = Result.ComputeGraphInfos;
 	for (EOptimusNodeGraphType GraphType : GraphTypes)
 	{
 		FString Name = InNodeGraph->GetName();
@@ -2470,9 +2603,6 @@ TArray<FOptimusComputeGraphInfo> UOptimusDeformer::CompileNodeGraphToComputeGrap
 		}
 
 		// Now that we've collected all the pieces, time to line them up.
-		ComputeGraph->DataInterfaces.Add(GraphDataInterface);
-		ComputeGraph->DataInterfaceToBinding.Add(0);		// Graph data interface always uses the primary binding.
-
 		for (const FOptimusInstancedNode& InstancedNode : InstancedNodes )
 		{
 			const FOptimusRoutedConstNode& ConnectedNode = InstancedNode.RoutedNode;
@@ -2586,7 +2716,7 @@ TArray<FOptimusComputeGraphInfo> UOptimusDeformer::CompileNodeGraphToComputeGrap
 				
 				FOptimus_ComputeKernelResult KernelSourceResult = KernelProvider->CreateComputeKernel(	
 					BoundKernel.Kernel, ConnectedNode.TraversalContext,
-					KernelInputs, KernelOutputs, ValueNodes,
+					KernelInputs, KernelOutputs,
 					KernelDataInterface,
 					BoundKernel.InputDataBindings, BoundKernel.OutputDataBindings
 				);
@@ -2871,7 +3001,7 @@ TArray<FOptimusComputeGraphInfo> UOptimusDeformer::CompileNodeGraphToComputeGrap
 	
 #endif
 
-	return GraphInfos;
+	return Result;
 }
 
 void UOptimusDeformer::OnDataTypeChanged(FName InTypeName)
@@ -3068,7 +3198,7 @@ void UOptimusDeformer::Notify(EOptimusGlobalNotifyType InNotifyType, UObject* In
 	case EOptimusGlobalNotifyType::ConstantValueChanged:
 		if (UOptimusNode_ConstantValue* ConstantValue = Cast<UOptimusNode_ConstantValue>(InObject))
 		{
-			ConstantValueUpdateDelegate.Broadcast(ConstantValue, ConstantValue->GetShaderValue());
+			ConstantValueUpdateDelegate.Broadcast(ConstantValue, ConstantValue->GetValue());
 		}
 		
 		break;
@@ -3138,6 +3268,11 @@ void UOptimusDeformer::PostLoad()
 	{
 		Graph->ConditionalPostLoad();
 	}
+
+	for (FOptimusComputeGraphInfo& Info : ComputeGraphs)
+	{
+		Info.ComputeGraph->ConditionalPostLoad();
+	}
 	
 	// Fixup any empty array entries.
 	Resources->Descriptions.RemoveAllSwap([](const TObjectPtr<UOptimusResourceDescription>& Value) { return Value == nullptr; });
@@ -3232,6 +3367,11 @@ void UOptimusDeformer::PostLoad()
 		PostLoadRemoveDeprecatedValueContainerGeneratorClass();
 	}
 
+	if (GetLinkerCustomVersion(FOptimusObjectVersion::GUID) < FOptimusObjectVersion::PropertyPinSupport)
+	{
+		PostLoadMoveValueFromGraphDataInterfaceToDeformerValueMap();
+	}
+	
 	// If the graph was saved at any previous version, and was clean, mark the status now as modified.
 	if (Status == EOptimusDeformerStatus::Compiled &&
 		GetLinkerCustomVersion(FOptimusObjectVersion::GUID) < FOptimusObjectVersion::LatestVersion)
@@ -3456,6 +3596,72 @@ void UOptimusDeformer::PostLoadRemoveDeprecatedValueContainerGeneratorClass()
 }
 
 
+void UOptimusDeformer::PostLoadMoveValueFromGraphDataInterfaceToDeformerValueMap()
+{
+PRAGMA_DISABLE_DEPRECATION_WARNINGS
+	int ConstantUniqueSuffix = 0;
+	for (FOptimusComputeGraphInfo& Info : ComputeGraphs)
+	{
+		if (UOptimusGraphDataInterface* GraphDataInterface = Info.ComputeGraph->GetGraphDataInterfaceForPostLoadFixUp())
+		{
+			for (FOptimusGraphVariableDescription& OldDescription : GraphDataInterface->Variables)
+			{
+				if (OldDescription.ShaderValue_DEPRECATED.IsValid())
+				{
+					FOptimusValueDescription NewDescription;
+					NewDescription.ValueUsage = EOptimusValueUsage::GPU;
+					NewDescription.DataType = FOptimusDataTypeRegistry::Get().FindType(OldDescription.ValueType);
+					NewDescription.ShaderValue = OldDescription.ShaderValue_DEPRECATED;
+					// Make sure the Id is unique across all graphs
+					// We could do better here if needed by finding the source constant node of this value, but settling for the simpler solution for now
+					// The only thing that would break is scrubbing constant node value without recompiling
+					FOptimusValueIdentifier ValueId = {EOptimusValueType::Constant, *(OldDescription.Name + TEXT("_PostLoadFixUp_") + FString::FromInt(ConstantUniqueSuffix))};
+					ConstantUniqueSuffix++;
+					ValueMap.Add(ValueId, MoveTemp(NewDescription));
+					
+					OldDescription.ValueId = ValueId;
+				}
+				else
+				{
+					FName VariableName = NAME_None;
+					// When source object was introduced, we also appended a unique index to the value name provided by each value provider
+					// so instead of using the name directly, we need to do this extra step
+					if (OldDescription.SourceObject_DEPRECATED.IsNull())
+					{
+						VariableName = *OldDescription.Name;
+					}
+					else
+					{
+						VariableName = *(Optimus::ExtractSourceValueName(OldDescription.Name));
+					}
+
+					FOptimusValueIdentifier ValueId = {EOptimusValueType::Variable, VariableName};
+					if (!ValueMap.Contains(ValueId))
+					{
+						for (UOptimusVariableDescription* VariableDescription : GetVariables())
+						{
+							if (VariableDescription->VariableName == VariableName)
+							{
+								FOptimusValueDescription NewDescription;
+								NewDescription.ValueUsage = EOptimusValueUsage::GPU;
+								NewDescription.DataType = FOptimusDataTypeRegistry::Get().FindType(OldDescription.ValueType);
+								NewDescription.ShaderValue = VariableDescription->DefaultValueStruct.GetShaderValue(VariableDescription->DataType);
+								ValueMap.Add(ValueId, MoveTemp(NewDescription));	
+								break;
+							}
+						}
+					}
+					
+					OldDescription.ValueId = ValueId;
+					
+				}
+			}
+		}
+	}
+PRAGMA_ENABLE_DEPRECATION_WARNINGS
+}
+
+
 UOptimusComponentSourceBinding* UOptimusDeformer::FindCompatibleBindingWithInterface(
 	const UOptimusComputeDataInterface* InDataInterface) const
 {
diff --git a/Engine/Plugins/Animation/DeformerGraph/Source/OptimusCore/Private/OptimusDeformerInstance.cpp b/Engine/Plugins/Animation/DeformerGraph/Source/OptimusCore/Private/OptimusDeformerInstance.cpp
--- a/Engine/Plugins/Animation/DeformerGraph/Source/OptimusCore/Private/OptimusDeformerInstance.cpp
+++ b/Engine/Plugins/Animation/DeformerGraph/Source/OptimusCore/Private/OptimusDeformerInstance.cpp
@@ -6,6 +6,7 @@
 #include "ComputeFramework/ComputeFramework.h"
 #include "ComputeWorkerInterface.h"
 #include "IOptimusPersistentBufferProvider.h"
+#include "IOptimusValueProvider.h"
 #include "DataInterfaces/OptimusDataInterfaceGraph.h"
 #include "DataInterfaces/OptimusDataInterfaceRawBuffer.h"
 #include "OptimusComputeGraph.h"
@@ -380,6 +381,10 @@ void UOptimusDeformerInstance::SetupFromDeformer(UOptimusDeformer* InDeformer)
 	// Create the persistent buffer pool
 	BufferPool = MakeShared<FOptimusPersistentBufferPool>();
 	
+	// Create local storage for deformer graph constants/variables
+	ValueMap = InDeformer->ValueMap;
+	DataInterfacePropertyOverrideMap = InDeformer->DataInterfacePropertyOverrideMap;
+	
 	// (Re)Create and bind data providers.
 	ComputeGraphExecInfos.Reset();
 	GraphsToRunOnNextTick.Reset();
@@ -436,30 +441,6 @@ void UOptimusDeformerInstance::SetupFromDeformer(UOptimusDeformer* InDeformer)
 		}
 	}
 	
-	// Create local storage for deformer graph variables.
-	Variables = NewObject<UOptimusVariableContainer>(this);
-	Variables->Descriptions.Reserve(InDeformer->GetVariables().Num());
-	TSet<const UOptimusVariableDescription*> Visited;
-	for (const UOptimusVariableDescription* VariableDescription : InDeformer->GetVariables())
-	{
-		if (!VariableDescription)
-		{
-			continue;
-		}
-		if (Visited.Contains(VariableDescription))
-		{
-			continue;
-		}
-		Visited.Add(VariableDescription);
-		
-		UOptimusVariableDescription* VariableDescriptionCopy = NewObject<UOptimusVariableDescription>();
-		VariableDescriptionCopy->Guid = VariableDescription->Guid;
-		VariableDescriptionCopy->VariableName = VariableDescription->VariableName;
-		VariableDescriptionCopy->DataType = VariableDescription->DataType;
-		// No need to copy the default value, we directly copy the shader value of the default value below
-		VariableDescriptionCopy->CachedShaderValue = VariableDescription->DefaultValueStruct.GetShaderValue(VariableDescription->DataType);
-		Variables->Descriptions.Add(VariableDescriptionCopy);
-	}
 
 	if (UMeshComponent* Ptr = MeshComponent.Get())
 	{
@@ -480,6 +461,27 @@ void UOptimusDeformerInstance::SetCanBeActive(bool bInCanBeActive)
 	bCanBeActive = bInCanBeActive;
 }
 
+FOptimusValueContainerStruct UOptimusDeformerInstance::GetDataInterfacePropertyOverride(const UComputeDataInterface* DataInterface, FName PinName)
+{
+	if (const FOptimusDataInterfacePropertyOverrideInfo* OverrideInfoPtr = DataInterfacePropertyOverrideMap.Find(DataInterface))
+	{
+		if (const FOptimusValueIdentifier* Overrider = OverrideInfoPtr->PinNameToValueIdMap.Find(PinName))
+		{
+			return ValueMap[*Overrider].Value;
+		}
+	}
+
+	// No override
+	return {};
+}
+
+const FShaderValueContainer& UOptimusDeformerInstance::GetShaderValue(
+	const FOptimusValueIdentifier& InValueId
+	) const
+{
+	return ValueMap[InValueId].ShaderValue;
+}
+
 void UOptimusDeformerInstance::AllocateResources()
 {
 }
@@ -590,164 +592,185 @@ EMeshDeformerOutputBuffer UOptimusDeformerInstance::GetOutputBuffers() const
 namespace
 {
 	template <typename T>
-	bool SetVariableValue(UOptimusVariableContainer const* InVariables, FName InVariableName, FName InTypeName, T const& InValue)
+	bool SetValue(TMap<FOptimusValueIdentifier, FOptimusValueDescription>& InValueMap, const FOptimusValueIdentifier& InValueId, FName InTypeName, T const& InValue)
 	{
 		FOptimusDataTypeHandle WantedType = FOptimusDataTypeRegistry::Get().FindType(InTypeName);
-		for (UOptimusVariableDescription* VariableDesc : InVariables->Descriptions)
+
+		if (FOptimusValueDescription* Description = InValueMap.Find(InValueId))
 		{
-			if (VariableDesc->VariableName == InVariableName && VariableDesc->DataType == WantedType)
+			if (Description->DataType == WantedType)
 			{
 				TUniquePtr<FProperty> Property(WantedType->CreateProperty(nullptr, NAME_None));
 				if (ensure(Property->GetSize() == sizeof(T)))
 				{
 					const uint8* ValueBytes = reinterpret_cast<const uint8*>(&InValue);
-					WantedType->ConvertPropertyValueToShader(TArrayView<const uint8>(ValueBytes, sizeof(T)), VariableDesc->CachedShaderValue);
+					TArrayView<const uint8> ValueView(ValueBytes, sizeof(T));
+					if (Description->ValueUsage == EOptimusValueUsage::GPU)
+					{
+						WantedType->ConvertPropertyValueToShader(ValueView, Description->ShaderValue);
+					}
+					
+					if (Description->ValueUsage == EOptimusValueUsage::CPU)
+					{
+						Description->Value.SetValue(Description->DataType, ValueView);
+					}
 				}
-
-				return true;
+				
+				return true;	
 			}
 		}
 
 		return false;
 	}
+	
+	template <typename T>
+	bool SetVariableValue(TMap<FOptimusValueIdentifier, FOptimusValueDescription>& InValueMap, FName InVariableName, FName InTypeName, T const& InValue)
+	{
+		return SetValue(InValueMap, {EOptimusValueType::Variable, InVariableName}, InTypeName, InValue);
+	}
 }
 
 
 bool UOptimusDeformerInstance::SetBoolVariable(FName InVariableName, bool InValue)
 {
-	return SetVariableValue(Variables, InVariableName, FOptimusDataTypeRegistry::GetTypeName(*FBoolProperty::StaticClass()), InValue);
+	return SetVariableValue(ValueMap, InVariableName, FOptimusDataTypeRegistry::GetTypeName(*FBoolProperty::StaticClass()), InValue);
 }
 
 bool UOptimusDeformerInstance::SetIntVariable(FName InVariableName, int32 InValue)
 {
-	return SetVariableValue<int32>(Variables, InVariableName, FOptimusDataTypeRegistry::GetTypeName(*FIntProperty::StaticClass()), InValue);
+	return SetVariableValue<int32>(ValueMap, InVariableName, FOptimusDataTypeRegistry::GetTypeName(*FIntProperty::StaticClass()), InValue);
 }
 
 bool UOptimusDeformerInstance::SetIntArrayVariable(FName InVariableName, const TArray<int32>& InValue)
 {
-	return SetVariableValue<TArray<int32>>(Variables, InVariableName, FOptimusDataTypeRegistry::GetArrayTypeName(*FIntProperty::StaticClass()), InValue);
+	return SetVariableValue<TArray<int32>>(ValueMap, InVariableName, FOptimusDataTypeRegistry::GetArrayTypeName(*FIntProperty::StaticClass()), InValue);
 }
 
 bool UOptimusDeformerInstance::SetInt2Variable(FName InVariableName, const FIntPoint& InValue)
 {
-	return SetVariableValue<FIntPoint>(Variables, InVariableName, FOptimusDataTypeRegistry::GetTypeName(TBaseStructure<FIntPoint>::Get()), InValue);
+	return SetVariableValue<FIntPoint>(ValueMap, InVariableName, FOptimusDataTypeRegistry::GetTypeName(TBaseStructure<FIntPoint>::Get()), InValue);
 }
 
 bool UOptimusDeformerInstance::SetInt2ArrayVariable(FName InVariableName, const TArray<FIntPoint>& InValue)
 {
-	return SetVariableValue<TArray<FIntPoint>>(Variables, InVariableName, FOptimusDataTypeRegistry::GetArrayTypeName(TBaseStructure<FIntPoint>::Get()), InValue);
+	return SetVariableValue<TArray<FIntPoint>>(ValueMap, InVariableName, FOptimusDataTypeRegistry::GetArrayTypeName(TBaseStructure<FIntPoint>::Get()), InValue);
 }
 
 bool UOptimusDeformerInstance::SetInt3Variable(FName InVariableName, const FIntVector& InValue)
 {
-	return SetVariableValue<FIntVector>(Variables, InVariableName, FOptimusDataTypeRegistry::GetTypeName(TBaseStructure<FIntVector>::Get()), InValue);
+	return SetVariableValue<FIntVector>(ValueMap, InVariableName, FOptimusDataTypeRegistry::GetTypeName(TBaseStructure<FIntVector>::Get()), InValue);
 }
 
 bool UOptimusDeformerInstance::SetInt3ArrayVariable(FName InVariableName, const TArray<FIntVector>& InValue)
 {
-	return SetVariableValue<TArray<FIntVector>>(Variables, InVariableName, FOptimusDataTypeRegistry::GetArrayTypeName(TBaseStructure<FIntVector>::Get()), InValue);
+	return SetVariableValue<TArray<FIntVector>>(ValueMap, InVariableName, FOptimusDataTypeRegistry::GetArrayTypeName(TBaseStructure<FIntVector>::Get()), InValue);
 }
 
 bool UOptimusDeformerInstance::SetInt4Variable(FName InVariableName, const FIntVector4& InValue)
 {
-	return SetVariableValue<FIntVector4>(Variables, InVariableName, FOptimusDataTypeRegistry::GetTypeName(TBaseStructure<FIntVector4>::Get()), InValue);
+	return SetVariableValue<FIntVector4>(ValueMap, InVariableName, FOptimusDataTypeRegistry::GetTypeName(TBaseStructure<FIntVector4>::Get()), InValue);
 }
 
 bool UOptimusDeformerInstance::SetInt4ArrayVariable(FName InVariableName, const TArray<FIntVector4>& InValue)
 {
-	return SetVariableValue<TArray<FIntVector4>>(Variables, InVariableName, FOptimusDataTypeRegistry::GetArrayTypeName(TBaseStructure<FIntVector4>::Get()), InValue);
+	return SetVariableValue<TArray<FIntVector4>>(ValueMap, InVariableName, FOptimusDataTypeRegistry::GetArrayTypeName(TBaseStructure<FIntVector4>::Get()), InValue);
 }
 
 
 bool UOptimusDeformerInstance::SetFloatVariable(FName InVariableName, double InValue)
 {
-	if (SetVariableValue<double>(Variables, InVariableName, FDoubleProperty::StaticClass()->GetFName(), InValue))
+	if (SetVariableValue<double>(ValueMap, InVariableName, FDoubleProperty::StaticClass()->GetFName(), InValue))
 	{
 		return true;
 	}
 	
 	// Fall back on float
-	return SetVariableValue<float>(Variables, InVariableName, FFloatProperty::StaticClass()->GetFName(), static_cast<float>(InValue));
+	return SetVariableValue<float>(ValueMap, InVariableName, FFloatProperty::StaticClass()->GetFName(), static_cast<float>(InValue));
 }
 
 bool UOptimusDeformerInstance::SetFloatArrayVariable(FName InVariableName, const TArray<double>& InValue)
 {
-	return SetVariableValue<TArray<double>>(Variables, InVariableName, FOptimusDataTypeRegistry::GetArrayTypeName(FDoubleProperty::StaticClass()->GetFName()), InValue);
+	return SetVariableValue<TArray<double>>(ValueMap, InVariableName, FOptimusDataTypeRegistry::GetArrayTypeName(FDoubleProperty::StaticClass()->GetFName()), InValue);
 }
 
 bool UOptimusDeformerInstance::SetVector2Variable(FName InVariableName, const FVector2D& InValue)
 {
-	return SetVariableValue<FVector2D>(Variables, InVariableName, FOptimusDataTypeRegistry::GetTypeName(TBaseStructure<FVector2D>::Get()), InValue);
+	return SetVariableValue<FVector2D>(ValueMap, InVariableName, FOptimusDataTypeRegistry::GetTypeName(TBaseStructure<FVector2D>::Get()), InValue);
 }
 
 bool UOptimusDeformerInstance::SetVector2ArrayVariable(FName InVariableName, const TArray<FVector2D>& InValue)
 {
-	return SetVariableValue<TArray<FVector2D>>(Variables, InVariableName, FOptimusDataTypeRegistry::GetArrayTypeName(TBaseStructure<FVector2D>::Get()), InValue);
+	return SetVariableValue<TArray<FVector2D>>(ValueMap, InVariableName, FOptimusDataTypeRegistry::GetArrayTypeName(TBaseStructure<FVector2D>::Get()), InValue);
 }
 
 bool UOptimusDeformerInstance::SetVectorVariable(FName InVariableName, const FVector& InValue)
 {
-	return SetVariableValue<FVector>(Variables, InVariableName, FOptimusDataTypeRegistry::GetTypeName(TBaseStructure<FVector>::Get()), InValue);
+	return SetVariableValue<FVector>(ValueMap, InVariableName, FOptimusDataTypeRegistry::GetTypeName(TBaseStructure<FVector>::Get()), InValue);
 }
 
 bool UOptimusDeformerInstance::SetVectorArrayVariable(FName InVariableName, const TArray<FVector>& InValue)
 {
-	return SetVariableValue<TArray<FVector>>(Variables, InVariableName, FOptimusDataTypeRegistry::GetArrayTypeName(TBaseStructure<FVector>::Get()), InValue);
+	return SetVariableValue<TArray<FVector>>(ValueMap, InVariableName, FOptimusDataTypeRegistry::GetArrayTypeName(TBaseStructure<FVector>::Get()), InValue);
 }
 
 bool UOptimusDeformerInstance::SetVector4Variable(FName InVariableName, const FVector4& InValue)
 {
-	return SetVariableValue<FVector4>(Variables, InVariableName, FOptimusDataTypeRegistry::GetTypeName(TBaseStructure<FVector4>::Get()), InValue);
+	return SetVariableValue<FVector4>(ValueMap, InVariableName, FOptimusDataTypeRegistry::GetTypeName(TBaseStructure<FVector4>::Get()), InValue);
 }
 
 bool UOptimusDeformerInstance::SetVector4ArrayVariable(FName InVariableName, const TArray<FVector4>& InValue)
 {
-	return SetVariableValue<TArray<FVector4>>(Variables, InVariableName, FOptimusDataTypeRegistry::GetArrayTypeName(TBaseStructure<FVector4>::Get()), InValue);
+	return SetVariableValue<TArray<FVector4>>(ValueMap, InVariableName, FOptimusDataTypeRegistry::GetArrayTypeName(TBaseStructure<FVector4>::Get()), InValue);
 }
 
 bool UOptimusDeformerInstance::SetLinearColorVariable(FName InVariableName, const FLinearColor& InValue)
 {
-	return SetVariableValue<FLinearColor>(Variables, InVariableName, FOptimusDataTypeRegistry::GetTypeName(TBaseStructure<FLinearColor>::Get()), InValue);
+	return SetVariableValue<FLinearColor>(ValueMap, InVariableName, FOptimusDataTypeRegistry::GetTypeName(TBaseStructure<FLinearColor>::Get()), InValue);
 }
 
 bool UOptimusDeformerInstance::SetLinearColorArrayVariable(FName InVariableName, const TArray<FLinearColor>& InValue)
 {
-	return SetVariableValue<TArray<FLinearColor>>(Variables, InVariableName, FOptimusDataTypeRegistry::GetArrayTypeName(TBaseStructure<FLinearColor>::Get()), InValue);
+	return SetVariableValue<TArray<FLinearColor>>(ValueMap, InVariableName, FOptimusDataTypeRegistry::GetArrayTypeName(TBaseStructure<FLinearColor>::Get()), InValue);
 }
 
 bool UOptimusDeformerInstance::SetQuatVariable(FName InVariableName, const FQuat& InValue)
 {
-	return SetVariableValue<FQuat>(Variables, InVariableName, FOptimusDataTypeRegistry::GetTypeName(TBaseStructure<FQuat>::Get()), InValue);
+	return SetVariableValue<FQuat>(ValueMap, InVariableName, FOptimusDataTypeRegistry::GetTypeName(TBaseStructure<FQuat>::Get()), InValue);
 }
 
 bool UOptimusDeformerInstance::SetQuatArrayVariable(FName InVariableName, const TArray<FQuat>& InValue)
 {
-	return SetVariableValue<TArray<FQuat>>(Variables, InVariableName, FOptimusDataTypeRegistry::GetArrayTypeName(TBaseStructure<FQuat>::Get()), InValue);
+	return SetVariableValue<TArray<FQuat>>(ValueMap, InVariableName, FOptimusDataTypeRegistry::GetArrayTypeName(TBaseStructure<FQuat>::Get()), InValue);
 }
 
 bool UOptimusDeformerInstance::SetRotatorVariable(FName InVariableName, const FRotator& InValue)
 {
-	return SetVariableValue<FRotator>(Variables, InVariableName, FOptimusDataTypeRegistry::GetTypeName(TBaseStructure<FRotator>::Get()), InValue);
+	return SetVariableValue<FRotator>(ValueMap, InVariableName, FOptimusDataTypeRegistry::GetTypeName(TBaseStructure<FRotator>::Get()), InValue);
 }
 
 bool UOptimusDeformerInstance::SetRotatorArrayVariable(FName InVariableName, const TArray<FRotator>& InValue)
 {
-	return SetVariableValue<TArray<FRotator>>(Variables, InVariableName, FOptimusDataTypeRegistry::GetArrayTypeName(TBaseStructure<FRotator>::Get()), InValue);
+	return SetVariableValue<TArray<FRotator>>(ValueMap, InVariableName, FOptimusDataTypeRegistry::GetArrayTypeName(TBaseStructure<FRotator>::Get()), InValue);
 }
 
 bool UOptimusDeformerInstance::SetTransformVariable(FName InVariableName, const FTransform& InValue)
 {
-	return SetVariableValue<FTransform>(Variables, InVariableName, FOptimusDataTypeRegistry::GetTypeName(TBaseStructure<FTransform>::Get()), InValue);
+	return SetVariableValue<FTransform>(ValueMap, InVariableName, FOptimusDataTypeRegistry::GetTypeName(TBaseStructure<FTransform>::Get()), InValue);
 }
 
 bool UOptimusDeformerInstance::SetTransformArrayVariable(FName InVariableName, const TArray<FTransform>& InValue)
 {
-	return SetVariableValue<TArray<FTransform>>(Variables, InVariableName, FOptimusDataTypeRegistry::GetArrayTypeName(TBaseStructure<FTransform>::Get()), InValue);
+	return SetVariableValue<TArray<FTransform>>(ValueMap, InVariableName, FOptimusDataTypeRegistry::GetArrayTypeName(TBaseStructure<FTransform>::Get()), InValue);
 }
 
-const TArray<UOptimusVariableDescription*>& UOptimusDeformerInstance::GetVariables() const
+bool UOptimusDeformerInstance::SetNameVariable(FName InVariableName, const FName& InValue)
 {
-	return Variables->Descriptions;
+	return SetVariableValue<FName>(ValueMap, InVariableName, FOptimusDataTypeRegistry::GetTypeName(*FNameProperty::StaticClass()), InValue);
+}
+
+bool UOptimusDeformerInstance::SetNameArrayVariable(FName InVariableName, const TArray<FName>& InValue)
+{
+	return SetVariableValue<TArray<FName>>(ValueMap, InVariableName, FOptimusDataTypeRegistry::GetArrayTypeName(*FNameProperty::StaticClass()), InValue);
 }
 
 
@@ -767,21 +790,24 @@ bool UOptimusDeformerInstance::EnqueueTriggerGraph(FName InTriggerGraphName)
 }
 
 
-void UOptimusDeformerInstance::SetConstantValueDirect(TSoftObjectPtr<UObject> InSourceObject, FShaderValueContainer const& InValue)
+void UOptimusDeformerInstance::SetConstantValueDirect(TSoftObjectPtr<UObject> InSourceObject, FOptimusValueContainerStruct const& InValue)
 {
-	// Poke constants into the UGraphDataProvider objects.
 	// This is an editor only operation when constant nodes are edited in the graph and we want to see the result without a full compile step.
-	// Not sure that this is the best approach, or whether to store constants on the UOptimusDeformerInstance like variables?
-	for (FOptimusDeformerInstanceExecInfo& ExecInfo : ComputeGraphExecInfos)
+	if (IOptimusValueProvider* ValueProvider = Cast<IOptimusValueProvider>(InSourceObject.LoadSynchronous()))
 	{
-		TArray< TObjectPtr<UComputeDataProvider> >& DataProviders = ExecInfo.ComputeGraphInstance.GetDataProviders();
-		for (UComputeDataProvider* DataProvider : DataProviders)
+		if (FOptimusValueDescription* Description = ValueMap.Find(ValueProvider->GetValueIdentifier()))
 		{
-			if (UOptimusGraphDataProvider* GraphDataProvider = Cast<UOptimusGraphDataProvider>(DataProvider))
+			check(Description->DataType == ValueProvider->GetValueDataType());
+			
+			if (Description->ValueUsage == EOptimusValueUsage::CPU)
 			{
-				GraphDataProvider->SetConstant(InSourceObject, InValue);
-				break;
+				Description->Value = InValue;
 			}
-		}
+			
+			if (Description->ValueUsage == EOptimusValueUsage::GPU)
+			{
+				Description->ShaderValue = InValue.GetShaderValue(ValueProvider->GetValueDataType());
+			}
+		}	
 	}
 }
diff --git a/Engine/Plugins/Animation/DeformerGraph/Source/OptimusCore/Private/OptimusValueContainer.cpp b/Engine/Plugins/Animation/DeformerGraph/Source/OptimusCore/Private/OptimusValueContainer.cpp
--- a/Engine/Plugins/Animation/DeformerGraph/Source/OptimusCore/Private/OptimusValueContainer.cpp
+++ b/Engine/Plugins/Animation/DeformerGraph/Source/OptimusCore/Private/OptimusValueContainer.cpp
@@ -107,13 +107,10 @@ FOptimusValueContainerStruct UOptimusValueContainer::MakeValueContainerStruct()
 	FOptimusDataTypeRef DataType = GetValueType();
 	if (ensure(ValueProperty) && ensure(DataType.IsValid()))
 	{
-		OtherContainer.SetType(DataType.Resolve());
-		
 		TArrayView<const uint8> ValueData(ValueProperty->ContainerPtrToValuePtr<uint8>(this), ValueProperty->GetSize());
-
-		uint8* OtherContainerMemory = OtherContainer.GetMutableValueMemory();
-
-		FMemory::Memcpy(OtherContainerMemory, ValueData.GetData(), ValueData.Num());
+		
+		OtherContainer.SetType(DataType.Resolve());
+		OtherContainer.SetValue(DataType, ValueData);
 	}
 
 	return OtherContainer;
diff --git a/Engine/Plugins/Animation/DeformerGraph/Source/OptimusCore/Private/OptimusValueContainerStruct.cpp b/Engine/Plugins/Animation/DeformerGraph/Source/OptimusCore/Private/OptimusValueContainerStruct.cpp
--- a/Engine/Plugins/Animation/DeformerGraph/Source/OptimusCore/Private/OptimusValueContainerStruct.cpp
+++ b/Engine/Plugins/Animation/DeformerGraph/Source/OptimusCore/Private/OptimusValueContainerStruct.cpp
@@ -17,18 +17,23 @@ void FOptimusValueContainerStruct::SetType(FOptimusDataTypeRef InDataType)
 	Value.AddProperty(ValuePropertyName, InDataType->CreateProperty(nullptr, ValuePropertyName));		
 }
 
-FShaderValueContainer FOptimusValueContainerStruct::GetShaderValue(FOptimusDataTypeRef InDataType) const
+void FOptimusValueContainerStruct::SetValue(FOptimusDataTypeRef InDataType, TArrayView<const uint8> InValue)
 {
-	check(Value.GetNumPropertiesInBag() == 1);
+	const FProperty* Property = GetValueProperty();
 
-	const UPropertyBag* BagStruct = Value.GetPropertyBagStruct();
-	TConstArrayView<FPropertyBagPropertyDesc> Descs = BagStruct->GetPropertyDescs();
+	if (ensure(Property->SameType(InDataType->CreateProperty(nullptr, NAME_None))))
+	{
+		Property->CopyCompleteValue(Value.GetMutableValue().GetMemory(), InValue.GetData());
+	}
+}
 
-	const FPropertyBagPropertyDesc& ValueDesc = Descs[0];
+FShaderValueContainer FOptimusValueContainerStruct::GetShaderValue(FOptimusDataTypeRef InDataType) const
+{
+	const FProperty* Property = GetValueProperty();
 
-	if (ensure(InDataType.IsValid()) && ensure(ValueDesc.CachedProperty))
+	if (ensure(InDataType.IsValid()) && ensure(Property))
 	{
-		TArrayView<const uint8> ValueData(ValueDesc.CachedProperty->ContainerPtrToValuePtr<uint8>(Value.GetValue().GetMemory()), ValueDesc.CachedProperty->GetSize());
+		TArrayView<const uint8> ValueData(Property->ContainerPtrToValuePtr<uint8>(Value.GetValue().GetMemory()), Property->GetSize());
 		FShaderValueContainer ValueResult = InDataType->MakeShaderValue();
 		if (InDataType->ConvertPropertyValueToShader(ValueData, ValueResult))
 		{
@@ -39,6 +44,19 @@ FShaderValueContainer FOptimusValueContainerStruct::GetShaderValue(FOptimusDataT
 	return {};
 }
 
+FString FOptimusValueContainerStruct::GetValueAsString() const
+{
+	if (const FProperty* Property = GetValueProperty())
+	{
+		FString ValueStr;
+		Property->ExportTextItem_InContainer(ValueStr, GetValueMemory(), nullptr, nullptr, PPF_None);
+
+		return ValueStr;
+	}
+
+	return {};
+}
+
 const FProperty* FOptimusValueContainerStruct::GetValueProperty() const
 {
 	check(Value.GetNumPropertiesInBag() == 1);
@@ -56,7 +74,4 @@ const uint8* FOptimusValueContainerStruct::GetValueMemory() const
 	return Value.GetValue().GetMemory();
 }
 
-uint8* FOptimusValueContainerStruct::GetMutableValueMemory()
-{
-	return Value.GetMutableValue().GetMemory();
-}
+
diff --git a/Engine/Plugins/Animation/DeformerGraph/Source/OptimusCore/Public/IOptimusComputeKernelProvider.h b/Engine/Plugins/Animation/DeformerGraph/Source/OptimusCore/Public/IOptimusComputeKernelProvider.h
--- a/Engine/Plugins/Animation/DeformerGraph/Source/OptimusCore/Public/IOptimusComputeKernelProvider.h
+++ b/Engine/Plugins/Animation/DeformerGraph/Source/OptimusCore/Public/IOptimusComputeKernelProvider.h
@@ -45,16 +45,6 @@ struct FOptimus_KernelParameterBinding
 };
 using FOptimus_KernelParameterBindingList = TArray<FOptimus_KernelParameterBinding>;
 
-// Maps from a data interface node to the data interface that it represents.
-using FOptimus_NodeToDataInterfaceMap =  TMap<const UOptimusNode*, UOptimusComputeDataInterface*>;
-
-// Maps from an output pin to the transient data interface, used to store intermediate results,
-// that it represents.
-using FOptimus_PinToDataInterfaceMap = TMap<const UOptimusNodePin*, UOptimusComputeDataInterface*>;
-
-// Maps from a kernel node to the kernel data interface 
-using FOptimus_KernelNodeToKernelDataInterfaceMap = TMap<const UOptimusNode*, UComputeDataInterface*>;
-
 using FOptimus_ComputeKernelResult = TVariant<UOptimusKernelSource* /* Kernel */, FText /* Error */>;
 
 struct FOptimus_KernelConnection
@@ -86,7 +76,6 @@ public:
 	 * this interface.
 	 * @param InKernelSourceOuter The outer object that will own the new kernel source.
 	 * @param InTraversalContext The current context being used to traverse from the graph. Used mainly to keep track of graph nesting.
-	 * @param InValueNodes 
 	 * @param InOutKernelDataInterface
 	 * @param OutInputDataBindings
 	 * @param OutOutputDataBindings
@@ -94,9 +83,8 @@ public:
 	virtual FOptimus_ComputeKernelResult CreateComputeKernel(
 		UObject* InKernelSourceOuter,
 		const FOptimusPinTraversalContext& InTraversalContext,
-		const FOptimus_KernelInputMap InKernelInputs,
-		const FOptimus_KernelOutputMap InKernelOutputs,
-		const TArray<const UOptimusNode*>& InValueNodes,
+		const FOptimus_KernelInputMap& InKernelInputs,
+		const FOptimus_KernelOutputMap& InKernelOutputs,
 		UComputeDataInterface* InOutKernelDataInterface,
 		FOptimus_InterfaceBindingMap& OutInputDataBindings,
 		FOptimus_InterfaceBindingMap& OutOutputDataBindings
diff --git a/Engine/Plugins/Animation/DeformerGraph/Source/OptimusCore/Public/IOptimusPropertyPinProvider.h b/Engine/Plugins/Animation/DeformerGraph/Source/OptimusCore/Public/IOptimusPropertyPinProvider.h
new file mode 100644
--- /dev/null
+++ b/Engine/Plugins/Animation/DeformerGraph/Source/OptimusCore/Public/IOptimusPropertyPinProvider.h
@@ -0,0 +1,27 @@
+﻿// Copyright Epic Games, Inc. All Rights Reserved.
+
+#pragma once
+
+#include "UObject/Interface.h"
+
+#include "IOptimusPropertyPinProvider.generated.h"
+
+
+class UOptimusNodePin;
+
+
+UINTERFACE()
+class OPTIMUSCORE_API UOptimusPropertyPinProvider :
+	public UInterface
+{
+	GENERATED_BODY()
+};
+
+
+class IOptimusPropertyPinProvider
+{
+	GENERATED_BODY()
+
+public:
+	virtual TArray<UOptimusNodePin*> GetPropertyPins() const = 0;
+};
diff --git a/Engine/Plugins/Animation/DeformerGraph/Source/OptimusCore/Public/IOptimusValueProvider.h b/Engine/Plugins/Animation/DeformerGraph/Source/OptimusCore/Public/IOptimusValueProvider.h
--- a/Engine/Plugins/Animation/DeformerGraph/Source/OptimusCore/Public/IOptimusValueProvider.h
+++ b/Engine/Plugins/Animation/DeformerGraph/Source/OptimusCore/Public/IOptimusValueProvider.h
@@ -2,11 +2,14 @@
 
 #pragma once
 
-#include "OptimusDataType.h"
 #include "UObject/Interface.h"
 
 #include "IOptimusValueProvider.generated.h"
 
+struct FOptimusValueContainerStruct;
+struct FOptimusDataTypeRef;
+struct FOptimusValueIdentifier;
+
 UINTERFACE()
 class UOptimusValueProvider :
 	public UInterface
@@ -24,9 +27,9 @@ class IOptimusValueProvider
 
 public:
 	// Returns the value name.
-	virtual FString GetValueName() const = 0;
-	// Returns the value type.
-	virtual FOptimusDataTypeRef GetValueType() const = 0;
-	// Returns the stored value as a shader-compatible value.
-	virtual FShaderValueContainer GetShaderValue() const = 0;
+	virtual FOptimusValueIdentifier GetValueIdentifier() const = 0;
+	// Returns the value data type.
+	virtual FOptimusDataTypeRef GetValueDataType() const = 0;
+	// Returns the stored value.
+	virtual FOptimusValueContainerStruct GetValue() const = 0;
 };
diff --git a/Engine/Plugins/Animation/DeformerGraph/Source/OptimusCore/Public/OptimusComputeDataInterface.h b/Engine/Plugins/Animation/DeformerGraph/Source/OptimusCore/Public/OptimusComputeDataInterface.h
--- a/Engine/Plugins/Animation/DeformerGraph/Source/OptimusCore/Public/OptimusComputeDataInterface.h
+++ b/Engine/Plugins/Animation/DeformerGraph/Source/OptimusCore/Public/OptimusComputeDataInterface.h
@@ -2,6 +2,7 @@
 
 #pragma once
 
+#include "OptimusDataType.h"
 #include "ComputeFramework/ComputeDataInterface.h"
 
 #include "Templates/SubclassOf.h"
@@ -123,6 +124,12 @@ struct FOptimusCDIPinDefinition
 	bool bMutable = true;
 };
 
+struct FOptimusCDIPropertyPinDefinition
+{
+	FName PinName = NAME_None;
+
+	FOptimusDataTypeRef DataType; 
+};
 
 UCLASS(Abstract, Const)
 class OPTIMUSCORE_API UOptimusComputeDataInterface : public UComputeDataInterface
@@ -147,8 +154,11 @@ public:
 	virtual FName GetCategory() const { return CategoryName::DataInterfaces; }
 
 	/// Returns the list of pins that will map to the shader functions provided by this data interface.
-	virtual TArray<FOptimusCDIPinDefinition> GetPinDefinitions() const PURE_VIRTUAL(UOptimusComputeDataInterface::GetDisplayName, return {};)
+	virtual TArray<FOptimusCDIPinDefinition> GetPinDefinitions() const PURE_VIRTUAL(UOptimusComputeDataInterface::GetPinDefinitions, return {};)
 
+	/// Returns the list of pins that are related to the properties of the data interface.
+	virtual TArray<FOptimusCDIPropertyPinDefinition> GetPropertyPinDefinitions() const { return {}; }
+	
 	/// Data interface can use this to set default values/pins
 	virtual void Initialize() {};
 
diff --git a/Engine/Plugins/Animation/DeformerGraph/Source/OptimusCore/Public/OptimusDataType.h b/Engine/Plugins/Animation/DeformerGraph/Source/OptimusCore/Public/OptimusDataType.h
--- a/Engine/Plugins/Animation/DeformerGraph/Source/OptimusCore/Public/OptimusDataType.h
+++ b/Engine/Plugins/Animation/DeformerGraph/Source/OptimusCore/Public/OptimusDataType.h
@@ -15,14 +15,15 @@
 UENUM(meta = (Bitflags))
 enum class EOptimusDataTypeUsageFlags : uint8
 {
-	None				 = 0,
+	None= 0,
 	
-	Resource			 = 1 << 0,		/** This type can be used in a resource */
-	Variable			 = 1 << 1,		/** This type can be used in a variable */
-	AnimAttributes       = 1 << 2,       /** This type can be used to query a single anim attribute on a single bone*/
-	DataInterfaceOutput  = 1 << 3,       /** This type can be used as output of a data interface*/
-	PinType				 = 1 << 4,       /** This type can be used as pin type*/
-	PerBoneAnimAttribute = 1 << 5,       /** This type can be used as to query a per-bone anim attribute to produce a bone buffer */
+	Resource				= 1 << 0,		/** This type can be used in a resource */
+	Variable				= 1 << 1,		/** This type can be used in a variable */
+	AnimAttributes			= 1 << 2,		/** This type can be used to query a single anim attribute on a single bone*/
+	DataInterfaceOutput		= 1 << 3,		/** This type can be used as output of a data interface (or input for terminal data interfaces)*/
+	PinType					= 1 << 4,		/** This type can be used as pin type*/
+	PerBoneAnimAttribute	= 1 << 5,		/** This type can be used to query a per-bone anim attribute to produce a bone buffer */
+	Property				= 1 << 6,		/** This type can be used in a variable, but can only be connected to a property pin on a node*/
 };
 ENUM_CLASS_FLAGS(EOptimusDataTypeUsageFlags)
 
diff --git a/Engine/Plugins/Animation/DeformerGraph/Source/OptimusCore/Public/OptimusDeformer.h b/Engine/Plugins/Animation/DeformerGraph/Source/OptimusCore/Public/OptimusDeformer.h
--- a/Engine/Plugins/Animation/DeformerGraph/Source/OptimusCore/Public/OptimusDeformer.h
+++ b/Engine/Plugins/Animation/DeformerGraph/Source/OptimusCore/Public/OptimusDeformer.h
@@ -11,6 +11,8 @@
 #include "OptimusComponentSource.h"
 #include "OptimusDataDomain.h"
 #include "OptimusNodeSubGraph.h"
+#include "OptimusValueContainerStruct.h"
+#include "OptimusValue.h"
 
 #include "Animation/MeshDeformer.h"
 #include "Interfaces/Interface_PreviewMeshProvider.h"
@@ -34,11 +36,10 @@ enum class EOptimusDiagnosticLevel : uint8;
 struct FOptimusCompilerDiagnostic;
 struct FOptimusCompoundAction;
 
-
 DECLARE_MULTICAST_DELEGATE_OneParam(FOptimusCompileBegin, UOptimusDeformer *);
 DECLARE_MULTICAST_DELEGATE_OneParam(FOptimusCompileEnd, UOptimusDeformer *);
 DECLARE_MULTICAST_DELEGATE_OneParam(FOptimusGraphCompileMessageDelegate, FOptimusCompilerDiagnostic const&);
-DECLARE_MULTICAST_DELEGATE_TwoParams(FOptimusConstantValueUpdate, TSoftObjectPtr<UObject>, FShaderValueContainer const&);
+DECLARE_MULTICAST_DELEGATE_TwoParams(FOptimusConstantValueUpdate, TSoftObjectPtr<UObject>, FOptimusValueContainerStruct const&);
 DECLARE_MULTICAST_DELEGATE_OneParam(FOptimusSetAllInstancesCanbeActive, bool);
 
 UENUM()
@@ -66,6 +67,15 @@ struct FOptimusComputeGraphInfo
 	TObjectPtr<UOptimusComputeGraph> ComputeGraph = nullptr;
 };
 
+struct FOptimusNodeGraphCompilationResult
+{
+	TArray<FOptimusComputeGraphInfo> ComputeGraphInfos;
+
+	TMap<TWeakObjectPtr<const UComputeDataInterface>, FOptimusDataInterfacePropertyOverrideInfo> DataInterfacePropertyOverrideMap;
+
+	TMap<FOptimusValueIdentifier, FOptimusValueDescription> ValueMap;	
+};
+
 /** A container class that owns component source bindings. This is used to ensure we don't end up
   * with a namespace clash between graphs, variables, bindings and resources.
   */
@@ -495,12 +505,19 @@ protected:
 	// The compute graphs to execute.
 	UPROPERTY()
 	TArray<FOptimusComputeGraphInfo> ComputeGraphs;
+	
+	UPROPERTY()
+	TMap<TWeakObjectPtr<const UComputeDataInterface>, FOptimusDataInterfacePropertyOverrideInfo> DataInterfacePropertyOverrideMap;
+	
+	UPROPERTY()
+	TMap<FOptimusValueIdentifier, FOptimusValueDescription> ValueMap;
 
 private:
 	void PostLoadFixupMissingComponentBindingsCompat();
 	void PostLoadFixupMismatchedResourceDataDomains();
 	void PostLoadRemoveDeprecatedExecutionNodes();
 	void PostLoadRemoveDeprecatedValueContainerGeneratorClass();
+	void PostLoadMoveValueFromGraphDataInterfaceToDeformerValueMap();
 
 	/** Find a compatible binding with the given data interface. Returns nullptr if no such binding exists */
 	UOptimusComponentSourceBinding* FindCompatibleBindingWithInterface(
@@ -514,7 +531,7 @@ private:
 	TArray<UOptimusNode*> GetAllNodesOfClass(UClass* InNodeClass) const;
 	
 	/// Compile a node graph to a compute graph. Returns one or two complete compute graphs if compilation succeeded. 
-	TArray<FOptimusComputeGraphInfo> CompileNodeGraphToComputeGraphs(
+	FOptimusNodeGraphCompilationResult CompileNodeGraphToComputeGraphs(
 		const UOptimusNodeGraph *InNodeGraph,
 		TFunction<void(EOptimusDiagnosticLevel, FText, const UObject*)> InErrorReporter
 		);
diff --git a/Engine/Plugins/Animation/DeformerGraph/Source/OptimusCore/Public/OptimusDeformerInstance.h b/Engine/Plugins/Animation/DeformerGraph/Source/OptimusCore/Public/OptimusDeformerInstance.h
--- a/Engine/Plugins/Animation/DeformerGraph/Source/OptimusCore/Public/OptimusDeformerInstance.h
+++ b/Engine/Plugins/Animation/DeformerGraph/Source/OptimusCore/Public/OptimusDeformerInstance.h
@@ -5,6 +5,7 @@
 #include "Animation/MeshDeformerInstance.h"
 #include "ComputeFramework/ComputeGraphInstance.h"
 #include "Engine/EngineTypes.h"
+#include "OptimusValue.h"
 
 #include "OptimusDeformerInstance.generated.h"
 
@@ -20,28 +21,8 @@ class UOptimusVariableContainer;
 class UOptimusVariableDescription;
 class UOptimusComponentSourceBinding;
 struct FShaderValueContainer;
-
-
-struct FOptimusDeformerInstanceComponentLodContext
-{
-	TArray<int32> LodIndexPerComponent;
-
-	friend uint32 GetTypeHash(const FOptimusDeformerInstanceComponentLodContext& InContext)
-	{
-		uint32 Hash = GetTypeHash(InContext.LodIndexPerComponent.Num());
-		
-		for (int32 LodIndex : InContext.LodIndexPerComponent)
-		{
-			Hash = HashCombineFast(Hash, GetTypeHash(LodIndex));
-		}
-		return Hash;
-	}
-
-	bool operator==(const FOptimusDeformerInstanceComponentLodContext& InOther) const
-	{
-		return LodIndexPerComponent == InOther.LodIndexPerComponent;
-	}
-};
+struct FOptimusValueContainerStruct;
+class UComputeDataInterface;
 
 class FOptimusPersistentBufferPool
 {
@@ -296,22 +277,29 @@ public:
 	UFUNCTION(BlueprintCallable, Category = "Deformer", meta = (DisplayName = "Set Variable (Transform Array)"))
 	bool SetTransformArrayVariable(FName InVariableName, const TArray<FTransform>& InValue);
 	
-	/** Get an array containing all the variables. */
-	UFUNCTION(BlueprintGetter)
-	const TArray<UOptimusVariableDescription*>& GetVariables() const;
+	UFUNCTION(BlueprintCallable, Category = "Deformer", meta = (DisplayName = "Set Variable (Name)"))
+	bool SetNameVariable(FName InVariableName, const FName& InValue);
+
+	UFUNCTION(BlueprintCallable, Category = "Deformer", meta = (DisplayName = "Set Variable (Name Array)"))
+	bool SetNameArrayVariable(FName InVariableName, const TArray<FName>& InValue);
 
 	/** Trigger a named trigger graph to run on the next tick */
 	UFUNCTION(BlueprintCallable, Category="Deformer")
 	bool EnqueueTriggerGraph(FName InTriggerGraphName);
 	
-	
-	
 	/** Directly set a graph constant value. */
-	void SetConstantValueDirect(TSoftObjectPtr<UObject> InSourceObject, FShaderValueContainer const& InValue);
+	void SetConstantValueDirect(TSoftObjectPtr<UObject> InSourceObject, FOptimusValueContainerStruct const& InValue);
 
 	FOptimusPersistentBufferPoolPtr GetBufferPool() const { return BufferPool; }
 
 	void SetCanBeActive(bool bInCanBeActive);
+
+	FOptimusValueContainerStruct GetDataInterfacePropertyOverride(
+		const UComputeDataInterface* DataInterface,
+		FName PinName
+		);
+
+	const FShaderValueContainer& GetShaderValue(const FOptimusValueIdentifier& InValueId) const;
 	
 	/** Making sure compute graphs belong to this instance does not run before instances before it */
 	int32 GraphSortPriorityOffset = 0;
@@ -341,14 +329,12 @@ private:
 	UPROPERTY()
 	TArray<FOptimusDeformerInstanceExecInfo> ComputeGraphExecInfos;
 
-	/** Storage for variable data. */
-	UPROPERTY()
-	TObjectPtr<UOptimusVariableContainer> Variables;
-
-	UPROPERTY()
+	TMap<FOptimusValueIdentifier, FOptimusValueDescription> ValueMap;
+	
+	TMap<TWeakObjectPtr<const UComputeDataInterface>, FOptimusDataInterfacePropertyOverrideInfo> DataInterfacePropertyOverrideMap;
+	
 	TArray<TWeakObjectPtr<UActorComponent>> WeakBoundComponents;
 	
-	UPROPERTY()
 	TArray<TWeakObjectPtr<const UOptimusComponentSource>> WeakComponentSources;
 	
 	// List of graphs that should be run on the next tick. 
diff --git a/Engine/Plugins/Animation/DeformerGraph/Source/OptimusCore/Public/OptimusObjectVersion.h b/Engine/Plugins/Animation/DeformerGraph/Source/OptimusCore/Public/OptimusObjectVersion.h
--- a/Engine/Plugins/Animation/DeformerGraph/Source/OptimusCore/Public/OptimusObjectVersion.h
+++ b/Engine/Plugins/Animation/DeformerGraph/Source/OptimusCore/Public/OptimusObjectVersion.h
@@ -23,6 +23,7 @@ struct OPTIMUSCORE_API FOptimusObjectVersion
 		KernelParameterBindingToggleAtomic,
 		PropertyBagValueContainer,
 		SkinnedMeshWriteDIColorBufferManualFetchSwizzle,
+		PropertyPinSupport,
 
 		// -----<new versions can be added above this line>-------------------------------------------------
 		VersionPlusOne,
diff --git a/Engine/Plugins/Animation/DeformerGraph/Source/OptimusCore/Public/OptimusValue.h b/Engine/Plugins/Animation/DeformerGraph/Source/OptimusCore/Public/OptimusValue.h
new file mode 100644
--- /dev/null
+++ b/Engine/Plugins/Animation/DeformerGraph/Source/OptimusCore/Public/OptimusValue.h
@@ -0,0 +1,73 @@
+﻿// Copyright Epic Games, Inc. All Rights Reserved.
+#pragma once
+#include "OptimusDataType.h"
+#include "OptimusValueContainerStruct.h"
+
+#include "OptimusValue.generated.h"
+
+UENUM()
+enum class EOptimusValueUsage: uint8
+{
+	None = 0,
+	CPU = 1 << 0,
+	GPU = 1 << 1
+};
+ENUM_CLASS_FLAGS(EOptimusValueUsage);
+
+UENUM()
+enum class EOptimusValueType : uint8
+{
+	Invalid,
+	Constant,
+	Variable
+};
+
+USTRUCT()
+struct FOptimusValueIdentifier
+{
+	GENERATED_BODY()
+
+	UPROPERTY()
+	EOptimusValueType Type;
+
+	UPROPERTY()
+	FName Name;
+	
+	friend uint32 GetTypeHash(const FOptimusValueIdentifier& InIdentifier)
+	{
+		return HashCombineFast(GetTypeHash(InIdentifier.Type), GetTypeHash(InIdentifier.Name));
+	}
+
+	bool operator==(FOptimusValueIdentifier const& InOther) const
+	{
+		return Type == InOther.Type && Name == InOther.Name;
+	}
+};
+
+USTRUCT()
+struct FOptimusValueDescription
+{
+	GENERATED_BODY()
+	
+	UPROPERTY()
+	FOptimusDataTypeRef DataType;
+
+	UPROPERTY()
+	EOptimusValueUsage ValueUsage = EOptimusValueUsage::None;
+	
+	UPROPERTY()
+	FOptimusValueContainerStruct Value;
+
+	UPROPERTY()
+	FShaderValueContainer ShaderValue;
+};
+
+USTRUCT()
+struct FOptimusDataInterfacePropertyOverrideInfo
+{
+	GENERATED_BODY()
+
+	UPROPERTY()
+	TMap<FName, FOptimusValueIdentifier> PinNameToValueIdMap;
+};
+
diff --git a/Engine/Plugins/Animation/DeformerGraph/Source/OptimusCore/Public/OptimusValueContainerStruct.h b/Engine/Plugins/Animation/DeformerGraph/Source/OptimusCore/Public/OptimusValueContainerStruct.h
--- a/Engine/Plugins/Animation/DeformerGraph/Source/OptimusCore/Public/OptimusValueContainerStruct.h
+++ b/Engine/Plugins/Animation/DeformerGraph/Source/OptimusCore/Public/OptimusValueContainerStruct.h
@@ -22,13 +22,12 @@ struct OPTIMUSCORE_API FOptimusValueContainerStruct
 
 	bool IsInitialized() const;
 	void SetType(FOptimusDataTypeRef InDataType);
+	void SetValue(FOptimusDataTypeRef InDataType, TArrayView<const uint8> InValue);
 	FShaderValueContainer GetShaderValue(FOptimusDataTypeRef InDataType) const;
+	FString GetValueAsString() const;
+	
+private:
 	const FProperty* GetValueProperty() const;
 	const uint8* GetValueMemory() const;
-	
-protected:
-	friend class UOptimusValueContainer;
-	
-	uint8* GetMutableValueMemory();
-	
+
 };
diff --git a/Engine/Plugins/Animation/DeformerGraph/Source/OptimusCore/Public/OptimusVariableDescription.h b/Engine/Plugins/Animation/DeformerGraph/Source/OptimusCore/Public/OptimusVariableDescription.h
--- a/Engine/Plugins/Animation/DeformerGraph/Source/OptimusCore/Public/OptimusVariableDescription.h
+++ b/Engine/Plugins/Animation/DeformerGraph/Source/OptimusCore/Public/OptimusVariableDescription.h
@@ -64,7 +64,7 @@ public:
 	FName VariableName;
 
 	/** The data type of the variable */
-	UPROPERTY(EditAnywhere, Category = VariableDefinition, meta=(UseInVariable))
+	UPROPERTY(EditAnywhere, Category = VariableDefinition, meta=(UseInVariable, UseInProperty))
 	FOptimusDataTypeRef DataType;
 
 	/** The default value for the variable. */
diff --git a/Engine/Plugins/Animation/DeformerGraph/Source/OptimusEditor/Private/OptimusDetailsCustomization.cpp b/Engine/Plugins/Animation/DeformerGraph/Source/OptimusEditor/Private/OptimusDetailsCustomization.cpp
--- a/Engine/Plugins/Animation/DeformerGraph/Source/OptimusEditor/Private/OptimusDetailsCustomization.cpp
+++ b/Engine/Plugins/Animation/DeformerGraph/Source/OptimusEditor/Private/OptimusDetailsCustomization.cpp
@@ -90,6 +90,10 @@ void FOptimusDataTypeRefCustomization::CustomizeHeader(
 		{
 			UsageMask |= EOptimusDataTypeUsageFlags::PerBoneAnimAttribute;
 		}
+		if (InPropertyHandle->HasMetaData(FName(TEXT("UseInProperty"))))
+		{
+			UsageMask |= EOptimusDataTypeUsageFlags::Property;
+		}
 
 		return UsageMask;
 	};
diff --git a/Engine/Plugins/Animation/DeformerGraph/Source/OptimusEditor/Private/OptimusEditorGraphSchema.cpp b/Engine/Plugins/Animation/DeformerGraph/Source/OptimusEditor/Private/OptimusEditorGraphSchema.cpp
--- a/Engine/Plugins/Animation/DeformerGraph/Source/OptimusEditor/Private/OptimusEditorGraphSchema.cpp
+++ b/Engine/Plugins/Animation/DeformerGraph/Source/OptimusEditor/Private/OptimusEditorGraphSchema.cpp
@@ -123,7 +123,7 @@ void UOptimusEditorGraphSchema::GetGraphActions(
 	{
 		// For now only allow variable/resource compatible types to spawn constant nodes
 		if (DataTypeHandle->CanCreateProperty() &&
-			EnumHasAnyFlags(DataTypeHandle->UsageFlags, EOptimusDataTypeUsageFlags::Resource | EOptimusDataTypeUsageFlags::Variable))
+			EnumHasAnyFlags(DataTypeHandle->UsageFlags, EOptimusDataTypeUsageFlags::Resource | EOptimusDataTypeUsageFlags::Variable | EOptimusDataTypeUsageFlags::Property))
 		{
 			const FText NodeName = FText::Format(LOCTEXT("ConstantValueNode", "{0} Constant"), DataTypeHandle->DisplayName);
 			const FText NodeCategory = FText::FromName(UOptimusNode::CategoryName::Values);
diff --git a/Engine/Plugins/Animation/DeformerGraph/Source/OptimusEditor/Private/Widgets/SOptimusEditorGraphExplorerItem.cpp b/Engine/Plugins/Animation/DeformerGraph/Source/OptimusEditor/Private/Widgets/SOptimusEditorGraphExplorerItem.cpp
--- a/Engine/Plugins/Animation/DeformerGraph/Source/OptimusEditor/Private/Widgets/SOptimusEditorGraphExplorerItem.cpp
+++ b/Engine/Plugins/Animation/DeformerGraph/Source/OptimusEditor/Private/Widgets/SOptimusEditorGraphExplorerItem.cpp
@@ -84,7 +84,7 @@ public:
 		ChildSlot
 		    [SNew(SOptimusDataTypeSelector)
 		            .CurrentDataType(this, &SVariableDataTypeSelectorHelper::OnGetDataType)
-		            .UsageMask(EOptimusDataTypeUsageFlags::Variable)
+		            .UsageMask(EOptimusDataTypeUsageFlags::Variable | EOptimusDataTypeUsageFlags::Property)
 		            .ViewType(SOptimusDataTypeSelector::EViewType::IconOnly)
 		            .bViewOnly(bInIsReadOnly.Get()) // FIXME: May be dynamic.
 		            .OnDataTypeChanged(this, &SVariableDataTypeSelectorHelper::OnDataTypeChanged)];
diff --git a/Engine/Plugins/Runtime/ComputeFramework/Source/ComputeFramework/Private/ComputeGraph.cpp b/Engine/Plugins/Runtime/ComputeFramework/Source/ComputeFramework/Private/ComputeGraph.cpp
--- a/Engine/Plugins/Runtime/ComputeFramework/Source/ComputeFramework/Private/ComputeGraph.cpp
+++ b/Engine/Plugins/Runtime/ComputeFramework/Source/ComputeFramework/Private/ComputeGraph.cpp
@@ -65,6 +65,11 @@ void UComputeGraph::PostLoad()
 {
 	Super::PostLoad();
 
+	for (UComputeDataInterface* DataInterface : DataInterfaces)
+	{
+		DataInterface->ConditionalPostLoad();
+	}
+	
 #if WITH_EDITOR
 	// PostLoad our kernel dependencies before any compiling.
 	for (UComputeKernel* Kernel : KernelInvocations)
