[Add Deformer Node] added additional settings to control which child component receives deformers from their parent

diff --git a/Engine/Plugins/Animation/DeformerGraph/Source/OptimusCore/Private/ControlRig/RigUnit_Optimus.cpp b/Engine/Plugins/Animation/DeformerGraph/Source/OptimusCore/Private/ControlRig/RigUnit_Optimus.cpp
--- a/Engine/Plugins/Animation/DeformerGraph/Source/OptimusCore/Private/ControlRig/RigUnit_Optimus.cpp
+++ b/Engine/Plugins/Animation/DeformerGraph/Source/OptimusCore/Private/ControlRig/RigUnit_Optimus.cpp
@@ -236,6 +236,60 @@ FRigUnit_AddOptimusDeformer_Execute()
 		// While Add is done on game thread, enqueue and set variable are done on anim thread
 		// Given that parent component always ticks before child components, it should be safe for the parent
 		// component to modify the deformer instance manager on the child components
+
+		const FRigVMTrait_OptimusDeformerSettings* SettingsTraitPtr = nullptr;
+		
+		for(const FRigVMTraitScope& Scope : Traits)
+		{
+			if (const FRigVMTrait_OptimusDeformerSettings* SettingsTrait = Scope.GetTrait<FRigVMTrait_OptimusDeformerSettings>())
+			{
+				SettingsTraitPtr = SettingsTrait;
+				break;
+			}
+		}
+
+		check(SettingsTraitPtr != nullptr);
+		
+		struct Local
+		{
+			static TArray<USkeletalMeshComponent*> GetComponentsToProcess(USkeletalMeshComponent* RigMeshComponent, bool DeformChildComponents, FName ExcludeChildComponentsWithTag)
+			{
+				TArray<USkeletalMeshComponent*> ComponentsToProcess = {RigMeshComponent};
+
+				if (DeformChildComponents)
+				{
+					TArray<USceneComponent*> ChildComponents;
+					RigMeshComponent->GetChildrenComponents(true, ChildComponents);
+					for (USceneComponent* Component : ChildComponents)
+					{
+						if (USkeletalMeshComponent* SkeletalMeshComponent = Cast<USkeletalMeshComponent>(Component))
+						{
+							if (!Component->ComponentHasTag(ExcludeChildComponentsWithTag))
+							{
+								ComponentsToProcess.Add(SkeletalMeshComponent);
+							}
+						}
+					}	
+				}
+
+				return ComponentsToProcess;
+			}
+			
+		};
+		
+		TArray<USkeletalMeshComponent*> ComponentsToProcess =
+			Local::GetComponentsToProcess(
+				RigMeshComponent, SettingsTraitPtr->DeformChildComponents, SettingsTraitPtr->ExcludeChildComponentsWithTag);
+				
+		for (const USkeletalMeshComponent* ComponentToProcess : ComponentsToProcess)
+		{
+			// Currently, there is only one deformer instance used by all LODs, so use LOD 0 here for now
+			// This might change in the future depending on how per-LOD instance is implemented
+			if (UOptimusDeformerDynamicInstanceManager* DeformerInstanceManager = Cast<UOptimusDeformerDynamicInstanceManager>(ComponentToProcess->GetMeshDeformerInstanceForLOD(0)))
+			{
+				DeformerInstanceManager->EnqueueRigDeformer(DeformerInstanceGuid, SettingsTraitPtr->ExecutionPhase, SettingsTraitPtr->ExecutionGroup);
+			}	
+		}
 		
 		for(const FRigVMTraitScope& Scope : Traits)
 		{
@@ -247,7 +301,9 @@ FRigUnit_AddOptimusDeformer_Execute()
 						WeakMesh = TWeakObjectPtr(RigMeshComponent),
 						WeakRig = TWeakObjectPtr(ExecuteContext.ControlRig),
 						DeformerGraphAsset = DeformerTrait->DeformerGraph,
-						DeformerInstanceGuid]()
+						DeformerInstanceGuid,
+						DeformChildComponents = SettingsTraitPtr->DeformChildComponents,
+						ExcludeChildComponentsWithTag = SettingsTraitPtr->ExcludeChildComponentsWithTag]()
 					{
 						if (WeakMesh.IsValid() && WeakRig.IsValid())
 						{
@@ -257,19 +313,9 @@ FRigUnit_AddOptimusDeformer_Execute()
 								if (UOptimusDeformer* DeformerGraph = DeformerGraphAsset.LoadSynchronous())
 								{
 									USkeletalMeshComponent* RigMeshComponent = WeakMesh.Get();
-
-									TArray<USceneComponent*> ChildComponents;
-									RigMeshComponent->GetChildrenComponents(true, ChildComponents);
-
-									TArray<USkeletalMeshComponent*> ComponentsToProcess = {RigMeshComponent};
-
-									for (USceneComponent* Component : ChildComponents)
-									{
-										if (USkeletalMeshComponent* SkeletalMeshComponent = Cast<USkeletalMeshComponent>(Component))
-										{
-											ComponentsToProcess.Add(SkeletalMeshComponent);
-										}
-									}
+									
+									TArray<USkeletalMeshComponent*> ComponentsToProcess =
+										Local::GetComponentsToProcess(RigMeshComponent, DeformChildComponents, ExcludeChildComponentsWithTag);
 
 									for (USkeletalMeshComponent* ComponentToProcess : ComponentsToProcess)
 									{
@@ -301,45 +347,19 @@ FRigUnit_AddOptimusDeformer_Execute()
 					}, TStatId(), NULL, ENamedThreads::GameThread);
 				}	
 			}
-			else
+			else if (const FRigVMTrait_OptimusVariableBase* VariableTrait = Scope.GetTrait<FRigVMTrait_OptimusVariableBase>())
 			{
-				TArray<USceneComponent*> ChildComponents;
-				RigMeshComponent->GetChildrenComponents(true, ChildComponents);
-				TArray<USkeletalMeshComponent*> ComponentsToProcess = {RigMeshComponent};
-				for (USceneComponent* Component : ChildComponents)
+				for (const USkeletalMeshComponent* ComponentToProcess : ComponentsToProcess)
 				{
-					if (USkeletalMeshComponent* SkeletalMeshComponent = Cast<USkeletalMeshComponent>(Component))
+					// Currently, there is only one deformer instance used by all LODs, so use LOD 0 here for now
+					// This might change in the future depending on how per-LOD instance is implemented
+					if (UOptimusDeformerDynamicInstanceManager* DeformerInstanceManager = Cast<UOptimusDeformerDynamicInstanceManager>(ComponentToProcess->GetMeshDeformerInstanceForLOD(0)))
 					{
-						ComponentsToProcess.Add(SkeletalMeshComponent);
-					}
-				}
-				
-				if (const FRigVMTrait_OptimusDeformerSettings* SettingsTrait = Scope.GetTrait<FRigVMTrait_OptimusDeformerSettings>())
-				{
-					for (const USkeletalMeshComponent* ComponentToProcess : ComponentsToProcess)
-					{
-						// Currently, there is only one deformer instance used by all LODs, so use LOD 0 here for now
-						// This might change in the future depending on how per-LOD instance is implemented
-						if (UOptimusDeformerDynamicInstanceManager* DeformerInstanceManager = Cast<UOptimusDeformerDynamicInstanceManager>(ComponentToProcess->GetMeshDeformerInstanceForLOD(0)))
+						if (DeformerInstanceGuid.IsValid())
 						{
-							DeformerInstanceManager->EnqueueRigDeformer(DeformerInstanceGuid, SettingsTrait->ExecutionPhase, SettingsTrait->ExecutionGroup);
-						}	
-					}
-				}
-				else if (const FRigVMTrait_OptimusVariableBase* VariableTrait = Scope.GetTrait<FRigVMTrait_OptimusVariableBase>())
-				{
-					for (const USkeletalMeshComponent* ComponentToProcess : ComponentsToProcess)
-					{
-						// Currently, there is only one deformer instance used by all LODs, so use LOD 0 here for now
-						// This might change in the future depending on how per-LOD instance is implemented
-						if (UOptimusDeformerDynamicInstanceManager* DeformerInstanceManager = Cast<UOptimusDeformerDynamicInstanceManager>(ComponentToProcess->GetMeshDeformerInstanceForLOD(0)))
-						{
-							if (DeformerInstanceGuid.IsValid())
+							if (UOptimusDeformerInstance* DeformerInstance = DeformerInstanceManager->GetRigDeformer(DeformerInstanceGuid))
 							{
-								if (UOptimusDeformerInstance* DeformerInstance = DeformerInstanceManager->GetRigDeformer(DeformerInstanceGuid))
-								{
-									VariableTrait->SetValue(DeformerInstance);
-								}
+								VariableTrait->SetValue(DeformerInstance);
 							}
 						}
 					}
diff --git a/Engine/Plugins/Animation/DeformerGraph/Source/OptimusCore/Private/ControlRig/RigUnit_Optimus.h b/Engine/Plugins/Animation/DeformerGraph/Source/OptimusCore/Private/ControlRig/RigUnit_Optimus.h
--- a/Engine/Plugins/Animation/DeformerGraph/Source/OptimusCore/Private/ControlRig/RigUnit_Optimus.h
+++ b/Engine/Plugins/Animation/DeformerGraph/Source/OptimusCore/Private/ControlRig/RigUnit_Optimus.h
@@ -36,6 +36,14 @@ struct OPTIMUSCORE_API FRigVMTrait_OptimusDeformerSettings: public FRigVMTrait
 	// Deformers are first sorted by execution group index, then by the order in which they are added
 	UPROPERTY(EditAnywhere, Category = "Trait", meta=(Input))
 	int32 ExecutionGroup = 1;
+
+	// Whether to apply the deformer to all child components as well
+	UPROPERTY(EditAnywhere, Category = "Trait", meta=(Input))
+	bool DeformChildComponents = true;
+
+	// Deformer won't be applied to child components that have the specified component tag
+	UPROPERTY(EditAnywhere, Category = "Trait", meta=(Input))
+	FName ExcludeChildComponentsWithTag = NAME_None;
 };
 
 /** Adds a deformer to the Skeletal Mesh Component*/
