[Deformer & ControlRig] added bool variable support , minor optimusnodepin name fix

diff --git a/Engine/Plugins/Animation/DeformerGraph/Source/OptimusCore/Private/ControlRig/RigUnit_Optimus.cpp b/Engine/Plugins/Animation/DeformerGraph/Source/OptimusCore/Private/ControlRig/RigUnit_Optimus.cpp
--- a/Engine/Plugins/Animation/DeformerGraph/Source/OptimusCore/Private/ControlRig/RigUnit_Optimus.cpp
+++ b/Engine/Plugins/Animation/DeformerGraph/Source/OptimusCore/Private/ControlRig/RigUnit_Optimus.cpp
@@ -196,6 +196,14 @@ TArray<FRigVMUserWorkflow> FRigUnit_AddOptimusDeformer::GetSupportedWorkflows(co
 						{
 							TraitName = Controller->AddTrait(Node->GetFName(), *FRigVMTrait_SetDeformerNameArrayVariable::StaticStruct()->GetPathName(), Variable->VariableName);
 						}
+						else if (Variable->DataType == FOptimusDataTypeRegistry::Get().FindType(*FBoolProperty::StaticClass()))
+						{
+							TraitName = Controller->AddTrait(Node->GetFName(), *FRigVMTrait_SetDeformerBoolVariable::StaticStruct()->GetPathName(), Variable->VariableName);
+						}
+						else if (Variable->DataType == FOptimusDataTypeRegistry::Get().FindArrayType(*FBoolProperty::StaticClass()))
+						{
+							TraitName = Controller->AddTrait(Node->GetFName(), *FRigVMTrait_SetDeformerBoolArrayVariable::StaticStruct()->GetPathName(), Variable->VariableName);
+						}
 						
 						Controller->SetPinExpansion(Node->FindTrait(TraitName)->GetPinPath(), true, true);
 					}
@@ -499,3 +507,13 @@ void FRigVMTrait_SetDeformerNameArrayVariable::SetValue(UOptimusDeformerInstance
 {
 	InInstance->SetNameArrayVariable(*GetName(), Value);
 }
+
+void FRigVMTrait_SetDeformerBoolVariable::SetValue(UOptimusDeformerInstance* InInstance) const
+{
+	InInstance->SetBoolVariable(*GetName(), Value);
+}
+
+void FRigVMTrait_SetDeformerBoolArrayVariable::SetValue(UOptimusDeformerInstance* InInstance) const
+{
+	InInstance->SetBoolArrayVariable(*GetName(), Value);
+}
diff --git a/Engine/Plugins/Animation/DeformerGraph/Source/OptimusCore/Private/ControlRig/RigUnit_Optimus.h b/Engine/Plugins/Animation/DeformerGraph/Source/OptimusCore/Private/ControlRig/RigUnit_Optimus.h
--- a/Engine/Plugins/Animation/DeformerGraph/Source/OptimusCore/Private/ControlRig/RigUnit_Optimus.h
+++ b/Engine/Plugins/Animation/DeformerGraph/Source/OptimusCore/Private/ControlRig/RigUnit_Optimus.h
@@ -373,3 +373,24 @@ struct OPTIMUSCORE_API FRigVMTrait_SetDeformerNameArrayVariable: public FRigVMTr
 	TArray<FName> Value;	
 };
 
+USTRUCT(BlueprintType)
+struct OPTIMUSCORE_API FRigVMTrait_SetDeformerBoolVariable: public FRigVMTrait_OptimusVariableBase
+{
+	GENERATED_BODY()
+
+	void SetValue(UOptimusDeformerInstance* InInstance) const override;
+	
+	UPROPERTY(EditAnywhere, Category = "Trait")
+	bool Value = false;	
+};
+
+USTRUCT(BlueprintType)
+struct OPTIMUSCORE_API FRigVMTrait_SetDeformerBoolArrayVariable: public FRigVMTrait_OptimusVariableBase
+{
+	GENERATED_BODY()
+
+	void SetValue(UOptimusDeformerInstance* InInstance) const override;
+	
+	UPROPERTY(EditAnywhere, Category = "Trait")
+	TArray<bool> Value;	
+};
diff --git a/Engine/Plugins/Animation/DeformerGraph/Source/OptimusCore/Private/OptimusDeformerInstance.cpp b/Engine/Plugins/Animation/DeformerGraph/Source/OptimusCore/Private/OptimusDeformerInstance.cpp
--- a/Engine/Plugins/Animation/DeformerGraph/Source/OptimusCore/Private/OptimusDeformerInstance.cpp
+++ b/Engine/Plugins/Animation/DeformerGraph/Source/OptimusCore/Private/OptimusDeformerInstance.cpp
@@ -636,6 +636,11 @@ bool UOptimusDeformerInstance::SetBoolVariable(FName InVariableName, bool InValu
 	return SetVariableValue(ValueMap, InVariableName, FOptimusDataTypeRegistry::GetTypeName(*FBoolProperty::StaticClass()), InValue);
 }
 
+bool UOptimusDeformerInstance::SetBoolArrayVariable(FName InVariableName, const TArray<bool>& InValue)
+{
+	return SetVariableValue(ValueMap, InVariableName, FOptimusDataTypeRegistry::GetArrayTypeName(*FBoolProperty::StaticClass()), InValue);
+}
+
 bool UOptimusDeformerInstance::SetIntVariable(FName InVariableName, int32 InValue)
 {
 	return SetVariableValue<int32>(ValueMap, InVariableName, FOptimusDataTypeRegistry::GetTypeName(*FIntProperty::StaticClass()), InValue);
diff --git a/Engine/Plugins/Animation/DeformerGraph/Source/OptimusCore/Private/OptimusNodePin.cpp b/Engine/Plugins/Animation/DeformerGraph/Source/OptimusCore/Private/OptimusNodePin.cpp
--- a/Engine/Plugins/Animation/DeformerGraph/Source/OptimusCore/Private/OptimusNodePin.cpp
+++ b/Engine/Plugins/Animation/DeformerGraph/Source/OptimusCore/Private/OptimusNodePin.cpp
@@ -109,8 +109,10 @@ FName UOptimusNodePin::GetUniqueName() const
 
 FText UOptimusNodePin::GetDisplayName() const
 {
-	// So bool.
-	const bool bIsBool = (CastField<FBoolProperty>(GetPropertyFromPin()) != nullptr);
+	
+	const FOptimusDataTypeHandle BoolType = FOptimusDataTypeRegistry::Get().FindType(*FBoolProperty::StaticClass());
+	const bool bIsBool = GetDataType() == BoolType;
+	
 	return FText::FromString(FName::NameToDisplayString(GetName(), bIsBool));
 }
 
diff --git a/Engine/Plugins/Animation/DeformerGraph/Source/OptimusCore/Public/OptimusDeformerInstance.h b/Engine/Plugins/Animation/DeformerGraph/Source/OptimusCore/Public/OptimusDeformerInstance.h
--- a/Engine/Plugins/Animation/DeformerGraph/Source/OptimusCore/Public/OptimusDeformerInstance.h
+++ b/Engine/Plugins/Animation/DeformerGraph/Source/OptimusCore/Public/OptimusDeformerInstance.h
@@ -195,6 +195,9 @@ public:
 	UFUNCTION(BlueprintCallable, Category="Deformer", meta=(DisplayName="Set Variable (Boolean)"))
 	bool SetBoolVariable(FName InVariableName, bool InValue);
 
+	UFUNCTION(BlueprintCallable, Category="Deformer", meta=(DisplayName="Set Variable (Boolean Array)"))
+	bool SetBoolArrayVariable(FName InVariableName, const TArray<bool>& InValue);
+	
 	/** Set the value of an integer variable. */
 	UFUNCTION(BlueprintCallable, Category="Deformer", meta=(DisplayName="Set Variable (Integer)"))
 	bool SetIntVariable(FName InVariableName, int32 InValue);
