[Deformer Graph] added support for overriding constant nodes with variable or constant nodes

diff --git a/Engine/Plugins/Animation/DeformerGraph/Source/OptimusCore/Private/OptimusDataTypeRegistry.cpp b/Engine/Plugins/Animation/DeformerGraph/Source/OptimusCore/Private/OptimusDataTypeRegistry.cpp
--- a/Engine/Plugins/Animation/DeformerGraph/Source/OptimusCore/Private/OptimusDataTypeRegistry.cpp
+++ b/Engine/Plugins/Animation/DeformerGraph/Source/OptimusCore/Private/OptimusDataTypeRegistry.cpp
@@ -1451,6 +1451,11 @@ FOptimusDataTypeHandle FOptimusDataTypeRegistry::FindType(const FProperty& InPro
 		const FName TypeName(*FString::Printf(TEXT("U%s"), *ObjectProperty->PropertyClass->GetName()));
 		return FindType(TypeName);
 	}
+	else if (const FArrayProperty* ArrayProperty = CastField<const FArrayProperty>(&InProperty))
+	{
+		const FProperty* PropertyForType = ArrayProperty->Inner;
+		return FindArrayType(*PropertyForType->GetClass());
+	}
 	else
 	{
 		return FindType(*InProperty.GetClass());
diff --git a/Engine/Plugins/Animation/DeformerGraph/Source/OptimusCore/Private/OptimusDeformer.cpp b/Engine/Plugins/Animation/DeformerGraph/Source/OptimusCore/Private/OptimusDeformer.cpp
--- a/Engine/Plugins/Animation/DeformerGraph/Source/OptimusCore/Private/OptimusDeformer.cpp
+++ b/Engine/Plugins/Animation/DeformerGraph/Source/OptimusCore/Private/OptimusDeformer.cpp
@@ -2140,41 +2140,47 @@ FOptimusNodeGraphCompilationResult UOptimusDeformer::CompileNodeGraphToComputeGr
 		}
 	}
 
-	// Create all the data interfaces: node, graph, kernel outputs, loop terminal data
-	
-	// The component binding for the graph data is the primary binding on the deformer.
-	UOptimusComponentSourceBinding* GraphDataComponentBinding = Bindings->Bindings[0];
-	
-	TMap<const UComputeDataInterface*, int32> DataInterfaceToBindingIndexMap;
-	
-	// Find all data interface nodes and create their data interfaces.
-	TMap<FOptimusRoutedConstNode, UOptimusComputeDataInterface*> NodeDataInterfaceMap;
-	
 	// Find all value nodes (constant and variable) 
-	TArray<const UOptimusNode *> ValueNodes;
+	TArray<const UOptimusNode *> ActiveValueNodes;
 	TMap<const UOptimusNode*, EOptimusValueUsage> ValueNodeUsageMap;
-	
-	TMap<FOptimusRoutedConstNode, TArray<UOptimusComputeDataInterface*>> LoopEntryToLoopDataInterfaces;
-	
-	for (const FOptimusRoutedConstNode& ConnectedNode : ConnectedNodes)
+	TMap<FOptimusRoutedConstNode, FOptimusRoutedConstNode> ConstantNodeOverrideMap;
+
+	// Propagate usage and override info backwards
+	for (int32 Index = ConnectedNodes.Num() - 1; Index >= 0; Index--)
 	{
-		if (const IOptimusDataInterfaceProvider* NodeDataInterfaceProvider = Cast<const IOptimusDataInterfaceProvider>(ConnectedNode.Node))
+		const FOptimusRoutedConstNode& ConnectedNode = ConnectedNodes[Index];
+		const UOptimusNode* Node = ConnectedNode.Node;
+		if (Cast<IOptimusValueProvider>(Node))
 		{
-			// Gets a copy of node's data interface
-			UOptimusComputeDataInterface* DataInterface = NodeDataInterfaceProvider->GetDataInterface(this);
-			if (!DataInterface)
+			EOptimusValueUsage ValueUsage = EOptimusValueUsage::None;
+
+			if (Cast<UOptimusNode_ConstantValue>(Node))
 			{
-				AddDiagnostic(EOptimusDiagnosticLevel::Error, LOCTEXT("NoDataInterfaceOnProvider", "No data interface object returned from node. Compilation aborted."));
-				return {};
+				TArray<UOptimusNodePin*> InputPins = Node->GetPinsByDirection(EOptimusNodePinDirection::Input, true);
+				TArray<FOptimusRoutedNodePin> SourcePins = InputPins[0]->GetConnectedPinsWithRouting(ConnectedNode.TraversalContext);
+				if (SourcePins.Num() == 0)
+				{
+					// No overrider, this value node is active
+					if (!ActiveValueNodes.Contains(Node))
+					{
+						ActiveValueNodes.Add(Node);
+					}
+				}
+				else
+				{
+					// Save overrider info
+					FOptimusRoutedConstNode SourceNode = {SourcePins[0].NodePin->GetOwningNode(), SourcePins[0].TraversalContext};
+					ConstantNodeOverrideMap.Add(ConnectedNode, SourceNode);	
+				}
+			}
+			else
+			{
+				// No overrider, this value node is active	
+				if (!ActiveValueNodes.Contains(Node))
+				{
+					ActiveValueNodes.Add(Node);
+				}	
 			}
-
-			NodeDataInterfaceMap.Add(ConnectedNode, DataInterface);
-			DataInterfaceToBindingIndexMap.Add(DataInterface) = NodeDataInterfaceProvider->GetComponentBinding(ConnectedNode.TraversalContext)->GetIndex();
-		}
-		else if (Cast<const IOptimusValueProvider>(ConnectedNode.Node))
-		{
-			const UOptimusNode* Node = ConnectedNode.Node;
-			EOptimusValueUsage ValueUsage = EOptimusValueUsage::None;
 			
 			TArray<UOptimusNodePin*> OutputPins = Node->GetPinsByDirection(EOptimusNodePinDirection::Output, true);
 			check(OutputPins.Num() == 1);
@@ -2196,20 +2202,48 @@ FOptimusNodeGraphCompilationResult UOptimusDeformer::CompileNodeGraphToComputeGr
 					}
 				}
 
-				// at least one connection requesting this value on GPU
+				// Inherit usage from the nodes that this node overrides
+				if (Cast<IOptimusValueProvider>(OtherNode))
+				{
+					ValueUsage |= ValueNodeUsageMap[OtherNode];
+					continue;
+				}
+				
+				// At least one connection requesting this value on GPU
 				ValueUsage |= EOptimusValueUsage::GPU;
-				break;
 			}
 
 			EOptimusValueUsage& ValueUsageRef = ValueNodeUsageMap.FindOrAdd(Node);
 			ValueUsageRef |= ValueUsage;
-
-			
-			// Node may get used in a function that has multiple instances in this graph
-			if (!ValueNodes.Contains(Node))
+		}
+	}
+	
+	// Create all the data interfaces: node, graph, kernel outputs, loop terminal data
+	
+	// The component binding for the graph data is the primary binding on the deformer.
+	UOptimusComponentSourceBinding* GraphDataComponentBinding = Bindings->Bindings[0];
+	
+	TMap<const UComputeDataInterface*, int32> DataInterfaceToBindingIndexMap;
+	
+	// Find all data interface nodes and create their data interfaces.
+	TMap<FOptimusRoutedConstNode, UOptimusComputeDataInterface*> NodeDataInterfaceMap;
+	
+	TMap<FOptimusRoutedConstNode, TArray<UOptimusComputeDataInterface*>> LoopEntryToLoopDataInterfaces;
+	
+	for (const FOptimusRoutedConstNode& ConnectedNode : ConnectedNodes)
+	{
+		if (const IOptimusDataInterfaceProvider* NodeDataInterfaceProvider = Cast<const IOptimusDataInterfaceProvider>(ConnectedNode.Node))
+		{
+			// Gets a copy of node's data interface
+			UOptimusComputeDataInterface* DataInterface = NodeDataInterfaceProvider->GetDataInterface(this);
+			if (!DataInterface)
 			{
-				ValueNodes.Add(Node);
+				AddDiagnostic(EOptimusDiagnosticLevel::Error, LOCTEXT("NoDataInterfaceOnProvider", "No data interface object returned from node. Compilation aborted."));
+				return {};
 			}
+
+			NodeDataInterfaceMap.Add(ConnectedNode, DataInterface);
+			DataInterfaceToBindingIndexMap.Add(DataInterface) = NodeDataInterfaceProvider->GetComponentBinding(ConnectedNode.TraversalContext)->GetIndex();
 		}
 		else if (const UOptimusNode_LoopTerminal* LoopTerminal = Cast<const UOptimusNode_LoopTerminal>(ConnectedNode.Node))
 		{
@@ -2236,9 +2270,9 @@ FOptimusNodeGraphCompilationResult UOptimusDeformer::CompileNodeGraphToComputeGr
 	DataInterfaceToBindingIndexMap.Add(GraphDataInterface) = GraphDataComponentBinding->GetIndex();
 	
 	TArray<FOptimusGraphVariableDescription> ValueNodeDescriptions;
-	ValueNodeDescriptions.Reserve(ValueNodes.Num());	
+	ValueNodeDescriptions.Reserve(ActiveValueNodes.Num());	
 
-	for (const UOptimusNode* ValueNode : ValueNodes)
+	for (const UOptimusNode* ValueNode : ActiveValueNodes)
 	{
 		const IOptimusValueProvider* ValueProvider = CastChecked<IOptimusValueProvider>(ValueNode);
 		FOptimusValueIdentifier ValueId = ValueProvider->GetValueIdentifier();
@@ -2278,6 +2312,20 @@ FOptimusNodeGraphCompilationResult UOptimusDeformer::CompileNodeGraphToComputeGr
 	TMap<FOptimusInstancedPin, UOptimusComputeDataInterface*> KernelOutputDataInterfaceMap;
 
 	TMap<TWeakObjectPtr<const UComputeDataInterface>, FOptimusDataInterfacePropertyOverrideInfo> NodeGraphDataInterfacePropertyOverrideMap;
+
+	auto GetRootValueProviderPin = [&ConstantNodeOverrideMap](const FOptimusRoutedConstNode& InStartNode)
+	{
+		const FOptimusRoutedConstNode* WorkItem = &InStartNode;
+		while(const FOptimusRoutedConstNode* NextWorkItem = ConstantNodeOverrideMap.Find(*WorkItem))
+		{
+			WorkItem = NextWorkItem;
+		}
+
+		const TArray<UOptimusNodePin*> NewSourcePins = WorkItem->Node->GetPinsByDirection(EOptimusNodePinDirection::Output, true);
+		check(NewSourcePins.Num() == 1);
+
+		return FOptimusRoutedConstNodePin{NewSourcePins[0], WorkItem->TraversalContext};
+	};
 	
 	for (const FOptimusInstancedNode& InstancedNode : InstancedNodes )
 	{
@@ -2322,7 +2370,7 @@ FOptimusNodeGraphCompilationResult UOptimusDeformer::CompileNodeGraphToComputeGr
 					
 					if (Cast<const IOptimusValueProvider>(SourcePin->GetOwningNode()))
 					{
-						KernelInputMap[InstancedNode].Add(Pin) = {GraphDataInterface, SourcePin};
+						KernelInputMap[InstancedNode].Add(Pin) = {GraphDataInterface, GetRootValueProviderPin(SourceRoutedNode).NodePin};
 					}
 					else if (UOptimusComputeDataInterface** NodeDataInterface = NodeDataInterfaceMap.Find(SourceRoutedNode))
 					{
@@ -2462,12 +2510,16 @@ FOptimusNodeGraphCompilationResult UOptimusDeformer::CompileNodeGraphToComputeGr
 					{
 						const UOptimusNode* SourceNode = SourceInstancedPin->Pin->GetOwningNode();
 
-						if (const IOptimusValueProvider* ValueProvider = Cast<IOptimusValueProvider>(SourceNode))
+						if (Cast<IOptimusValueProvider>(SourceNode))
 						{
+							FOptimusRoutedConstNodePin RootValueProviderPin = GetRootValueProviderPin(SourceInstancedPin->InstancedNode.RoutedNode);
+
+							IOptimusValueProvider* RootValueProvider = CastChecked<IOptimusValueProvider>(RootValueProviderPin.NodePin->GetOwningNode());
+							
 							NodeGraphDataInterfacePropertyOverrideMap.FindOrAdd(ProviderDataInterface)
 									.PinNameToValueIdMap
 									.Add(Pin->GetFName(),
-										ValueProvider->GetValueIdentifier());	
+										RootValueProvider->GetValueIdentifier());	
 						}
 					}
 				}
