[Deformer graph] minor fix for value nodes with struct output pin (currently subpin should be ignored)

diff --git a/Engine/Plugins/Animation/DeformerGraph/Source/OptimusCore/Private/OptimusDeformer.cpp b/Engine/Plugins/Animation/DeformerGraph/Source/OptimusCore/Private/OptimusDeformer.cpp
--- a/Engine/Plugins/Animation/DeformerGraph/Source/OptimusCore/Private/OptimusDeformer.cpp
+++ b/Engine/Plugins/Animation/DeformerGraph/Source/OptimusCore/Private/OptimusDeformer.cpp
@@ -2182,7 +2182,7 @@ FOptimusNodeGraphCompilationResult UOptimusDeformer::CompileNodeGraphToComputeGr
 				}	
 			}
 			
-			TArray<UOptimusNodePin*> OutputPins = Node->GetPinsByDirection(EOptimusNodePinDirection::Output, true);
+			TArray<UOptimusNodePin*> OutputPins = Node->GetPinsByDirection(EOptimusNodePinDirection::Output, false);
 			check(OutputPins.Num() == 1);
 
 			const UOptimusNodePin* OutputPin = OutputPins[0];
@@ -2321,7 +2321,7 @@ FOptimusNodeGraphCompilationResult UOptimusDeformer::CompileNodeGraphToComputeGr
 			WorkItem = NextWorkItem;
 		}
 
-		const TArray<UOptimusNodePin*> NewSourcePins = WorkItem->Node->GetPinsByDirection(EOptimusNodePinDirection::Output, true);
+		const TArray<UOptimusNodePin*> NewSourcePins = WorkItem->Node->GetPinsByDirection(EOptimusNodePinDirection::Output, false);
 		check(NewSourcePins.Num() == 1);
 
 		return FOptimusRoutedConstNodePin{NewSourcePins[0], WorkItem->TraversalContext};
