[Deformer Graph] new data interface that allows for using skin weight as mask

diff --git a/Engine/Plugins/Animation/DeformerGraph/Shaders/Private/DataInterfaceAdvancedSkeleton.ush b/Engine/Plugins/Animation/DeformerGraph/Shaders/Private/DataInterfaceAdvancedSkeleton.ush
--- a/Engine/Plugins/Animation/DeformerGraph/Shaders/Private/DataInterfaceAdvancedSkeleton.ush
+++ b/Engine/Plugins/Animation/DeformerGraph/Shaders/Private/DataInterfaceAdvancedSkeleton.ush
@@ -1,6 +1,6 @@
 // Copyright Epic Games, Inc. All Rights Reserved.
 
-#include "/Engine/Private/GpuSkinCommon.ush"
+#include "/Plugin/Optimus/Private/DataInterfaceSkeletonCommon.ush"
 
 uint {DataInterfaceName}_NumVertices;
 uint {DataInterfaceName}_InputStreamStart;
@@ -19,135 +19,62 @@ uint ReadNumVertices_{DataInterfaceName}()
 	return {DataInterfaceName}_NumVertices;
 }
 
-#if !GPUSKIN_UNLIMITED_BONE_INFLUENCE
-
-int4 GetBlendIndices_{DataInterfaceName}(uint StreamOffset, uint ExtraInfluenceIndex)
-{
-	return GetBlendIndices({DataInterfaceName}_InputWeightStream, StreamOffset, ExtraInfluenceIndex);
-}
-
-float4 GetBlendWeights_{DataInterfaceName}(uint StreamOffset, uint ExtraInfluenceIndex)
-{
-	return GetBlendWeights( { DataInterfaceName }_InputWeightStream, StreamOffset, ExtraInfluenceIndex, { DataInterfaceName }_NumBoneInfluences);
-}
-
-#endif // GPUSKIN_UNLIMITED_BONE_INFLUENCE
-
-float3x4 GetBoneMatrix_{DataInterfaceName}(uint BoneIndex)
-{
-	uint BufferIndex = BoneIndex * 3;
-	return float3x4({DataInterfaceName}_BoneMatrices[BufferIndex], {DataInterfaceName}_BoneMatrices[BufferIndex + 1], {DataInterfaceName}_BoneMatrices[BufferIndex + 2]);
-}
-
 uint ReadNumBones_{DataInterfaceName}(uint VertexIndex)
 {
-#if !ENABLE_DEFORMER_BONES
-	return 1;
-#elif GPUSKIN_UNLIMITED_BONE_INFLUENCE
-	const uint BufferIndex 		= {DataInterfaceName}_InputStreamStart + VertexIndex;
-	uint BlendOffsetCount = {DataInterfaceName}_InputWeightLookupStream[BufferIndex];
-	return BlendOffsetCount & 0xff;
-#else // !GPUSKIN_UNLIMITED_BONE_INFLUENCE
-	return {DataInterfaceName}_NumBoneInfluences;
-#endif
+	return ReadNumBones(
+		{DataInterfaceName}_InputStreamStart,
+		{DataInterfaceName}_InputWeightLookupStream,
+		{DataInterfaceName}_NumBoneInfluences,
+		VertexIndex
+		);
 }
 
 float3x4 ReadBoneMatrix_{DataInterfaceName}(uint VertexIndex, uint BoneIndex)
 {
-#if !ENABLE_DEFORMER_BONES
-
-	return float3x4(1, 0, 0, 0, 0, 1, 0, 0, 0, 0, 1, 0);
-
-#elif GPUSKIN_UNLIMITED_BONE_INFLUENCE
-	const uint BufferIndex 		= {DataInterfaceName}_InputStreamStart + VertexIndex;
-	uint BlendOffsetCount = {DataInterfaceName}_InputWeightLookupStream[BufferIndex];
-	int StreamOffset = BlendOffsetCount >> 8;
-	int BoneIndexSize = {DataInterfaceName}_InputWeightIndexSize & 0xff;
-	int BoneIndexOffset = StreamOffset + (BoneIndexSize * BoneIndex);
-	int BoneGlobalIndex = {DataInterfaceName}_InputWeightStream[BoneIndexOffset];
-	if (BoneIndexSize > 1)
-	{
-		BoneGlobalIndex = {DataInterfaceName}_InputWeightStream[BoneIndexOffset + 1] << 8 | BoneGlobalIndex;
-	}
-	
-	return GetBoneMatrix_{DataInterfaceName}(BoneGlobalIndex);
-
-#else // !GPUSKIN_UNLIMITED_BONE_INFLUENCE
-	
-	uint StreamOffset = {DataInterfaceName}_InputWeightStart + (VertexIndex * ({DataInterfaceName}_InputWeightStride / 4));
-	int BlendIndex = GetBlendIndices_{DataInterfaceName}(StreamOffset, BoneIndex / 4)[BoneIndex & 0x3];
-	return GetBoneMatrix_{DataInterfaceName}(BlendIndex);
-
-#endif
+	return ReadBoneMatrix(
+		{DataInterfaceName}_InputStreamStart,
+		{DataInterfaceName}_InputWeightLookupStream,
+		{DataInterfaceName}_InputWeightIndexSize,
+		{DataInterfaceName}_InputWeightStream,
+		{DataInterfaceName}_InputWeightStart,
+		{DataInterfaceName}_InputWeightStride,
+		{DataInterfaceName}_BoneMatrices,
+		VertexIndex,
+		BoneIndex
+		);
 }
 
 float ReadBoneWeight_{DataInterfaceName}(uint VertexIndex, uint BoneIndex)
 {
-#if !ENABLE_DEFORMER_BONES
-
-	return 1;
-
-#elif GPUSKIN_UNLIMITED_BONE_INFLUENCE
-	const uint BufferIndex 		= {DataInterfaceName}_InputStreamStart + VertexIndex;	
-	uint BlendOffsetCount = {DataInterfaceName}_InputWeightLookupStream[BufferIndex];
-	int NumBoneInfluencesLocal = BlendOffsetCount & 0xff;
-	int StreamOffset = BlendOffsetCount >> 8;
-	int BoneIndexSize = {DataInterfaceName}_InputWeightIndexSize & 0xff;
-	int BoneWeightSize = {DataInterfaceName}_InputWeightIndexSize >> 8;
-	int WeightsOffset = StreamOffset + (BoneIndexSize * NumBoneInfluencesLocal);
-	int BoneWeightOffset = WeightsOffset + BoneIndex * BoneWeightSize;
-	
-	if (BoneWeightSize > 1)
-	{
-		return float({DataInterfaceName}_InputWeightStream[BoneWeightOffset + 1] << 8 | {DataInterfaceName}_InputWeightStream[BoneWeightOffset]) / 65535.0;
-	}
-	else
-	{
-		return float({DataInterfaceName}_InputWeightStream[BoneWeightOffset]) / 255.0f;
-	}
-
-#else // !GPUSKIN_UNLIMITED_BONE_INFLUENCE
-	
-	uint StreamOffset = {DataInterfaceName}_InputWeightStart + (VertexIndex * ({DataInterfaceName}_InputWeightStride / 4));
-	float BlendWeight = GetBlendWeights_{DataInterfaceName}(StreamOffset, BoneIndex / 4)[BoneIndex & 0x3];
-	return BlendWeight;
-
-#endif
+	return ReadBoneWeight(
+		{DataInterfaceName}_InputStreamStart,
+		{DataInterfaceName}_InputWeightLookupStream,
+		{DataInterfaceName}_InputWeightIndexSize,
+		{DataInterfaceName}_InputWeightStream,
+		{DataInterfaceName}_InputWeightStart,
+		{DataInterfaceName}_InputWeightStride,
+		{DataInterfaceName}_NumBoneInfluences,
+		VertexIndex,
+		BoneIndex
+		);
 }
 
 float3x4 ReadWeightedBoneMatrix_{DataInterfaceName}(uint VertexIndex)
 {
-#if !ENABLE_DEFORMER_BONES
-
-	return float3x4(1, 0, 0, 0, 0, 1, 0, 0, 0, 0, 1, 0);
-
-#elif GPUSKIN_UNLIMITED_BONE_INFLUENCE
-
-	const uint BufferIndex 		= {DataInterfaceName}_InputStreamStart + VertexIndex;
-	const uint BlendOffsetCount = {DataInterfaceName}_InputWeightLookupStream[BufferIndex];
-	return ComputeBoneMatrixWithUnlimitedInfluences({DataInterfaceName}_BoneMatrices, {DataInterfaceName}_InputWeightStream, {DataInterfaceName}_InputWeightIndexSize, BlendOffsetCount);
-
-#else
-
-	const uint StreamOffset = {DataInterfaceName}_InputWeightStart + (VertexIndex * ({DataInterfaceName}_InputWeightStride / 4));
-	const uint NumInfluences = {DataInterfaceName}_NumBoneInfluences;
-
-	FGPUSkinIndexAndWeight IndicesAndWeights = (FGPUSkinIndexAndWeight)0;
-	if (NumInfluences > 0)
-	{
-		IndicesAndWeights.BlendIndices  = GetBlendIndices_{DataInterfaceName}(StreamOffset, 0);
-		IndicesAndWeights.BlendWeights  = GetBlendWeights_{DataInterfaceName}(StreamOffset, 0);
-	}
-	if (NumInfluences > 4)
-	{
-		IndicesAndWeights.BlendIndices2 = GetBlendIndices_{DataInterfaceName}(StreamOffset, 1);
-		IndicesAndWeights.BlendWeights2 = GetBlendWeights_{DataInterfaceName}(StreamOffset, 1);
-	}
-	return ComputeBoneMatrixWithLimitedInfluences({DataInterfaceName}_BoneMatrices, IndicesAndWeights, NumInfluences > 4);
-
-#endif
+	return ReadWeightedBoneMatrix(
+		{DataInterfaceName}_InputStreamStart,
+		{DataInterfaceName}_InputWeightLookupStream,
+		{DataInterfaceName}_InputWeightIndexSize,
+		{DataInterfaceName}_InputWeightStream,
+		{DataInterfaceName}_InputWeightStart,
+		{DataInterfaceName}_InputWeightStride,
+		{DataInterfaceName}_NumBoneInfluences
+		{DataInterfaceName}_BoneMatrices,
+		VertexIndex
+		);
 }
 
+
 float3x4 ReadLayeredBoneMatrix_{DataInterfaceName}(uint VertexIndex, uint BoneIndex)
 {
 #if !ENABLE_DEFORMER_BONES
@@ -171,7 +98,7 @@ float3x4 ReadLayeredBoneMatrix_{DataInterfaceName}(uint VertexIndex, uint BoneIn
 #else // !GPUSKIN_UNLIMITED_BONE_INFLUENCE
 	
 	uint StreamOffset = {DataInterfaceName}_InputWeightStart + (VertexIndex * ({DataInterfaceName}_InputWeightStride / 4));
-	int BlendIndex = GetBlendIndices_{DataInterfaceName}(StreamOffset, BoneIndex / 4)[BoneIndex & 0x3];
+	int BlendIndex = GetBlendIndices({DataInterfaceName}_InputWeightStream, StreamOffset, BoneIndex / 4)[BoneIndex & 0x3];
 	return {DataInterfaceName}_LayeredBoneMatrices[BlendIndex];
 
 #endif
diff --git a/Engine/Plugins/Animation/DeformerGraph/Shaders/Private/DataInterfaceSkeleton.ush b/Engine/Plugins/Animation/DeformerGraph/Shaders/Private/DataInterfaceSkeleton.ush
--- a/Engine/Plugins/Animation/DeformerGraph/Shaders/Private/DataInterfaceSkeleton.ush
+++ b/Engine/Plugins/Animation/DeformerGraph/Shaders/Private/DataInterfaceSkeleton.ush
@@ -1,6 +1,6 @@
 // Copyright Epic Games, Inc. All Rights Reserved.
 
-#include "/Engine/Private/GpuSkinCommon.ush"
+#include "/Plugin/Optimus/Private/DataInterfaceSkeletonCommon.ush"
 
 uint {DataInterfaceName}_NumVertices;
 uint {DataInterfaceName}_InputStreamStart;
@@ -17,131 +17,57 @@ uint ReadNumVertices_{DataInterfaceName}()
 	return {DataInterfaceName}_NumVertices;
 }
 
-#if !GPUSKIN_UNLIMITED_BONE_INFLUENCE
-
-int4 GetBlendIndices_{DataInterfaceName}(uint StreamOffset, uint ExtraInfluenceIndex)
-{
-	return GetBlendIndices({DataInterfaceName}_InputWeightStream, StreamOffset, ExtraInfluenceIndex);
-}
-
-float4 GetBlendWeights_{DataInterfaceName}(uint StreamOffset, uint ExtraInfluenceIndex)
-{
-	return GetBlendWeights( { DataInterfaceName }_InputWeightStream, StreamOffset, ExtraInfluenceIndex, { DataInterfaceName }_NumBoneInfluences);
-}
-
-#endif // GPUSKIN_UNLIMITED_BONE_INFLUENCE
-
-float3x4 GetBoneMatrix_{DataInterfaceName}(uint BoneIndex)
-{
-	uint BufferIndex = BoneIndex * 3;
-	return float3x4({DataInterfaceName}_BoneMatrices[BufferIndex], {DataInterfaceName}_BoneMatrices[BufferIndex + 1], {DataInterfaceName}_BoneMatrices[BufferIndex + 2]);
-}
-
 uint ReadNumBones_{DataInterfaceName}(uint VertexIndex)
 {
-#if !ENABLE_DEFORMER_BONES
-	return 1;
-#elif GPUSKIN_UNLIMITED_BONE_INFLUENCE
-	const uint BufferIndex 		= {DataInterfaceName}_InputStreamStart + VertexIndex;
-	uint BlendOffsetCount = {DataInterfaceName}_InputWeightLookupStream[BufferIndex];
-	return BlendOffsetCount & 0xff;
-#else // !GPUSKIN_UNLIMITED_BONE_INFLUENCE
-	return {DataInterfaceName}_NumBoneInfluences;
-#endif
+	return ReadNumBones(
+		{DataInterfaceName}_InputStreamStart,
+		{DataInterfaceName}_InputWeightLookupStream,
+		{DataInterfaceName}_NumBoneInfluences,
+		VertexIndex
+		);
 }
 
 float3x4 ReadBoneMatrix_{DataInterfaceName}(uint VertexIndex, uint BoneIndex)
 {
-#if !ENABLE_DEFORMER_BONES
-
-	return float3x4(1, 0, 0, 0, 0, 1, 0, 0, 0, 0, 1, 0);
-
-#elif GPUSKIN_UNLIMITED_BONE_INFLUENCE
-	const uint BufferIndex 		= {DataInterfaceName}_InputStreamStart + VertexIndex;
-	uint BlendOffsetCount = {DataInterfaceName}_InputWeightLookupStream[BufferIndex];
-	int StreamOffset = BlendOffsetCount >> 8;
-	int BoneIndexSize = {DataInterfaceName}_InputWeightIndexSize & 0xff;
-	int BoneIndexOffset = StreamOffset + (BoneIndexSize * BoneIndex);
-	int BoneGlobalIndex = {DataInterfaceName}_InputWeightStream[BoneIndexOffset];
-	if (BoneIndexSize > 1)
-	{
-		BoneGlobalIndex = {DataInterfaceName}_InputWeightStream[BoneIndexOffset + 1] << 8 | BoneGlobalIndex;
-	}
-	
-	return GetBoneMatrix_{DataInterfaceName}(BoneGlobalIndex);
-
-#else // !GPUSKIN_UNLIMITED_BONE_INFLUENCE
-	
-	uint StreamOffset = {DataInterfaceName}_InputWeightStart + (VertexIndex * ({DataInterfaceName}_InputWeightStride / 4));
-	int BlendIndex = GetBlendIndices_{DataInterfaceName}(StreamOffset, BoneIndex / 4)[BoneIndex & 0x3];
-	return GetBoneMatrix_{DataInterfaceName}(BlendIndex);
-
-#endif
+	return ReadBoneMatrix(
+		{DataInterfaceName}_InputStreamStart,
+		{DataInterfaceName}_InputWeightLookupStream,
+		{DataInterfaceName}_InputWeightIndexSize,
+		{DataInterfaceName}_InputWeightStream,
+		{DataInterfaceName}_InputWeightStart,
+		{DataInterfaceName}_InputWeightStride,
+		{DataInterfaceName}_BoneMatrices,
+		VertexIndex,
+		BoneIndex
+		);
 }
 
 float ReadBoneWeight_{DataInterfaceName}(uint VertexIndex, uint BoneIndex)
 {
-#if !ENABLE_DEFORMER_BONES
-
-	return 1;
-
-#elif GPUSKIN_UNLIMITED_BONE_INFLUENCE
-	const uint BufferIndex 		= {DataInterfaceName}_InputStreamStart + VertexIndex;	
-	uint BlendOffsetCount = {DataInterfaceName}_InputWeightLookupStream[BufferIndex];
-	int NumBoneInfluencesLocal = BlendOffsetCount & 0xff;
-	int StreamOffset = BlendOffsetCount >> 8;
-	int BoneIndexSize = {DataInterfaceName}_InputWeightIndexSize & 0xff;
-	int BoneWeightSize = {DataInterfaceName}_InputWeightIndexSize >> 8;
-	int WeightsOffset = StreamOffset + (BoneIndexSize * NumBoneInfluencesLocal);
-	int BoneWeightOffset = WeightsOffset + BoneIndex * BoneWeightSize;
-	
-	if (BoneWeightSize > 1)
-	{
-		return float({DataInterfaceName}_InputWeightStream[BoneWeightOffset + 1] << 8 | {DataInterfaceName}_InputWeightStream[BoneWeightOffset]) / 65535.0;
-	}
-	else
-	{
-		return float({DataInterfaceName}_InputWeightStream[BoneWeightOffset]) / 255.0f;
-	}
-
-#else // !GPUSKIN_UNLIMITED_BONE_INFLUENCE
-	
-	uint StreamOffset = {DataInterfaceName}_InputWeightStart + (VertexIndex * ({DataInterfaceName}_InputWeightStride / 4));
-	float BlendWeight = GetBlendWeights_{DataInterfaceName}(StreamOffset, BoneIndex / 4)[BoneIndex & 0x3];
-	return BlendWeight;
-
-#endif
+	return ReadBoneWeight(
+		{DataInterfaceName}_InputStreamStart,
+		{DataInterfaceName}_InputWeightLookupStream,
+		{DataInterfaceName}_InputWeightIndexSize,
+		{DataInterfaceName}_InputWeightStream,
+		{DataInterfaceName}_InputWeightStart,
+		{DataInterfaceName}_InputWeightStride,
+		{DataInterfaceName}_NumBoneInfluences,	
+		VertexIndex,
+		BoneIndex
+		);
 }
 
 float3x4 ReadWeightedBoneMatrix_{DataInterfaceName}(uint VertexIndex)
 {
-#if !ENABLE_DEFORMER_BONES
-
-	return float3x4(1, 0, 0, 0, 0, 1, 0, 0, 0, 0, 1, 0);
-
-#elif GPUSKIN_UNLIMITED_BONE_INFLUENCE
-
-	const uint BufferIndex 		= {DataInterfaceName}_InputStreamStart + VertexIndex;
-	const uint BlendOffsetCount = {DataInterfaceName}_InputWeightLookupStream[BufferIndex];
-	return ComputeBoneMatrixWithUnlimitedInfluences({DataInterfaceName}_BoneMatrices, {DataInterfaceName}_InputWeightStream, {DataInterfaceName}_InputWeightIndexSize, BlendOffsetCount);
-
-#else
-
-	const uint StreamOffset = {DataInterfaceName}_InputWeightStart + (VertexIndex * ({DataInterfaceName}_InputWeightStride / 4));
-	const uint NumInfluences = {DataInterfaceName}_NumBoneInfluences;
-
-	FGPUSkinIndexAndWeight IndicesAndWeights = (FGPUSkinIndexAndWeight)0;
-	if (NumInfluences > 0)
-	{
-		IndicesAndWeights.BlendIndices  = GetBlendIndices_{DataInterfaceName}(StreamOffset, 0);
-		IndicesAndWeights.BlendWeights  = GetBlendWeights_{DataInterfaceName}(StreamOffset, 0);
-	}
-	if (NumInfluences > 4)
-	{
-		IndicesAndWeights.BlendIndices2 = GetBlendIndices_{DataInterfaceName}(StreamOffset, 1);
-		IndicesAndWeights.BlendWeights2 = GetBlendWeights_{DataInterfaceName}(StreamOffset, 1);
-	}
-	return ComputeBoneMatrixWithLimitedInfluences({DataInterfaceName}_BoneMatrices, IndicesAndWeights, NumInfluences > 4);
-
-#endif
+	return ReadWeightedBoneMatrix(
+		{DataInterfaceName}_InputStreamStart,
+		{DataInterfaceName}_InputWeightLookupStream,
+		{DataInterfaceName}_InputWeightIndexSize,
+		{DataInterfaceName}_InputWeightStream,
+		{DataInterfaceName}_InputWeightStart,
+		{DataInterfaceName}_InputWeightStride,
+		{DataInterfaceName}_NumBoneInfluences,
+		{DataInterfaceName}_BoneMatrices,
+		VertexIndex
+		);
 }
diff --git a/Engine/Plugins/Animation/DeformerGraph/Shaders/Private/DataInterfaceSkeletonCommon.ush b/Engine/Plugins/Animation/DeformerGraph/Shaders/Private/DataInterfaceSkeletonCommon.ush
new file mode 100644
--- /dev/null
+++ b/Engine/Plugins/Animation/DeformerGraph/Shaders/Private/DataInterfaceSkeletonCommon.ush
@@ -0,0 +1,153 @@
+// Copyright Epic Games, Inc. All Rights Reserved.
+#pragma once
+#include "/Engine/Private/GpuSkinCommon.ush"
+
+uint ReadNumBones(
+	uint InputStreamStart,
+	Buffer<uint> InputWeightLookupStream,
+	uint NumBoneInfluences,
+	uint VertexIndex
+	)
+{
+#if !ENABLE_DEFORMER_BONES
+	return 1;
+#elif GPUSKIN_UNLIMITED_BONE_INFLUENCE
+	const uint BufferIndex 		= InputStreamStart + VertexIndex;
+	uint BlendOffsetCount = InputWeightLookupStream[BufferIndex];
+	return BlendOffsetCount & 0xff;
+#else // !GPUSKIN_UNLIMITED_BONE_INFLUENCE
+	return NumBoneInfluences;
+#endif
+}
+
+float ReadBoneWeight(
+	uint InputStreamStart,
+	Buffer<uint> InputWeightLookupStream,
+	uint InputWeightIndexSize,
+	Buffer<uint> InputWeightStream,
+	uint InputWeightStart,
+	uint InputWeightStride,
+	uint NumBoneInfluences,
+	uint VertexIndex,
+	uint BoneIndex
+	)
+{
+#if !ENABLE_DEFORMER_BONES
+
+	return 1;
+
+#elif GPUSKIN_UNLIMITED_BONE_INFLUENCE
+	const uint BufferIndex 		= InputStreamStart + VertexIndex;	
+	uint BlendOffsetCount = InputWeightLookupStream[BufferIndex];
+	int NumBoneInfluencesLocal = BlendOffsetCount & 0xff;
+	int StreamOffset = BlendOffsetCount >> 8;
+	int BoneIndexSize = InputWeightIndexSize & 0xff;
+	int BoneWeightSize = InputWeightIndexSize >> 8;
+	int WeightsOffset = StreamOffset + (BoneIndexSize * NumBoneInfluencesLocal);
+	int BoneWeightOffset = WeightsOffset + BoneIndex * BoneWeightSize;
+	
+	if (BoneWeightSize > 1)
+	{
+		return float(InputWeightStream[BoneWeightOffset + 1] << 8 | InputWeightStream[BoneWeightOffset]) / 65535.0;
+	}
+	else
+	{
+		return float(InputWeightStream[BoneWeightOffset]) / 255.0f;
+	}
+
+#else // !GPUSKIN_UNLIMITED_BONE_INFLUENCE
+	
+	uint StreamOffset = InputWeightStart + (VertexIndex * (InputWeightStride / 4));
+	float BlendWeight = GetBlendWeights(InputWeightStream, StreamOffset, BoneIndex / 4, NumBoneInfluences)[BoneIndex & 0x3];
+	return BlendWeight;
+
+#endif
+}
+
+float3x4 GetBoneMatrix(Buffer<float4> BoneMatrices, uint BoneIndex)
+{
+	uint BufferIndex = BoneIndex * 3;
+	return float3x4(BoneMatrices[BufferIndex], BoneMatrices[BufferIndex + 1], BoneMatrices[BufferIndex + 2]);
+}
+
+float3x4 ReadBoneMatrix(
+	uint InputStreamStart,
+	Buffer<uint> InputWeightLookupStream,
+	uint InputWeightIndexSize,
+	Buffer<uint> InputWeightStream,
+	uint InputWeightStart,
+	uint InputWeightStride,
+	Buffer<float4> BoneMatrices,	
+	uint VertexIndex,
+	uint BoneIndex
+	)
+{
+#if !ENABLE_DEFORMER_BONES
+
+	return float3x4(1, 0, 0, 0, 0, 1, 0, 0, 0, 0, 1, 0);
+
+#elif GPUSKIN_UNLIMITED_BONE_INFLUENCE
+	const uint BufferIndex 		= InputStreamStart + VertexIndex;
+	uint BlendOffsetCount = InputWeightLookupStream[BufferIndex];
+	int StreamOffset = BlendOffsetCount >> 8;
+	int BoneIndexSize = InputWeightIndexSize & 0xff;
+	int BoneIndexOffset = StreamOffset + (BoneIndexSize * BoneIndex);
+	int BoneGlobalIndex = InputWeightStream[BoneIndexOffset];
+	if (BoneIndexSize > 1)
+	{
+		BoneGlobalIndex = InputWeightStream[BoneIndexOffset + 1] << 8 | BoneGlobalIndex;
+	}
+	
+	return GetBoneMatrix(BoneMatrices, BoneGlobalIndex);
+
+#else // !GPUSKIN_UNLIMITED_BONE_INFLUENCE
+	
+	uint StreamOffset = InputWeightStart + (VertexIndex * (InputWeightStride / 4));
+	int BlendIndex = GetBlendIndices(InputWeightStream ,StreamOffset, BoneIndex / 4)[BoneIndex & 0x3];
+	return GetBoneMatrix(BoneMatrices,BlendIndex);
+
+#endif
+}
+
+float3x4 ReadWeightedBoneMatrix(
+	uint InputStreamStart,
+	Buffer<uint> InputWeightLookupStream,
+	uint InputWeightIndexSize,
+	Buffer<uint> InputWeightStream,
+	uint InputWeightStart,
+	uint InputWeightStride,
+	uint NumBoneInfluences,
+	Buffer<float4> BoneMatrices,		
+	uint VertexIndex
+	)
+{
+#if !ENABLE_DEFORMER_BONES
+
+	return float3x4(1, 0, 0, 0, 0, 1, 0, 0, 0, 0, 1, 0);
+
+#elif GPUSKIN_UNLIMITED_BONE_INFLUENCE
+
+	const uint BufferIndex 		= InputStreamStart + VertexIndex;
+	const uint BlendOffsetCount = InputWeightLookupStream[BufferIndex];
+	return ComputeBoneMatrixWithUnlimitedInfluences(BoneMatrices, InputWeightStream, InputWeightIndexSize, BlendOffsetCount);
+
+#else
+
+	const uint StreamOffset = InputWeightStart + (VertexIndex * (InputWeightStride / 4));
+	const uint NumInfluences = NumBoneInfluences;
+
+	FGPUSkinIndexAndWeight IndicesAndWeights = (FGPUSkinIndexAndWeight)0;
+	if (NumInfluences > 0)
+	{
+		IndicesAndWeights.BlendIndices  = GetBlendIndices(InputWeightStream, StreamOffset, 0);
+		IndicesAndWeights.BlendWeights  = GetBlendWeights(InputWeightStream, StreamOffset, 0, NumInfluences);
+	}
+	if (NumInfluences > 4)
+	{
+		IndicesAndWeights.BlendIndices2 = GetBlendIndices(InputWeightStream, StreamOffset, 1);
+		IndicesAndWeights.BlendWeights2 = GetBlendWeights(InputWeightStream, StreamOffset, 1, NumInfluences);
+	}
+	return ComputeBoneMatrixWithLimitedInfluences(BoneMatrices, IndicesAndWeights, NumInfluences > 4);
+
+#endif
+}
diff --git a/Engine/Plugins/Animation/DeformerGraph/Shaders/Private/DataInterfaceSkinWeightsAsVertexMask.ush b/Engine/Plugins/Animation/DeformerGraph/Shaders/Private/DataInterfaceSkinWeightsAsVertexMask.ush
new file mode 100644
--- /dev/null
+++ b/Engine/Plugins/Animation/DeformerGraph/Shaders/Private/DataInterfaceSkinWeightsAsVertexMask.ush
@@ -0,0 +1,98 @@
+// Copyright Epic Games, Inc. All Rights Reserved.
+
+#include "/Plugin/Optimus/Private/DataInterfaceSkeletonCommon.ush"
+
+uint {DataInterfaceName}_NumVertices;
+uint {DataInterfaceName}_InputStreamStart;
+uint {DataInterfaceName}_NumBoneInfluences;
+uint {DataInterfaceName}_InputWeightStart;
+uint {DataInterfaceName}_InputWeightStride;
+uint {DataInterfaceName}_InputWeightIndexSize;
+Buffer<uint> {DataInterfaceName}_InputWeightStream;
+Buffer<uint> {DataInterfaceName}_InputWeightLookupStream;
+
+StructuredBuffer<uint> {DataInterfaceName}_BoneIsSelected;
+
+uint ReadNumVertices_{DataInterfaceName}()
+{
+	return {DataInterfaceName}_NumVertices;
+}
+
+uint ReadNumBones_{DataInterfaceName}(uint VertexIndex)
+{
+	return ReadNumBones(
+		{DataInterfaceName}_InputStreamStart,
+		{DataInterfaceName}_InputWeightLookupStream,
+		{DataInterfaceName}_NumBoneInfluences,
+		VertexIndex
+		);
+}
+
+float ReadBoneWeight_{DataInterfaceName}(uint VertexIndex, uint BoneIndex)
+{
+	return ReadBoneWeight(
+		{DataInterfaceName}_InputStreamStart,
+		{DataInterfaceName}_InputWeightLookupStream,
+		{DataInterfaceName}_InputWeightIndexSize,
+		{DataInterfaceName}_InputWeightStream,
+		{DataInterfaceName}_InputWeightStart,
+		{DataInterfaceName}_InputWeightStride,
+		{DataInterfaceName}_NumBoneInfluences,	
+		VertexIndex,
+		BoneIndex
+		);
+}
+
+uint ReadBoneIsSelected_{DataInterfaceName}(uint VertexIndex, uint BoneIndex)
+{
+#if !ENABLE_DEFORMER_BONES
+
+	return 0;
+
+#elif GPUSKIN_UNLIMITED_BONE_INFLUENCE
+	const uint BufferIndex 		= {DataInterfaceName}_InputStreamStart + VertexIndex;
+	uint BlendOffsetCount = {DataInterfaceName}_InputWeightLookupStream[BufferIndex];
+	int StreamOffset = BlendOffsetCount >> 8;
+	int BoneIndexSize = {DataInterfaceName}_InputWeightIndexSize & 0xff;
+	int BoneIndexOffset = StreamOffset + (BoneIndexSize * BoneIndex);
+	int BoneGlobalIndex = {DataInterfaceName}_InputWeightStream[BoneIndexOffset];
+	if (BoneIndexSize > 1)
+	{
+		BoneGlobalIndex = {DataInterfaceName}_InputWeightStream[BoneIndexOffset + 1] << 8 | BoneGlobalIndex;
+	}
+	
+	return {DataInterfaceName}_BoneIsSelected[BoneGlobalIndex];
+
+#else // !GPUSKIN_UNLIMITED_BONE_INFLUENCE
+	
+	uint StreamOffset = {DataInterfaceName}_InputWeightStart + (VertexIndex * ({DataInterfaceName}_InputWeightStride / 4));
+	int BlendIndex = GetBlendIndices({DataInterfaceName}_InputWeightStream, StreamOffset, BoneIndex / 4)[BoneIndex & 0x3];
+	return {DataInterfaceName}_BoneIsSelected[BlendIndex];
+
+#endif
+}
+
+
+float ReadMask_{DataInterfaceName}(uint VertexIndex)
+{
+#if !ENABLE_DEFORMER_BONES
+
+	return 0;
+
+#else
+
+	float MaskValue = 0;
+	
+	uint NumBones = ReadNumBones_{DataInterfaceName}(VertexIndex);
+	
+	uint BoneIndex = 0;
+	for (BoneIndex = 0; BoneIndex < NumBones; BoneIndex++)
+	{
+		float BoneWeight = ReadBoneWeight_{DataInterfaceName}(VertexIndex, BoneIndex);
+		uint BoneIsSelected = ReadBoneIsSelected_{DataInterfaceName}(VertexIndex, BoneIndex);
+		MaskValue += BoneIsSelected ? BoneWeight : 0;		
+	}
+
+	return MaskValue;
+#endif
+}
diff --git a/Engine/Plugins/Animation/DeformerGraph/Source/OptimusCore/Private/DataInterfaces/OptimusDataInterfaceSkinWeightsAsVertexMask.cpp b/Engine/Plugins/Animation/DeformerGraph/Source/OptimusCore/Private/DataInterfaces/OptimusDataInterfaceSkinWeightsAsVertexMask.cpp
new file mode 100644
--- /dev/null
+++ b/Engine/Plugins/Animation/DeformerGraph/Source/OptimusCore/Private/DataInterfaces/OptimusDataInterfaceSkinWeightsAsVertexMask.cpp
@@ -0,0 +1,568 @@
+// Copyright Epic Games, Inc. All Rights Reserved.
+
+#include "OptimusDataInterfaceSkinWeightsAsVertexMask.h"
+
+#include "OptimusDataDomain.h"
+#include "OptimusDataTypeRegistry.h"
+#include "OptimusDeformerInstance.h"
+#include "ShaderParameterMetadataBuilder.h"
+#include "SkeletalMeshDeformerHelpers.h"
+#include "SkeletalRenderPublic.h"
+#include "SystemTextures.h"
+#include "ComputeFramework/ComputeKernelPermutationVector.h"
+#include "Nodes/OptimusNode_DataInterface.h"
+#include "Rendering/SkeletalMeshRenderData.h"
+#include "Animation/SkinWeightProfileManager.h"
+#include "Components/SkeletalMeshComponent.h"
+
+FString UOptimusSkinWeightsAsVertexMaskDataInterface::GetDisplayName() const
+{
+	return TEXT("Skin Weights as Vertex Mask");
+}
+
+TArray<FOptimusCDIPinDefinition> UOptimusSkinWeightsAsVertexMaskDataInterface::GetPinDefinitions() const
+{
+	TArray<FOptimusCDIPinDefinition> Defs;
+	Defs.Add({"Mask", "ReadMask", Optimus::DomainName::Vertex, "ReadNumVertices" });
+
+	return Defs;
+}
+
+TArray<FOptimusCDIPropertyPinDefinition> UOptimusSkinWeightsAsVertexMaskDataInterface::GetPropertyPinDefinitions() const
+{
+	TArray<FOptimusCDIPropertyPinDefinition> Defs;
+
+	const FOptimusDataTypeHandle NameType = FOptimusDataTypeRegistry::Get().FindType(*FNameProperty::StaticClass());
+	const FOptimusDataTypeHandle NameArrayType = FOptimusDataTypeRegistry::Get().FindArrayType(*FNameProperty::StaticClass());
+	const FOptimusDataTypeHandle IntType = FOptimusDataTypeRegistry::Get().FindType(*FIntProperty::StaticClass());
+	const FOptimusDataTypeHandle BoolType = FOptimusDataTypeRegistry::Get().FindType(*FBoolProperty::StaticClass());
+	
+	Defs.Add({SkinWeightProfilePropertyName, NameType});
+	Defs.Add({BoneNamesPropertyName, NameArrayType});
+	Defs.Add({ExpandTowardsRootPropertyName, IntType});
+	Defs.Add({ExpandTowardsLeafPropertyName, IntType});
+	Defs.Add({DebugDrawIncludedBonesPropertyName, BoolType});
+
+	return Defs;
+}
+
+TSubclassOf<UActorComponent> UOptimusSkinWeightsAsVertexMaskDataInterface::GetRequiredComponentClass() const
+{
+	return USkeletalMeshComponent::StaticClass();
+}
+
+void UOptimusSkinWeightsAsVertexMaskDataInterface::RegisterPropertyChangeDelegatesForOwningNode(UOptimusNode_DataInterface* InNode)
+{
+	OnPinDefinitionChangedDelegate.BindUObject(InNode, &UOptimusNode_DataInterface::RecreatePinsFromPinDefinitions);
+	OnPinDefinitionRenamedDelegate.BindUObject(InNode, &UOptimusNode_DataInterface::RenamePinFromPinDefinition);
+	OnDisplayNameChangedDelegate.BindUObject(InNode, &UOptimusNode_DataInterface::UpdateDisplayNameFromDataInterface);
+}
+
+void UOptimusSkinWeightsAsVertexMaskDataInterface::GetSupportedInputs(TArray<FShaderFunctionDefinition>& OutFunctions) const
+{
+	OutFunctions.AddDefaulted_GetRef()
+		.SetName(TEXT("ReadNumVertices"))
+		.AddReturnType(EShaderFundamentalType::Uint);
+
+	OutFunctions.AddDefaulted_GetRef()
+		.SetName(TEXT("ReadMask"))
+		.AddReturnType(EShaderFundamentalType::Float)
+		.AddParam(EShaderFundamentalType::Uint);
+}
+
+BEGIN_SHADER_PARAMETER_STRUCT(FSkinWeightsAsVertexMaskDataInterfaceParameters, )
+	SHADER_PARAMETER(uint32, NumVertices)
+	SHADER_PARAMETER(uint32, InputStreamStart)
+	SHADER_PARAMETER(uint32, NumBoneInfluences)
+	SHADER_PARAMETER(uint32, InputWeightStart)
+	SHADER_PARAMETER(uint32, InputWeightStride)
+	SHADER_PARAMETER(uint32, InputWeightIndexSize)
+	SHADER_PARAMETER_SRV(Buffer<uint>, InputWeightStream)
+	SHADER_PARAMETER_SRV(Buffer<uint>, InputWeightLookupStream)
+
+	SHADER_PARAMETER_RDG_BUFFER_SRV(StructuredBuffer<uint32>, BoneIsSelected)
+END_SHADER_PARAMETER_STRUCT()
+
+void UOptimusSkinWeightsAsVertexMaskDataInterface::GetShaderParameters(TCHAR const* UID, FShaderParametersMetadataBuilder& InOutBuilder,
+	FShaderParametersMetadataAllocations& InOutAllocations) const
+{
+	InOutBuilder.AddNestedStruct<FSkinWeightsAsVertexMaskDataInterfaceParameters>(UID);
+}
+
+void UOptimusSkinWeightsAsVertexMaskDataInterface::GetPermutations(FComputeKernelPermutationVector& OutPermutationVector) const
+{
+	// Need to be able to support these permutations according to the skeletal mesh settings.
+	// todo[CF]: I think GPUSKIN_UNLIMITED_BONE_INFLUENCE and GPUSKIN_BONE_INDEX_UINT16/GPUSKIN_BONE_WEIGHTS_UINT16 are mutually exclusive. So we could save permutations here.
+	OutPermutationVector.AddPermutation(TEXT("ENABLE_DEFORMER_BONES"), 2);
+	OutPermutationVector.AddPermutation(TEXT("GPUSKIN_UNLIMITED_BONE_INFLUENCE"), 2);
+	OutPermutationVector.AddPermutation(TEXT("GPUSKIN_BONE_INDEX_UINT16"), 2);
+	OutPermutationVector.AddPermutation(TEXT("GPUSKIN_BONE_WEIGHTS_UINT16"), 2);
+}
+
+TCHAR const* UOptimusSkinWeightsAsVertexMaskDataInterface::TemplateFilePath = TEXT("/Plugin/Optimus/Private/DataInterfaceSkinWeightsAsVertexMask.ush");
+
+void UOptimusSkinWeightsAsVertexMaskDataInterface::GetShaderHash(FString& InOutKey) const
+{
+	return GetShaderFileHash(TemplateFilePath, EShaderPlatform::SP_PCD3D_SM5).AppendString(InOutKey);
+}
+
+TCHAR const* UOptimusSkinWeightsAsVertexMaskDataInterface::GetShaderVirtualPath() const
+{
+	return TemplateFilePath;	
+}
+
+void UOptimusSkinWeightsAsVertexMaskDataInterface::GetHLSL(FString& OutHLSL, FString const& InDataInterfaceName) const
+{
+	TMap<FString, FStringFormatArg> TemplateArgs =
+	{
+		{ TEXT("DataInterfaceName"), InDataInterfaceName },
+	};
+
+	FString TemplateFile;
+	LoadShaderSourceFile(TemplateFilePath, EShaderPlatform::SP_PCD3D_SM5, &TemplateFile, nullptr);
+	OutHLSL += FString::Format(*TemplateFile, TemplateArgs);
+}
+
+UComputeDataProvider* UOptimusSkinWeightsAsVertexMaskDataInterface::CreateDataProvider(TObjectPtr<UObject> InBinding, uint64 InInputMask,
+	uint64 InOutputMask) const
+{
+	UOptimusSkinWeightsAsVertexMaskDataProvider* Provider = NewObject<UOptimusSkinWeightsAsVertexMaskDataProvider>();
+	Provider->Init(this, Cast<USkeletalMeshComponent>(InBinding));
+	return Provider;
+}
+
+
+void UOptimusSkinWeightsAsVertexMaskDataProvider::SetDeformerInstance(UOptimusDeformerInstance* InInstance)
+{
+	DeformerInstance = InInstance;
+}
+
+void UOptimusSkinWeightsAsVertexMaskDataProvider::Init(const UOptimusSkinWeightsAsVertexMaskDataInterface* InDataInterface, USkeletalMeshComponent* InSkeletalMesh)
+{
+	WeakDataInterface = InDataInterface;
+	SkeletalMesh = InSkeletalMesh;
+
+	SkinWeightProfile = InDataInterface->SkinWeightProfile;
+	BoneNames = InDataInterface->BoneNames;
+	ExpandTowardsRoot = InDataInterface->ExpandTowardsRoot;
+	ExpandTowardsLeaf = InDataInterface->ExpandTowardsLeaf;
+	bDebugDrawIncludedBones = InDataInterface->bDebugDrawIncludedBones;
+	DebugDrawColor = InDataInterface->DebugDrawColor;
+}
+
+FComputeDataProviderRenderProxy* UOptimusSkinWeightsAsVertexMaskDataProvider::GetRenderProxy()
+{
+	FOptimusSkinWeightsAsVertexMaskDataProviderProxy* Proxy = new FOptimusSkinWeightsAsVertexMaskDataProviderProxy();
+
+	if (SkeletalMesh && SkeletalMesh->MeshObject)
+	{
+
+		
+		const UOptimusSkinWeightsAsVertexMaskDataInterface* DataInterface = WeakDataInterface.Get();
+		
+		{
+			const FOptimusValueContainerStruct ValueContainer =
+				DeformerInstance->GetDataInterfacePropertyOverride( DataInterface, DataInterface->SkinWeightProfilePropertyName );
+
+			TValueOrError<FName, EPropertyBagResult> Value = ValueContainer.Value.GetValueName(FOptimusValueContainerStruct::ValuePropertyName);
+			if (Value.HasValue())
+			{
+				SkinWeightProfile = Value.GetValue();
+			}
+		}
+
+		{
+			const FOptimusValueContainerStruct ValueContainer =
+				DeformerInstance->GetDataInterfacePropertyOverride( DataInterface, DataInterface->BoneNamesPropertyName);
+
+			TValueOrError<const FPropertyBagArrayRef, EPropertyBagResult> Value = ValueContainer.Value.GetArrayRef(FOptimusValueContainerStruct::ValuePropertyName);
+			if (Value.HasValue())
+			{
+				TArray<FName> NewBoneNames;
+				NewBoneNames.SetNum(Value.GetValue().Num());
+				for (int32 Index = 0 ; Index < NewBoneNames.Num(); Index++)
+				{
+					NewBoneNames[Index] = Value.GetValue().GetValueName(Index).GetValue();
+				}
+
+				if (NewBoneNames != BoneNames)
+				{
+					bIsInitialized = false;
+					BoneNames = MoveTemp(NewBoneNames);
+				}
+			}
+		}
+
+		{
+			const FOptimusValueContainerStruct ValueContainer =
+				DeformerInstance->GetDataInterfacePropertyOverride( DataInterface, DataInterface->ExpandTowardsRootPropertyName);
+
+			TValueOrError<int32 , EPropertyBagResult> Value = ValueContainer.Value.GetValueInt32(FOptimusValueContainerStruct::ValuePropertyName);
+			if (Value.HasValue())
+			{
+				if (ExpandTowardsRoot != Value.GetValue())
+				{
+					bIsInitialized = false;
+					ExpandTowardsRoot = Value.GetValue();
+				}
+			}
+		}
+
+		{
+			const FOptimusValueContainerStruct ValueContainer =
+				DeformerInstance->GetDataInterfacePropertyOverride( DataInterface, DataInterface->ExpandTowardsLeafPropertyName);
+
+			TValueOrError<int32 , EPropertyBagResult> Value = ValueContainer.Value.GetValueInt32(FOptimusValueContainerStruct::ValuePropertyName);
+			if (Value.HasValue())
+			{
+				if (ExpandTowardsLeaf != Value.GetValue())
+				{
+					bIsInitialized = false;
+					ExpandTowardsLeaf = Value.GetValue();
+				}
+			}
+		}
+
+		{
+			const FOptimusValueContainerStruct ValueContainer =
+				DeformerInstance->GetDataInterfacePropertyOverride( DataInterface, DataInterface->DebugDrawIncludedBonesPropertyName);
+
+			TValueOrError<bool, EPropertyBagResult> Value = ValueContainer.Value.GetValueBool(FOptimusValueContainerStruct::ValuePropertyName);
+			if (Value.HasValue())
+			{
+				if (bDebugDrawIncludedBones != Value.GetValue())
+				{
+					bDebugDrawIncludedBones = Value.GetValue();
+				}
+			}
+		}
+
+		ExpandTowardsRoot = FMath::Max(0, ExpandTowardsRoot);
+		ExpandTowardsLeaf = FMath::Max(0, ExpandTowardsLeaf);
+		
+		FSkeletalMeshObject* SkeletalMeshObject = SkeletalMesh->MeshObject;
+		const int32 LodIndex = SkeletalMeshObject->GetLOD();
+		FSkeletalMeshRenderData const& SkeletalMeshRenderData = SkeletalMeshObject->GetSkeletalMeshRenderData();
+		FSkeletalMeshLODRenderData const& LodRenderData = SkeletalMeshRenderData.LODRenderData[LodIndex];
+
+		bool bSkinWeightBufferReady = false;
+
+		if (SkinWeightProfile.IsNone())
+		{
+			bSkinWeightBufferReady = true;
+		}
+		else if (LodRenderData.SkinWeightProfilesData.ContainsProfile(SkinWeightProfile))
+		{
+			// Retrieve this profile's skin weight buffer
+			const FSkinWeightProfileStack ProfileStack{SkinWeightProfile};
+			FSkinWeightVertexBuffer* Buffer = LodRenderData.SkinWeightProfilesData.GetOverrideBuffer(ProfileStack);
+			bSkinWeightBufferReady = Buffer ? true : false;
+		}
+
+		if (!bSkinWeightBufferReady)
+		{
+			if (LodRenderData.SkinWeightProfilesData.ContainsProfile(SkinWeightProfile))
+			{
+				const FSkinWeightProfileStack ProfileStack{SkinWeightProfile};
+				// Put in a skin weight profile request
+				if (FSkinWeightProfileManager* Manager = FSkinWeightProfileManager::Get(SkeletalMesh->GetWorld()))
+				{
+					FRequestFinished DummyCallback = [](TWeakObjectPtr<USkeletalMesh> WeakMesh, FSkinWeightProfileStack ProfileStack){};
+					Manager->RequestSkinWeightProfileStack(ProfileStack, SkeletalMesh->GetSkinnedAsset(), this, DummyCallback);
+				}		
+			}
+		}
+		else
+		{
+			if (!bIsInitialized)
+			{
+				bIsInitialized = true;
+
+				CachedSelectedBones.Reset();
+				CachedBoneIsSelectedPerSection.Reset();
+				
+				const FReferenceSkeleton& RefSkeleton = SkeletalMesh->GetSkinnedAsset()->GetRefSkeleton();
+
+
+				// Negative distance : Towards Root
+				// Positive distance : Towards Leaf
+				TArray<TOptional<int32>> BoneDistance;
+				BoneDistance.SetNum(RefSkeleton.GetNum());
+
+				// Set distance for all parents of the seed bones
+				for (const FName& BoneName : BoneNames)
+				{
+					int32 Index = RefSkeleton.FindBoneIndex(BoneName);
+
+					int32 Distance = 0;
+					while (Index != INDEX_NONE && FMath::Abs(Distance) <= ExpandTowardsRoot)
+					{
+						if (BoneDistance[Index].IsSet() && BoneDistance[Index].GetValue() <= Distance)
+						{
+							// Skip if the rest of this chain has been visited by a closer seed bone
+							break;
+						}
+							
+						CachedSelectedBones.Add(Index);
+						BoneDistance[Index] = Distance;
+						Distance--;
+						Index = RefSkeleton.GetParentIndex(Index);
+					}
+				}
+
+				// Indices are sorted from parent to children, so distance value can flood from parent to children
+				for (int32 Index = 0; Index < RefSkeleton.GetNum(); Index++)
+				{
+					if (!BoneDistance[Index].IsSet())
+					{
+						int32 Parent = RefSkeleton.GetParentIndex(Index);
+						if (Parent != INDEX_NONE)
+						{
+							if (BoneDistance[Parent].IsSet())
+							{
+								int32 ParentDist = BoneDistance[Parent].GetValue();
+								
+								// Seed bone has to get to this bone by going towards the root
+								if (ParentDist < 0)
+								{
+									if (FMath::Abs(ParentDist) < ExpandTowardsRoot)
+									{
+										BoneDistance[Index] = ParentDist - 1;
+										CachedSelectedBones.Add(Index);
+									}
+								}
+								else
+								{
+									// Seed bone has to get to this bone by going towards the leaf 
+									if (ParentDist < ExpandTowardsLeaf)
+									{
+										BoneDistance[Index] = ParentDist + 1;
+										CachedSelectedBones.Add(Index);
+									}	
+								}
+							}
+						}
+					}
+				}
+				
+				for (int32 SectionIndex = 0; SectionIndex < LodRenderData.RenderSections.Num(); ++SectionIndex)
+				{
+					TArray<uint32>& BoneIsSelectedRef = CachedBoneIsSelectedPerSection.AddDefaulted_GetRef();
+				
+					FSkelMeshRenderSection const& RenderSection = LodRenderData.RenderSections[SectionIndex];
+					const int32 NumBones = RenderSection.BoneMap.Num();
+					BoneIsSelectedRef.Init(false, NumBones);
+
+					for (int32 BoneIndex = 0; BoneIndex < NumBones; BoneIndex++)
+					{
+						int32 FinalBoneIndex = RenderSection.BoneMap[BoneIndex];
+						if (CachedSelectedBones.Contains(FinalBoneIndex))
+						{
+							BoneIsSelectedRef[BoneIndex] = true;
+						}
+					}
+				}	
+			}
+
+			if (bDebugDrawIncludedBones)
+			{
+				FFunctionGraphTask::CreateAndDispatchWhenReady([ WeakThis = TWeakObjectPtr(this)]()
+					{
+						if (WeakThis.IsValid())
+						{
+							USkeletalMeshComponent* Mesh = WeakThis->SkeletalMesh;
+							const FReferenceSkeleton& RefSkeleton = Mesh->GetSkinnedAsset()->GetRefSkeleton();
+					
+							const TArray<FTransform>& CurrentComponentSpaceBoneMatrix = Mesh->GetComponentSpaceTransforms();
+
+							for (int32 Index = 0 ; Index < CurrentComponentSpaceBoneMatrix.Num(); Index++)
+							{
+								int32 ParentIndex = RefSkeleton.GetParentIndex(Index);
+								if (ParentIndex!=INDEX_NONE && WeakThis->CachedSelectedBones.Contains(ParentIndex))
+								{
+									const FTransform& Child =  CurrentComponentSpaceBoneMatrix[Index] * Mesh->GetComponentTransform();
+									const FTransform& Parent = CurrentComponentSpaceBoneMatrix[ParentIndex] * Mesh->GetComponentTransform();
+									DrawDebugLine(Mesh->GetWorld(), Child.GetLocation(), Parent.GetLocation(), WeakThis->DebugDrawColor, false, -1, SDPG_Foreground);	
+								}
+							}
+
+						}
+					}, TStatId(), NULL, ENamedThreads::GameThread);	
+			}
+			
+			
+			Proxy->SkeletalMeshObject = SkeletalMeshObject;
+			Proxy->SkinWeightProfile = SkinWeightProfile;
+			Proxy->BoneIsSelectedPerSection = CachedBoneIsSelectedPerSection;
+		}
+	}
+
+	return Proxy;
+}
+
+
+
+bool FOptimusSkinWeightsAsVertexMaskDataProviderProxy::IsValid(FValidationData const& InValidationData) const
+{
+	if (InValidationData.ParameterStructSize != sizeof(FParameters))
+	{
+		return false;
+	}
+	if (SkeletalMeshObject == nullptr)
+	{
+		return false;
+	}
+	const int32 LodIndex = SkeletalMeshObject->GetLOD();
+	FSkeletalMeshRenderData const& SkeletalMeshRenderData = SkeletalMeshObject->GetSkeletalMeshRenderData();
+	FSkeletalMeshLODRenderData const* LodRenderData = &SkeletalMeshRenderData.LODRenderData[LodIndex];
+	if (LodRenderData->RenderSections.Num() != InValidationData.NumInvocations)
+	{
+		return false;
+	}
+	
+	FSkinWeightVertexBuffer const* WeightBuffer = LodRenderData->GetSkinWeightVertexBuffer();
+	if (LodRenderData->SkinWeightProfilesData.ContainsProfile(SkinWeightProfile))
+	{
+		const FSkinWeightProfileStack ProfileStack{SkinWeightProfile};
+		WeightBuffer = LodRenderData->SkinWeightProfilesData.GetOverrideBuffer(ProfileStack);
+	}
+	
+	if (WeightBuffer == nullptr)
+	{
+		return false;
+	}
+
+	if (BoneIsSelectedPerSection.IsEmpty())
+	{
+		return false;
+	}
+
+	return true;	
+}
+
+void FOptimusSkinWeightsAsVertexMaskDataProviderProxy::AllocateResources(FRDGBuilder& GraphBuilder, FAllocationData const& InAllocationData)
+{
+	const int32 LodIndex = SkeletalMeshObject->GetLOD();
+	FSkeletalMeshRenderData const& SkeletalMeshRenderData = SkeletalMeshObject->GetSkeletalMeshRenderData();
+	FSkeletalMeshLODRenderData const* LodRenderData = &SkeletalMeshRenderData.LODRenderData[LodIndex];
+	
+	
+	BoneIsSelectedBuffersPerSection.AddDefaulted(LodRenderData->RenderSections.Num());
+	BoneIsSelectedBufferSRVsPerSection.AddDefaulted(LodRenderData->RenderSections.Num());	
+	
+	for (int32 InvocationIndex = 0 ; InvocationIndex < LodRenderData->RenderSections.Num(); InvocationIndex++)
+	{
+		int32 NumBones = LodRenderData->RenderSections[InvocationIndex].BoneMap.Num();
+		
+		const TArray<uint32>& BoneIsSelected = BoneIsSelectedPerSection[InvocationIndex];
+
+		int32 Stride = sizeof(uint32);
+		int32 ByteSize = NumBones * Stride;
+		
+		FRDGBufferRef Buffer = GraphBuilder.CreateBuffer(FRDGBufferDesc::CreateStructuredDesc(Stride,  NumBones), TEXT("BoneIsSelectedBuffer"), ERDGBufferFlags::None);
+		FRDGBufferSRVRef BufferSRV = GraphBuilder.CreateSRV(Buffer);
+
+		BoneIsSelectedBuffersPerSection[InvocationIndex] = Buffer;
+		BoneIsSelectedBufferSRVsPerSection[InvocationIndex] = BufferSRV;
+
+		GraphBuilder.QueueBufferUpload(Buffer, BoneIsSelected.GetData(), ByteSize, ERDGInitialDataFlags::None);	
+	}
+}
+
+struct FSkinWeightsAsVertexMaskDataInterfacePermutationIds
+{
+	uint32 EnableDeformerBones = 0;
+	uint32 UnlimitedBoneInfluence = 0;
+	uint32 BoneIndexUint16 = 0;
+	uint32 BoneWeightsUint16 = 0;
+
+	FSkinWeightsAsVertexMaskDataInterfacePermutationIds(FComputeKernelPermutationVector const& PermutationVector)
+	{
+		{
+			static FString Name(TEXT("ENABLE_DEFORMER_BONES"));
+			static uint32 Hash = GetTypeHash(Name);
+			EnableDeformerBones = PermutationVector.GetPermutationBits(Name, Hash, 1);
+		}
+		{
+			static FString Name(TEXT("GPUSKIN_UNLIMITED_BONE_INFLUENCE"));
+			static uint32 Hash = GetTypeHash(Name);
+			UnlimitedBoneInfluence = PermutationVector.GetPermutationBits(Name, Hash, 1);
+		}
+		{
+			static FString Name(TEXT("GPUSKIN_BONE_INDEX_UINT16"));
+			static uint32 Hash = GetTypeHash(Name);
+			BoneIndexUint16 = PermutationVector.GetPermutationBits(Name, Hash, 1);
+		}
+		{
+			static FString Name(TEXT("GPUSKIN_BONE_WEIGHTS_UINT16"));
+			static uint32 Hash = GetTypeHash(Name);
+			BoneWeightsUint16 = PermutationVector.GetPermutationBits(Name, Hash, 1);
+		}
+	}
+};
+
+void FOptimusSkinWeightsAsVertexMaskDataProviderProxy::GatherPermutations(FPermutationData& InOutPermutationData) const
+{
+	const int32 LodIndex = SkeletalMeshObject->GetLOD();
+	FSkeletalMeshRenderData const& SkeletalMeshRenderData = SkeletalMeshObject->GetSkeletalMeshRenderData();
+	FSkeletalMeshLODRenderData const* LodRenderData = &SkeletalMeshRenderData.LODRenderData[LodIndex];
+
+	FSkinWeightsAsVertexMaskDataInterfacePermutationIds PermutationIds(InOutPermutationData.PermutationVector);
+	for (int32 InvocationIndex = 0; InvocationIndex < InOutPermutationData.NumInvocations; ++InvocationIndex)
+	{
+		const bool bPreviousFrame = false;
+		FRHIShaderResourceView* BoneBufferSRV = FSkeletalMeshDeformerHelpers::GetBoneBufferForReading(SkeletalMeshObject, LodIndex, InvocationIndex, bPreviousFrame);
+		
+		FSkinWeightVertexBuffer const* WeightBuffer = LodRenderData->GetSkinWeightVertexBuffer();
+		if (LodRenderData->SkinWeightProfilesData.ContainsProfile(SkinWeightProfile))
+		{
+			const FSkinWeightProfileStack ProfileStack{SkinWeightProfile};
+			WeightBuffer = LodRenderData->SkinWeightProfilesData.GetOverrideBuffer(ProfileStack);
+		}
+		
+		check(WeightBuffer != nullptr);
+		FRHIShaderResourceView* SkinWeightBufferSRV = WeightBuffer->GetDataVertexBuffer()->GetSRV();
+		const bool bUnlimitedBoneInfluences = WeightBuffer->GetBoneInfluenceType() == GPUSkinBoneInfluenceType::UnlimitedBoneInfluence;
+		FRHIShaderResourceView* InputWeightLookupStreamSRV = bUnlimitedBoneInfluences ? WeightBuffer->GetLookupVertexBuffer()->GetSRV() : nullptr;
+		const bool bValidBones = (BoneBufferSRV != nullptr) && (SkinWeightBufferSRV != nullptr) && (!bUnlimitedBoneInfluences || InputWeightLookupStreamSRV != nullptr);
+		const bool bUse16BitBoneIndex = WeightBuffer->Use16BitBoneIndex();
+		const bool bUse16BitBoneWeights = WeightBuffer->Use16BitBoneWeight();
+
+		InOutPermutationData.PermutationIds[InvocationIndex] |= (bValidBones ? PermutationIds.EnableDeformerBones : 0);
+		InOutPermutationData.PermutationIds[InvocationIndex] |= (bUnlimitedBoneInfluences ? PermutationIds.UnlimitedBoneInfluence : 0);
+		InOutPermutationData.PermutationIds[InvocationIndex] |= (bUse16BitBoneIndex ? PermutationIds.BoneIndexUint16 : 0);
+		InOutPermutationData.PermutationIds[InvocationIndex] |= (bUse16BitBoneWeights ? PermutationIds.BoneWeightsUint16 : 0);
+	}
+}
+
+void FOptimusSkinWeightsAsVertexMaskDataProviderProxy::GatherDispatchData(FDispatchData const& InDispatchData)
+{
+	const int32 LodIndex = SkeletalMeshObject->GetLOD();
+	FSkeletalMeshRenderData const& SkeletalMeshRenderData = SkeletalMeshObject->GetSkeletalMeshRenderData();
+	FSkeletalMeshLODRenderData const* LodRenderData = &SkeletalMeshRenderData.LODRenderData[LodIndex];
+
+	FRHIShaderResourceView* NullSRVBinding = GWhiteVertexBufferWithSRV->ShaderResourceViewRHI.GetReference();
+	
+	const TStridedView<FParameters> ParameterArray = MakeStridedParameterView<FParameters>(InDispatchData);
+	for (int32 InvocationIndex = 0; InvocationIndex < InDispatchData.NumInvocations; ++InvocationIndex)
+	{
+		FSkelMeshRenderSection const& RenderSection = LodRenderData->RenderSections[InvocationIndex];
+
+		FSkinWeightVertexBuffer const* WeightBuffer = LodRenderData->GetSkinWeightVertexBuffer();
+		check(WeightBuffer != nullptr);
+		FRHIShaderResourceView* SkinWeightBufferSRV = WeightBuffer->GetDataVertexBuffer()->GetSRV();
+		const bool bUnlimitedBoneInfluences = WeightBuffer->GetBoneInfluenceType() == GPUSkinBoneInfluenceType::UnlimitedBoneInfluence;
+		FRHIShaderResourceView* InputWeightLookupStreamSRV = bUnlimitedBoneInfluences ? WeightBuffer->GetLookupVertexBuffer()->GetSRV() : nullptr;
+
+		FParameters& Parameters = ParameterArray[InvocationIndex];
+		Parameters.NumVertices = RenderSection.NumVertices;
+		Parameters.InputStreamStart = RenderSection.BaseVertexIndex;
+		Parameters.NumBoneInfluences = WeightBuffer->GetMaxBoneInfluences();
+		Parameters.InputWeightStart = (WeightBuffer->GetConstantInfluencesVertexStride() * RenderSection.GetVertexBufferIndex()) / sizeof(float);
+		Parameters.InputWeightStride = WeightBuffer->GetConstantInfluencesVertexStride();
+		Parameters.InputWeightIndexSize = WeightBuffer->GetBoneIndexByteSize() | (WeightBuffer->GetBoneWeightByteSize() << 8);
+		Parameters.InputWeightStream = SkinWeightBufferSRV != nullptr ? SkinWeightBufferSRV : NullSRVBinding;
+		Parameters.InputWeightLookupStream = InputWeightLookupStreamSRV != nullptr ? InputWeightLookupStreamSRV : NullSRVBinding;
+		Parameters.BoneIsSelected = BoneIsSelectedBufferSRVsPerSection[InvocationIndex];
+	}
+}
diff --git a/Engine/Plugins/Animation/DeformerGraph/Source/OptimusCore/Private/DataInterfaces/OptimusDataInterfaceSkinWeightsAsVertexMask.h b/Engine/Plugins/Animation/DeformerGraph/Source/OptimusCore/Private/DataInterfaces/OptimusDataInterfaceSkinWeightsAsVertexMask.h
new file mode 100644
--- /dev/null
+++ b/Engine/Plugins/Animation/DeformerGraph/Source/OptimusCore/Private/DataInterfaces/OptimusDataInterfaceSkinWeightsAsVertexMask.h
@@ -0,0 +1,146 @@
+// Copyright Epic Games, Inc. All Rights Reserved.
+
+#pragma once
+
+#include "IOptimusDeformerInstanceAccessor.h"
+#include "OptimusComputeDataInterface.h"
+#include "OptimusDataType.h"
+#include "ComputeFramework/ComputeDataProvider.h"
+
+#include "OptimusDataInterfaceSkinWeightsAsVertexMask.generated.h"
+
+class FSkinWeightsAsVertexMaskDataInterfaceParameters;
+class FSkeletalMeshObject;
+class USkeletalMeshComponent;
+class FSkeletalMeshLODRenderData;
+struct FReferenceSkeleton;
+
+
+/** Compute Framework Data Interface for skeletal data. */
+UCLASS(Category = ComputeFramework)
+class OPTIMUSCORE_API UOptimusSkinWeightsAsVertexMaskDataInterface : public UOptimusComputeDataInterface
+{
+	GENERATED_BODY()
+	
+public:
+	
+	//~ Begin UOptimusComputeDataInterface Interface
+	FString GetDisplayName() const override;
+	TArray<FOptimusCDIPinDefinition> GetPinDefinitions() const override;
+	TArray<FOptimusCDIPropertyPinDefinition> GetPropertyPinDefinitions() const override;
+	TSubclassOf<UActorComponent> GetRequiredComponentClass() const override;
+
+	void RegisterPropertyChangeDelegatesForOwningNode(UOptimusNode_DataInterface* InNode) override;
+	//~ End UOptimusComputeDataInterface Interface
+	
+
+	//~ Begin UComputeDataInterface Interface
+	TCHAR const* GetClassName() const override { return TEXT("SkinWeightMask"); }
+	void GetSupportedInputs(TArray<FShaderFunctionDefinition>& OutFunctions) const override;
+	void GetShaderParameters(TCHAR const* UID, FShaderParametersMetadataBuilder& InOutBuilder, FShaderParametersMetadataAllocations& InOutAllocations) const override;
+	void GetPermutations(FComputeKernelPermutationVector& OutPermutationVector) const override;
+	void GetShaderHash(FString& InOutKey) const override;
+	TCHAR const* GetShaderVirtualPath() const override;
+	void GetHLSL(FString& OutHLSL, FString const& InDataInterfaceName) const override;
+	UComputeDataProvider* CreateDataProvider(TObjectPtr<UObject> InBinding, uint64 InInputMask, uint64 InOutputMask) const override;
+	//~ End UComputeDataInterface Interface
+	
+	UPROPERTY(EditAnywhere, Category = "Data Interface")
+	FName SkinWeightProfile = NAME_None;
+
+	// Weights of selected bones are combined to form a single a vertex weight map 
+	UPROPERTY(EditAnywhere, Category = "Data Interface")
+	TArray<FName> BoneNames = {TEXT("Root")};
+	
+	// Include all bones within the radius by expanding selection towards the root
+	UPROPERTY(EditAnywhere, Category = "Data Interface")
+	int32 ExpandTowardsRoot = 0;
+
+	// Include children up to the specified depth
+	UPROPERTY(EditAnywhere, Category = "Data Interface")
+	int32 ExpandTowardsLeaf = 999;
+
+	UPROPERTY(EditAnywhere, Category = "Debug")
+	bool bDebugDrawIncludedBones = true;
+
+	UPROPERTY(EditAnywhere, Category = "Debug")
+	FColor DebugDrawColor = FColor::Green;
+	
+private:
+	friend class UOptimusSkinWeightsAsVertexMaskDataProvider;
+	FName SkinWeightProfilePropertyName = GET_MEMBER_NAME_CHECKED(UOptimusSkinWeightsAsVertexMaskDataInterface, SkinWeightProfile);
+	FName BoneNamesPropertyName = GET_MEMBER_NAME_CHECKED(UOptimusSkinWeightsAsVertexMaskDataInterface, BoneNames);
+	FName ExpandTowardsRootPropertyName = GET_MEMBER_NAME_CHECKED(UOptimusSkinWeightsAsVertexMaskDataInterface, ExpandTowardsRoot);
+	FName ExpandTowardsLeafPropertyName = GET_MEMBER_NAME_CHECKED(UOptimusSkinWeightsAsVertexMaskDataInterface, ExpandTowardsLeaf);
+	FName DebugDrawIncludedBonesPropertyName = GET_MEMBER_NAME_CHECKED(UOptimusSkinWeightsAsVertexMaskDataInterface, bDebugDrawIncludedBones);
+	
+	static TCHAR const* TemplateFilePath;
+	FOnPinDefinitionChanged OnPinDefinitionChangedDelegate;
+	FOnPinDefinitionRenamed OnPinDefinitionRenamedDelegate;
+	FSimpleDelegate OnDisplayNameChangedDelegate;
+};
+
+
+/** Compute Framework Data Provider for reading skeletal mesh. */
+UCLASS(BlueprintType, editinlinenew, Category = ComputeFramework)
+class UOptimusSkinWeightsAsVertexMaskDataProvider :
+	public UComputeDataProvider,
+	public IOptimusDeformerInstanceAccessor
+
+{
+	GENERATED_BODY()
+
+public:
+	void Init(const UOptimusSkinWeightsAsVertexMaskDataInterface* InDataInterface, USkeletalMeshComponent* InSkeletalMesh);
+	
+	void SetDeformerInstance(UOptimusDeformerInstance* InInstance) override;
+
+	UPROPERTY(EditAnywhere, BlueprintReadWrite, Category = Binding)
+	TObjectPtr<USkeletalMeshComponent> SkeletalMesh = nullptr;
+
+	FName SkinWeightProfile = NAME_None;
+	
+	TArray<FName> BoneNames;
+	int32 ExpandTowardsRoot = 0;
+	int32 ExpandTowardsLeaf = 0;
+	bool bDebugDrawIncludedBones = false;
+	FColor DebugDrawColor = FColor::Green;
+
+	bool bIsInitialized = false;
+	TSet<int32> CachedSelectedBones;
+	TArray<TArray<uint32>> CachedBoneIsSelectedPerSection;
+	
+	//~ Begin UComputeDataProvider Interface
+	FComputeDataProviderRenderProxy* GetRenderProxy() override;
+	//~ End UComputeDataProvider Interface
+
+
+private:
+	UPROPERTY()
+	TObjectPtr<UOptimusDeformerInstance> DeformerInstance;
+
+	TWeakObjectPtr<const UOptimusSkinWeightsAsVertexMaskDataInterface> WeakDataInterface;
+};
+
+class FOptimusSkinWeightsAsVertexMaskDataProviderProxy : public FComputeDataProviderRenderProxy
+{
+public:
+	FOptimusSkinWeightsAsVertexMaskDataProviderProxy() = default;
+
+	//~ Begin FComputeDataProviderRenderProxy Interface
+	bool IsValid(FValidationData const& InValidationData) const override;
+	void AllocateResources(FRDGBuilder& GraphBuilder, FAllocationData const& InAllocationData) override;
+	void GatherPermutations(FPermutationData& InOutPermutationData) const override;
+	void GatherDispatchData(FDispatchData const& InDispatchData) override;
+	//~ End FComputeDataProviderRenderProxy Interface
+
+	using FParameters = FSkinWeightsAsVertexMaskDataInterfaceParameters;
+
+	FSkeletalMeshObject* SkeletalMeshObject = nullptr;
+	
+	FName SkinWeightProfile = NAME_None;
+
+	TArray<TArray<uint32>> BoneIsSelectedPerSection;
+	TArray<FRDGBufferRef> BoneIsSelectedBuffersPerSection;
+	TArray<FRDGBufferSRVRef> BoneIsSelectedBufferSRVsPerSection;
+};
