[Deformer Graph] exposed attribute name as property pin on VertexAttributeDI

diff --git a/Engine/Plugins/Animation/DeformerGraph/Source/OptimusCore/Private/DataInterfaces/OptimusDataInterfaceAdvancedSkeleton.cpp b/Engine/Plugins/Animation/DeformerGraph/Source/OptimusCore/Private/DataInterfaces/OptimusDataInterfaceAdvancedSkeleton.cpp
--- a/Engine/Plugins/Animation/DeformerGraph/Source/OptimusCore/Private/DataInterfaces/OptimusDataInterfaceAdvancedSkeleton.cpp
+++ b/Engine/Plugins/Animation/DeformerGraph/Source/OptimusCore/Private/DataInterfaces/OptimusDataInterfaceAdvancedSkeleton.cpp
@@ -64,6 +64,11 @@ FString FOptimusAnimAttributeBufferDescription::GetFormattedId(
 	return  UniqueId;
 }
 
+FName UOptimusAdvancedSkeletonDataInterface::GetSkinWeightProfilePropertyName()
+{
+	return GET_MEMBER_NAME_CHECKED(UOptimusAdvancedSkeletonDataInterface, SkinWeightProfile);
+}
+
 FString UOptimusAdvancedSkeletonDataInterface::GetUnusedAttributeName(int32 CurrentAttributeIndex, const FString& InName) const
 {
 	TSet<FString> UsedNames;
@@ -276,7 +281,7 @@ TArray<FOptimusCDIPropertyPinDefinition> UOptimusAdvancedSkeletonDataInterface::
 	const FOptimusDataTypeHandle NameType = FOptimusDataTypeRegistry::Get().FindType(*FNameProperty::StaticClass());
 	
 	PropertyPinDefinitions.Add(
-		{SkinWeightProfilePropertyName, NameType}
+		{GetSkinWeightProfilePropertyName(), NameType}
 	);
 
 	return PropertyPinDefinitions;
@@ -709,25 +714,28 @@ FComputeDataProviderRenderProxy* UOptimusAdvancedSkeletonDataProvider::GetRender
 
 	if (SkeletalMesh && SkeletalMesh->MeshObject)
 	{
-		const UOptimusAdvancedSkeletonDataInterface* DataInterface = WeakDataInterface.Get();
-		FOptimusValueContainerStruct ValueContainer =
-			DeformerInstance->GetDataInterfacePropertyOverride(
-				WeakDataInterface.Get(),
-				DataInterface->SkinWeightProfilePropertyName
-				);
-
-		TValueOrError<FName, EPropertyBagResult> Value = ValueContainer.Value.GetValueName(FOptimusValueContainerStruct::ValuePropertyName);
-		if (Value.HasValue())
+		if (const UOptimusAdvancedSkeletonDataInterface* DataInterface = WeakDataInterface.Get())
 		{
-			if (SkinWeightProfile != Value.GetValue())
+			FOptimusValueContainerStruct ValueContainer =
+					DeformerInstance->GetDataInterfacePropertyOverride(
+						DataInterface,
+						UOptimusAdvancedSkeletonDataInterface::GetSkinWeightProfilePropertyName()
+						);
+
+			TValueOrError<FName, EPropertyBagResult> Value = ValueContainer.Value.GetValueName(FOptimusValueContainerStruct::ValuePropertyName);
+			if (Value.HasValue())
 			{
-				SkinWeightProfile = Value.GetValue();
-				if (bEnableLayeredSkinning)
+				if (SkinWeightProfile != Value.GetValue())
 				{
-					bIsLayeredSkinInitialized = false;
+					SkinWeightProfile = Value.GetValue();
+					if (bEnableLayeredSkinning)
+					{
+						bIsLayeredSkinInitialized = false;
+					}
 				}
 			}
 		}
+		
 		FSkeletalMeshObject* SkeletalMeshObject = SkeletalMesh->MeshObject;
 		const int32 LodIndex = SkeletalMeshObject->GetLOD();
 		FSkeletalMeshRenderData const& SkeletalMeshRenderData = SkeletalMeshObject->GetSkeletalMeshRenderData();
diff --git a/Engine/Plugins/Animation/DeformerGraph/Source/OptimusCore/Private/DataInterfaces/OptimusDataInterfaceAdvancedSkeleton.h b/Engine/Plugins/Animation/DeformerGraph/Source/OptimusCore/Private/DataInterfaces/OptimusDataInterfaceAdvancedSkeleton.h
--- a/Engine/Plugins/Animation/DeformerGraph/Source/OptimusCore/Private/DataInterfaces/OptimusDataInterfaceAdvancedSkeleton.h
+++ b/Engine/Plugins/Animation/DeformerGraph/Source/OptimusCore/Private/DataInterfaces/OptimusDataInterfaceAdvancedSkeleton.h
@@ -156,7 +156,7 @@ public:
 
 private:
 	friend class UOptimusAdvancedSkeletonDataProvider;
-	FName SkinWeightProfilePropertyName = GET_MEMBER_NAME_CHECKED(UOptimusAdvancedSkeletonDataInterface, SkinWeightProfile);
+	static FName GetSkinWeightProfilePropertyName();
 	
 	FString GetUnusedAttributeName(int32 CurrentAttributeIndex ,const FString& InName) const;
 	void UpdateAttributePinNamesAndHlslIds();
diff --git a/Engine/Plugins/Animation/DeformerGraph/Source/OptimusCore/Private/DataInterfaces/OptimusDataInterfaceSkinWeightsAsVertexMask.cpp b/Engine/Plugins/Animation/DeformerGraph/Source/OptimusCore/Private/DataInterfaces/OptimusDataInterfaceSkinWeightsAsVertexMask.cpp
--- a/Engine/Plugins/Animation/DeformerGraph/Source/OptimusCore/Private/DataInterfaces/OptimusDataInterfaceSkinWeightsAsVertexMask.cpp
+++ b/Engine/Plugins/Animation/DeformerGraph/Source/OptimusCore/Private/DataInterfaces/OptimusDataInterfaceSkinWeightsAsVertexMask.cpp
@@ -15,6 +15,31 @@
 #include "Animation/SkinWeightProfileManager.h"
 #include "Components/SkeletalMeshComponent.h"
 
+FName UOptimusSkinWeightsAsVertexMaskDataInterface::GetSkinWeightProfilePropertyName()
+{
+	return GET_MEMBER_NAME_CHECKED(UOptimusSkinWeightsAsVertexMaskDataInterface, SkinWeightProfile);	
+}
+
+FName UOptimusSkinWeightsAsVertexMaskDataInterface::GetBoneNamesPropertyName()
+{
+	return GET_MEMBER_NAME_CHECKED(UOptimusSkinWeightsAsVertexMaskDataInterface, BoneNames);	
+}
+
+FName UOptimusSkinWeightsAsVertexMaskDataInterface::GetExpandTowardsRootPropertyName()
+{
+	return GET_MEMBER_NAME_CHECKED(UOptimusSkinWeightsAsVertexMaskDataInterface, ExpandTowardsRoot);
+}
+
+FName UOptimusSkinWeightsAsVertexMaskDataInterface::GetExpandTowardsLeafPropertyName()
+{
+	return GET_MEMBER_NAME_CHECKED(UOptimusSkinWeightsAsVertexMaskDataInterface, ExpandTowardsLeaf);
+}
+
+FName UOptimusSkinWeightsAsVertexMaskDataInterface::GetDebugDrawIncludedBonesPropertyName()
+{
+	return GET_MEMBER_NAME_CHECKED(UOptimusSkinWeightsAsVertexMaskDataInterface, bDebugDrawIncludedBones);
+}
+
 FString UOptimusSkinWeightsAsVertexMaskDataInterface::GetDisplayName() const
 {
 	return TEXT("Skin Weights as Vertex Mask");
@@ -37,11 +62,11 @@ TArray<FOptimusCDIPropertyPinDefinition> UOptimusSkinWeightsAsVertexMaskDataInte
 	const FOptimusDataTypeHandle IntType = FOptimusDataTypeRegistry::Get().FindType(*FIntProperty::StaticClass());
 	const FOptimusDataTypeHandle BoolType = FOptimusDataTypeRegistry::Get().FindType(*FBoolProperty::StaticClass());
 	
-	Defs.Add({SkinWeightProfilePropertyName, NameType});
-	Defs.Add({BoneNamesPropertyName, NameArrayType});
-	Defs.Add({ExpandTowardsRootPropertyName, IntType});
-	Defs.Add({ExpandTowardsLeafPropertyName, IntType});
-	Defs.Add({DebugDrawIncludedBonesPropertyName, BoolType});
+	Defs.Add({GetSkinWeightProfilePropertyName(), NameType});
+	Defs.Add({GetBoneNamesPropertyName(), NameArrayType});
+	Defs.Add({GetExpandTowardsRootPropertyName(), IntType});
+	Defs.Add({GetExpandTowardsLeafPropertyName(), IntType});
+	Defs.Add({GetDebugDrawIncludedBonesPropertyName(), BoolType});
 
 	return Defs;
 }
@@ -158,85 +183,86 @@ FComputeDataProviderRenderProxy* UOptimusSkinWeightsAsVertexMaskDataProvider::Ge
 	{
 
 		
-		const UOptimusSkinWeightsAsVertexMaskDataInterface* DataInterface = WeakDataInterface.Get();
-		
+		if (const UOptimusSkinWeightsAsVertexMaskDataInterface* DataInterface = WeakDataInterface.Get())
 		{
-			const FOptimusValueContainerStruct ValueContainer =
-				DeformerInstance->GetDataInterfacePropertyOverride( DataInterface, DataInterface->SkinWeightProfilePropertyName );
-
-			TValueOrError<FName, EPropertyBagResult> Value = ValueContainer.Value.GetValueName(FOptimusValueContainerStruct::ValuePropertyName);
-			if (Value.HasValue())
 			{
-				SkinWeightProfile = Value.GetValue();
-			}
-		}
+				const FOptimusValueContainerStruct ValueContainer =
+					DeformerInstance->GetDataInterfacePropertyOverride( DataInterface, UOptimusSkinWeightsAsVertexMaskDataInterface::GetSkinWeightProfilePropertyName() );
 
-		{
-			const FOptimusValueContainerStruct ValueContainer =
-				DeformerInstance->GetDataInterfacePropertyOverride( DataInterface, DataInterface->BoneNamesPropertyName);
-
-			TValueOrError<const FPropertyBagArrayRef, EPropertyBagResult> Value = ValueContainer.Value.GetArrayRef(FOptimusValueContainerStruct::ValuePropertyName);
-			if (Value.HasValue())
-			{
-				TArray<FName> NewBoneNames;
-				NewBoneNames.SetNum(Value.GetValue().Num());
-				for (int32 Index = 0 ; Index < NewBoneNames.Num(); Index++)
+				TValueOrError<FName, EPropertyBagResult> Value = ValueContainer.Value.GetValueName(FOptimusValueContainerStruct::ValuePropertyName);
+				if (Value.HasValue())
 				{
-					NewBoneNames[Index] = Value.GetValue().GetValueName(Index).GetValue();
+					SkinWeightProfile = Value.GetValue();
 				}
+			}
 
-				if (NewBoneNames != BoneNames)
+			{
+				const FOptimusValueContainerStruct ValueContainer =
+					DeformerInstance->GetDataInterfacePropertyOverride( DataInterface, UOptimusSkinWeightsAsVertexMaskDataInterface::GetBoneNamesPropertyName());
+
+				TValueOrError<const FPropertyBagArrayRef, EPropertyBagResult> Value = ValueContainer.Value.GetArrayRef(FOptimusValueContainerStruct::ValuePropertyName);
+				if (Value.HasValue())
 				{
-					bIsInitialized = false;
-					BoneNames = MoveTemp(NewBoneNames);
+					TArray<FName> NewBoneNames;
+					NewBoneNames.SetNum(Value.GetValue().Num());
+					for (int32 Index = 0 ; Index < NewBoneNames.Num(); Index++)
+					{
+						NewBoneNames[Index] = Value.GetValue().GetValueName(Index).GetValue();
+					}
+
+					if (NewBoneNames != BoneNames)
+					{
+						bIsInitialized = false;
+						BoneNames = MoveTemp(NewBoneNames);
+					}
 				}
 			}
-		}
-
-		{
-			const FOptimusValueContainerStruct ValueContainer =
-				DeformerInstance->GetDataInterfacePropertyOverride( DataInterface, DataInterface->ExpandTowardsRootPropertyName);
 
-			TValueOrError<int32 , EPropertyBagResult> Value = ValueContainer.Value.GetValueInt32(FOptimusValueContainerStruct::ValuePropertyName);
-			if (Value.HasValue())
 			{
-				if (ExpandTowardsRoot != Value.GetValue())
+				const FOptimusValueContainerStruct ValueContainer =
+					DeformerInstance->GetDataInterfacePropertyOverride( DataInterface, UOptimusSkinWeightsAsVertexMaskDataInterface::GetExpandTowardsRootPropertyName());
+
+				TValueOrError<int32 , EPropertyBagResult> Value = ValueContainer.Value.GetValueInt32(FOptimusValueContainerStruct::ValuePropertyName);
+				if (Value.HasValue())
 				{
-					bIsInitialized = false;
-					ExpandTowardsRoot = Value.GetValue();
+					if (ExpandTowardsRoot != Value.GetValue())
+					{
+						bIsInitialized = false;
+						ExpandTowardsRoot = Value.GetValue();
+					}
 				}
 			}
-		}
-
-		{
-			const FOptimusValueContainerStruct ValueContainer =
-				DeformerInstance->GetDataInterfacePropertyOverride( DataInterface, DataInterface->ExpandTowardsLeafPropertyName);
 
-			TValueOrError<int32 , EPropertyBagResult> Value = ValueContainer.Value.GetValueInt32(FOptimusValueContainerStruct::ValuePropertyName);
-			if (Value.HasValue())
 			{
-				if (ExpandTowardsLeaf != Value.GetValue())
+				const FOptimusValueContainerStruct ValueContainer =
+					DeformerInstance->GetDataInterfacePropertyOverride( DataInterface, UOptimusSkinWeightsAsVertexMaskDataInterface::GetExpandTowardsLeafPropertyName());
+
+				TValueOrError<int32 , EPropertyBagResult> Value = ValueContainer.Value.GetValueInt32(FOptimusValueContainerStruct::ValuePropertyName);
+				if (Value.HasValue())
 				{
-					bIsInitialized = false;
-					ExpandTowardsLeaf = Value.GetValue();
+					if (ExpandTowardsLeaf != Value.GetValue())
+					{
+						bIsInitialized = false;
+						ExpandTowardsLeaf = Value.GetValue();
+					}
 				}
 			}
-		}
-
-		{
-			const FOptimusValueContainerStruct ValueContainer =
-				DeformerInstance->GetDataInterfacePropertyOverride( DataInterface, DataInterface->DebugDrawIncludedBonesPropertyName);
 
-			TValueOrError<bool, EPropertyBagResult> Value = ValueContainer.Value.GetValueBool(FOptimusValueContainerStruct::ValuePropertyName);
-			if (Value.HasValue())
 			{
-				if (bDebugDrawIncludedBones != Value.GetValue())
+				const FOptimusValueContainerStruct ValueContainer =
+					DeformerInstance->GetDataInterfacePropertyOverride( DataInterface, UOptimusSkinWeightsAsVertexMaskDataInterface::GetDebugDrawIncludedBonesPropertyName());
+
+				TValueOrError<bool, EPropertyBagResult> Value = ValueContainer.Value.GetValueBool(FOptimusValueContainerStruct::ValuePropertyName);
+				if (Value.HasValue())
 				{
-					bDebugDrawIncludedBones = Value.GetValue();
+					if (bDebugDrawIncludedBones != Value.GetValue())
+					{
+						bDebugDrawIncludedBones = Value.GetValue();
+					}
 				}
-			}
+			}	
 		}
-
+		
 		ExpandTowardsRoot = FMath::Max(0, ExpandTowardsRoot);
 		ExpandTowardsLeaf = FMath::Max(0, ExpandTowardsLeaf);
 		
@@ -366,6 +392,7 @@ FComputeDataProviderRenderProxy* UOptimusSkinWeightsAsVertexMaskDataProvider::Ge
 
 			if (bDebugDrawIncludedBones)
 			{
+#if WITH_EDITOR
 				FFunctionGraphTask::CreateAndDispatchWhenReady([ WeakThis = TWeakObjectPtr(this)]()
 					{
 						if (WeakThis.IsValid())
@@ -388,6 +415,7 @@ FComputeDataProviderRenderProxy* UOptimusSkinWeightsAsVertexMaskDataProvider::Ge
 
 						}
 					}, TStatId(), NULL, ENamedThreads::GameThread);	
+#endif
 			}
 			
 			
diff --git a/Engine/Plugins/Animation/DeformerGraph/Source/OptimusCore/Private/DataInterfaces/OptimusDataInterfaceSkinWeightsAsVertexMask.h b/Engine/Plugins/Animation/DeformerGraph/Source/OptimusCore/Private/DataInterfaces/OptimusDataInterfaceSkinWeightsAsVertexMask.h
--- a/Engine/Plugins/Animation/DeformerGraph/Source/OptimusCore/Private/DataInterfaces/OptimusDataInterfaceSkinWeightsAsVertexMask.h
+++ b/Engine/Plugins/Animation/DeformerGraph/Source/OptimusCore/Private/DataInterfaces/OptimusDataInterfaceSkinWeightsAsVertexMask.h
@@ -68,11 +68,11 @@ public:
 	
 private:
 	friend class UOptimusSkinWeightsAsVertexMaskDataProvider;
-	FName SkinWeightProfilePropertyName = GET_MEMBER_NAME_CHECKED(UOptimusSkinWeightsAsVertexMaskDataInterface, SkinWeightProfile);
-	FName BoneNamesPropertyName = GET_MEMBER_NAME_CHECKED(UOptimusSkinWeightsAsVertexMaskDataInterface, BoneNames);
-	FName ExpandTowardsRootPropertyName = GET_MEMBER_NAME_CHECKED(UOptimusSkinWeightsAsVertexMaskDataInterface, ExpandTowardsRoot);
-	FName ExpandTowardsLeafPropertyName = GET_MEMBER_NAME_CHECKED(UOptimusSkinWeightsAsVertexMaskDataInterface, ExpandTowardsLeaf);
-	FName DebugDrawIncludedBonesPropertyName = GET_MEMBER_NAME_CHECKED(UOptimusSkinWeightsAsVertexMaskDataInterface, bDebugDrawIncludedBones);
+	static FName GetSkinWeightProfilePropertyName();
+	static FName GetBoneNamesPropertyName();
+	static FName GetExpandTowardsRootPropertyName();
+	static FName GetExpandTowardsLeafPropertyName();
+	static FName GetDebugDrawIncludedBonesPropertyName();
 	
 	static TCHAR const* TemplateFilePath;
 	FOnPinDefinitionChanged OnPinDefinitionChangedDelegate;
diff --git a/Engine/Plugins/Animation/DeformerGraph/Source/OptimusCore/Private/DataInterfaces/OptimusDataInterfaceSkinnedMeshVertexAttribute.cpp b/Engine/Plugins/Animation/DeformerGraph/Source/OptimusCore/Private/DataInterfaces/OptimusDataInterfaceSkinnedMeshVertexAttribute.cpp
--- a/Engine/Plugins/Animation/DeformerGraph/Source/OptimusCore/Private/DataInterfaces/OptimusDataInterfaceSkinnedMeshVertexAttribute.cpp
+++ b/Engine/Plugins/Animation/DeformerGraph/Source/OptimusCore/Private/DataInterfaces/OptimusDataInterfaceSkinnedMeshVertexAttribute.cpp
@@ -2,6 +2,9 @@
 
 #include "OptimusDataInterfaceSkinnedMeshVertexAttribute.h"
 
+#include "OptimusDataTypeRegistry.h"
+#include "OptimusDeformerInstance.h"
+#include "OptimusValueContainerStruct.h"
 #include "ComponentSources/OptimusSkinnedMeshComponentSource.h"
 #include "ComputeFramework/ShaderParamTypeDefinition.h"
 #include "Rendering/SkeletalMeshAttributeVertexBuffer.h"
@@ -10,6 +13,10 @@
 #include "ShaderParameterMetadataBuilder.h"
 #include "SkeletalRenderPublic.h"
 
+FName UOptimusSkinnedMeshVertexAttributeDataInterface::GetAttributeNamePropertyName()
+{
+	return GET_MEMBER_NAME_CHECKED(UOptimusSkinnedMeshVertexAttributeDataInterface, AttributeName);
+}
 
 FString UOptimusSkinnedMeshVertexAttributeDataInterface::GetDisplayName() const
 {
@@ -25,6 +32,15 @@ TArray<FOptimusCDIPinDefinition> UOptimusSkinnedMeshVertexAttributeDataInterface
 	return Defs;
 }
 
+TArray<FOptimusCDIPropertyPinDefinition> UOptimusSkinnedMeshVertexAttributeDataInterface::GetPropertyPinDefinitions() const
+{
+	TArray<FOptimusCDIPropertyPinDefinition> Defs;
+	const FOptimusDataTypeHandle NameType = FOptimusDataTypeRegistry::Get().FindType(*FNameProperty::StaticClass());
+	Defs.Add({GetAttributeNamePropertyName(), NameType});
+
+	return Defs;
+}
+
 
 TSubclassOf<UActorComponent> UOptimusSkinnedMeshVertexAttributeDataInterface::GetRequiredComponentClass() const
 {
@@ -90,6 +106,7 @@ UComputeDataProvider* UOptimusSkinnedMeshVertexAttributeDataInterface::CreateDat
 	UOptimusSkinnedMeshVertexAttributeDataProvider* Provider = NewObject<UOptimusSkinnedMeshVertexAttributeDataProvider>();
 	Provider->SkinnedMeshComponent = Cast<USkinnedMeshComponent>(InBinding);
 	Provider->AttributeName = AttributeName;
+	Provider->WeakDataInterface = this;
 
 	return Provider;
 }
@@ -112,9 +129,29 @@ bool UOptimusSkinnedMeshVertexAttributeDataProvider::IsValid() const
 
 FComputeDataProviderRenderProxy* UOptimusSkinnedMeshVertexAttributeDataProvider::GetRenderProxy()
 {
+	if (const UOptimusSkinnedMeshVertexAttributeDataInterface* DataInterface = WeakDataInterface.Get())
+	{
+		FOptimusValueContainerStruct ValueContainer =
+			DeformerInstance->GetDataInterfacePropertyOverride(
+				DataInterface,
+				UOptimusSkinnedMeshVertexAttributeDataInterface::GetAttributeNamePropertyName()
+				);
+
+		TValueOrError<FName, EPropertyBagResult> Value = ValueContainer.Value.GetValueName(FOptimusValueContainerStruct::ValuePropertyName);
+		if (Value.HasValue())
+		{
+			AttributeName = Value.GetValue();
+		}
+	}
+	
 	return new FOptimusSkinnedMeshVertexAttributeDataProviderProxy(SkinnedMeshComponent, AttributeName);
 }
 
+void UOptimusSkinnedMeshVertexAttributeDataProvider::SetDeformerInstance(UOptimusDeformerInstance* InInstance)
+{
+	DeformerInstance = InInstance;
+}
+
 
 FOptimusSkinnedMeshVertexAttributeDataProviderProxy::FOptimusSkinnedMeshVertexAttributeDataProviderProxy(
 	USkinnedMeshComponent* InSkinnedMeshComponent, 
diff --git a/Engine/Plugins/Animation/DeformerGraph/Source/OptimusCore/Private/DataInterfaces/OptimusDataInterfaceSkinnedMeshVertexAttribute.h b/Engine/Plugins/Animation/DeformerGraph/Source/OptimusCore/Private/DataInterfaces/OptimusDataInterfaceSkinnedMeshVertexAttribute.h
--- a/Engine/Plugins/Animation/DeformerGraph/Source/OptimusCore/Private/DataInterfaces/OptimusDataInterfaceSkinnedMeshVertexAttribute.h
+++ b/Engine/Plugins/Animation/DeformerGraph/Source/OptimusCore/Private/DataInterfaces/OptimusDataInterfaceSkinnedMeshVertexAttribute.h
@@ -2,6 +2,7 @@
 
 #pragma once
 
+#include "IOptimusDeformerInstanceAccessor.h"
 #include "OptimusComputeDataInterface.h"
 
 #include "ComputeFramework/ComputeDataProvider.h"
@@ -22,6 +23,7 @@ class UOptimusSkinnedMeshVertexAttributeDataInterface :
 	//~ Begin UOptimusComputeDataInterface Interface
 	FString GetDisplayName() const override;
 	TArray<FOptimusCDIPinDefinition> GetPinDefinitions() const override;
+	TArray<FOptimusCDIPropertyPinDefinition> GetPropertyPinDefinitions() const override;
 	TSubclassOf<UActorComponent> GetRequiredComponentClass() const override;
 	//~ End UOptimusComputeDataInterface Interface
 
@@ -37,12 +39,17 @@ class UOptimusSkinnedMeshVertexAttributeDataInterface :
 
 	UPROPERTY(EditAnywhere, Category="Vertex Attribute")
 	FName AttributeName;
+private:
+	friend class UOptimusSkinnedMeshVertexAttributeDataProvider;
+	static FName GetAttributeNamePropertyName();
 };
 
 
 /** Compute Framework Data Provider for reading skeletal mesh. */
 UCLASS(BlueprintType, editinlinenew, Category = ComputeFramework)
-class UOptimusSkinnedMeshVertexAttributeDataProvider : public UComputeDataProvider
+class UOptimusSkinnedMeshVertexAttributeDataProvider :
+	public UComputeDataProvider,
+	public IOptimusDeformerInstanceAccessor
 {
 	GENERATED_BODY()
 
@@ -53,11 +60,19 @@ public:
 	UPROPERTY(EditAnywhere, BlueprintReadWrite, Category = VertexAttribute)
 	FName AttributeName;
 
+    TWeakObjectPtr<const UOptimusSkinnedMeshVertexAttributeDataInterface> WeakDataInterface;
+	
 	//~ Begin UComputeDataProvider Interface
 	bool IsValid() const override;
 	FComputeDataProviderRenderProxy* GetRenderProxy() override;
 	//~ End UComputeDataProvider Interface
 
+	void SetDeformerInstance(UOptimusDeformerInstance* InInstance) override;
+
+	
+private:
+	UPROPERTY()
+	TObjectPtr<UOptimusDeformerInstance> DeformerInstance;
 };
 
 class FOptimusSkinnedMeshVertexAttributeDataProviderProxy : public FComputeDataProviderRenderProxy
