[Deformer Graph] minor fix: uninitialized struct member

diff --git a/Engine/Plugins/Animation/DeformerGraph/Source/OptimusCore/Public/OptimusValue.h b/Engine/Plugins/Animation/DeformerGraph/Source/OptimusCore/Public/OptimusValue.h
--- a/Engine/Plugins/Animation/DeformerGraph/Source/OptimusCore/Public/OptimusValue.h
+++ b/Engine/Plugins/Animation/DeformerGraph/Source/OptimusCore/Public/OptimusValue.h
@@ -28,10 +28,10 @@ struct FOptimusValueIdentifier
 	GENERATED_BODY()
 
 	UPROPERTY()
-	EOptimusValueType Type;
+	EOptimusValueType Type = EOptimusValueType::Invalid;
 
 	UPROPERTY()
-	FName Name;
+	FName Name = NAME_None;
 	
 	friend uint32 GetTypeHash(const FOptimusValueIdentifier& InIdentifier)
 	{
