[Deformer Graph] fixed crash when loading the deformer asset containing the lattice function reference node but the lattice struct is not registered

diff --git a/Engine/Plugins/Animation/DeformerGraph/Source/OptimusCore/Private/OptimusDataType.cpp b/Engine/Plugins/Animation/DeformerGraph/Source/OptimusCore/Private/OptimusDataType.cpp
--- a/Engine/Plugins/Animation/DeformerGraph/Source/OptimusCore/Private/OptimusDataType.cpp
+++ b/Engine/Plugins/Animation/DeformerGraph/Source/OptimusCore/Private/OptimusDataType.cpp
@@ -45,7 +45,7 @@ FOptimusDataTypeHandle FOptimusDataTypeRef::Resolve() const
 	// so we have to register these types on demand.
 	if (!TypeHandle.IsValid())
 	{
-		if (TypeObject.IsValid())
+		if (TypeObject.LoadSynchronous())
 		{
 			if (Registry.RegisterStructType(Cast<UScriptStruct>(TypeObject.Get())))
 			{
