Added some safety checks to compute graph/optimus to prevent crashes on nullptr.

diff --git a/Engine/Plugins/Animation/DeformerGraph/Source/OptimusCore/Private/Nodes/OptimusNode_DataInterface.cpp b/Engine/Plugins/Animation/DeformerGraph/Source/OptimusCore/Private/Nodes/OptimusNode_DataInterface.cpp
--- a/Engine/Plugins/Animation/DeformerGraph/Source/OptimusCore/Private/Nodes/OptimusNode_DataInterface.cpp
+++ b/Engine/Plugins/Animation/DeformerGraph/Source/OptimusCore/Private/Nodes/OptimusNode_DataInterface.cpp
@@ -573,6 +573,11 @@ void UOptimusNode_DataInterface::CreatePinFromDefinition(const FOptimusCDIPinDef
 
 void UOptimusNode_DataInterface::CreatePropertyPinsFromDataInterface(const UOptimusComputeDataInterface* InDataInterface, bool bSupportUndo)
 {
+	if (!InDataInterface)
+	{
+		return;
+	}
+
 	const TArray<FOptimusCDIPropertyPinDefinition> PropertyPinDefinitions = InDataInterface->GetPropertyPinDefinitions();
 
 	// Property pins should go before any shader pins
diff --git a/Engine/Plugins/Runtime/ComputeFramework/Source/ComputeFramework/Private/ComputeGraph.cpp b/Engine/Plugins/Runtime/ComputeFramework/Source/ComputeFramework/Private/ComputeGraph.cpp
--- a/Engine/Plugins/Runtime/ComputeFramework/Source/ComputeFramework/Private/ComputeGraph.cpp
+++ b/Engine/Plugins/Runtime/ComputeFramework/Source/ComputeFramework/Private/ComputeGraph.cpp
@@ -67,7 +67,10 @@ void UComputeGraph::PostLoad()
 
 	for (UComputeDataInterface* DataInterface : DataInterfaces)
 	{
-		DataInterface->ConditionalPostLoad();
+		if (DataInterface)
+		{
+			DataInterface->ConditionalPostLoad();
+		}
 	}
 	
 #if WITH_EDITOR
