[Deformer Graph] added a missing array reset to avoid unbounded memory allocation when copying array variable values.

diff --git a/Engine/Plugins/Animation/DeformerGraph/Source/OptimusCore/Private/OptimusDataTypeRegistry.cpp b/Engine/Plugins/Animation/DeformerGraph/Source/OptimusCore/Private/OptimusDataTypeRegistry.cpp
--- a/Engine/Plugins/Animation/DeformerGraph/Source/OptimusCore/Private/OptimusDataTypeRegistry.cpp
+++ b/Engine/Plugins/Animation/DeformerGraph/Source/OptimusCore/Private/OptimusDataTypeRegistry.cpp
@@ -959,6 +959,7 @@ bool FOptimusDataTypeRegistry::RegisterArrayTypeIfApplicable(FOptimusDataTypeHan
 			// Convert each element, store them in a separate buffer that is to be uploaded
 			FScriptArrayHelper ArrayHelper(ArrayProperty, InRawValue.GetData());
 			TArray<uint8>& Buffer = OutShaderValue.ArrayList[0].ArrayOfValues;
+			Buffer.Reset();
 			Buffer.AddZeroed(ElementShaderValueSize * ArrayHelper.Num());
 					
 			for (int32 Index = 0; Index < ArrayHelper.Num(); Index++)
