[RigVM & Deformer Graph] added optimusdeformer as allowed pin class in rigvm

diff --git a/Engine/Plugins/Animation/DeformerGraph/Source/OptimusCore/Private/OptimusCoreModule.cpp b/Engine/Plugins/Animation/DeformerGraph/Source/OptimusCore/Private/OptimusCoreModule.cpp
--- a/Engine/Plugins/Animation/DeformerGraph/Source/OptimusCore/Private/OptimusCoreModule.cpp
+++ b/Engine/Plugins/Animation/DeformerGraph/Source/OptimusCore/Private/OptimusCoreModule.cpp
@@ -14,6 +14,7 @@
 #include "AssetRegistry/AssetRegistryModule.h"
 #include "AssetRegistry/AssetData.h"
 #include "Nodes/OptimusNode_FunctionReference.h"
+#include "RigVMCore/RigVMRegistry.h"
 
 // Unique serialization id for Optimus .
 const FGuid FOptimusObjectVersion::GUID(0x93ede1aa, 0x10ca7375, 0x4df98a28, 0x49b157a0);
@@ -29,7 +30,14 @@ void FOptimusCoreModule::StartupModule()
 	FOptimusDataTypeRegistry::RegisterBuiltinTypes();
 	FOptimusDataTypeRegistry::RegisterEngineCallbacks();
 	UOptimusComputeDataInterface::RegisterAllTypes();
-	
+
+	FRigVMRegistry& RigVMRegistry = FRigVMRegistry::Get();
+	static TPair<UClass*, FRigVMRegistry::ERegisterObjectOperation> const AllowedObjectTypes[] =
+	{
+		{ UOptimusDeformer::StaticClass(), FRigVMRegistry::ERegisterObjectOperation::Class },
+	};
+
+	RigVMRegistry.RegisterObjectTypes(AllowedObjectTypes);
 }
 
 void FOptimusCoreModule::ShutdownModule()
diff --git a/Engine/Plugins/Runtime/RigVM/Source/RigVM/Private/RigVMCore/RigVMRegistry.cpp b/Engine/Plugins/Runtime/RigVM/Source/RigVM/Private/RigVMCore/RigVMRegistry.cpp
--- a/Engine/Plugins/Runtime/RigVM/Source/RigVM/Private/RigVMCore/RigVMRegistry.cpp
+++ b/Engine/Plugins/Runtime/RigVM/Source/RigVM/Private/RigVMCore/RigVMRegistry.cpp
@@ -1288,6 +1288,10 @@ bool FRigVMRegistry_NoLock::IsAllowedType_NoLock(const FProperty* InProperty) co
 	{
 		return IsAllowedType_NoLock(ObjectProperty->PropertyClass);
 	}
+	if(const FSoftObjectProperty* SoftObjectProperty = CastField<FSoftObjectProperty>(InProperty))
+	{
+		return IsAllowedType_NoLock(SoftObjectProperty->PropertyClass);
+	}
 	if(const FEnumProperty* EnumProperty = CastField<FEnumProperty>(InProperty))
 	{
 		return IsAllowedType_NoLock(EnumProperty->GetEnum());
