[Deformer Graph] fixed out of bound crash when using bone mask datainterface. Old code didnot consider different LOD can have different bone map per section

diff --git a/Engine/Plugins/Animation/DeformerGraph/Source/OptimusCore/Private/DataInterfaces/OptimusDataInterfaceSkinWeightsAsVertexMask.cpp b/Engine/Plugins/Animation/DeformerGraph/Source/OptimusCore/Private/DataInterfaces/OptimusDataInterfaceSkinWeightsAsVertexMask.cpp
--- a/Engine/Plugins/Animation/DeformerGraph/Source/OptimusCore/Private/DataInterfaces/OptimusDataInterfaceSkinWeightsAsVertexMask.cpp
+++ b/Engine/Plugins/Animation/DeformerGraph/Source/OptimusCore/Private/DataInterfaces/OptimusDataInterfaceSkinWeightsAsVertexMask.cpp
@@ -267,9 +267,9 @@ FComputeDataProviderRenderProxy* UOptimusSkinWeightsAsVertexMaskDataProvider::Ge
 		ExpandTowardsLeaf = FMath::Max(0, ExpandTowardsLeaf);
 		
 		FSkeletalMeshObject* SkeletalMeshObject = SkeletalMesh->MeshObject;
-		const int32 LodIndex = SkeletalMeshObject->GetLOD();
+		const int32 CurrentLodIndex = SkeletalMeshObject->GetLOD();
 		FSkeletalMeshRenderData const& SkeletalMeshRenderData = SkeletalMeshObject->GetSkeletalMeshRenderData();
-		FSkeletalMeshLODRenderData const& LodRenderData = SkeletalMeshRenderData.LODRenderData[LodIndex];
+		FSkeletalMeshLODRenderData const& CurrentLodRenderData = SkeletalMeshRenderData.LODRenderData[CurrentLodIndex];
 
 		bool bSkinWeightBufferReady = false;
 
@@ -277,17 +277,17 @@ FComputeDataProviderRenderProxy* UOptimusSkinWeightsAsVertexMaskDataProvider::Ge
 		{
 			bSkinWeightBufferReady = true;
 		}
-		else if (LodRenderData.SkinWeightProfilesData.ContainsProfile(SkinWeightProfile))
+		else if (CurrentLodRenderData.SkinWeightProfilesData.ContainsProfile(SkinWeightProfile))
 		{
 			// Retrieve this profile's skin weight buffer
 			const FSkinWeightProfileStack ProfileStack{SkinWeightProfile};
-			FSkinWeightVertexBuffer* Buffer = LodRenderData.SkinWeightProfilesData.GetOverrideBuffer(ProfileStack);
+			FSkinWeightVertexBuffer* Buffer = CurrentLodRenderData.SkinWeightProfilesData.GetOverrideBuffer(ProfileStack);
 			bSkinWeightBufferReady = Buffer ? true : false;
 		}
 
 		if (!bSkinWeightBufferReady)
 		{
-			if (LodRenderData.SkinWeightProfilesData.ContainsProfile(SkinWeightProfile))
+			if (CurrentLodRenderData.SkinWeightProfilesData.ContainsProfile(SkinWeightProfile))
 			{
 				const FSkinWeightProfileStack ProfileStack{SkinWeightProfile};
 				// Put in a skin weight profile request
@@ -305,7 +305,7 @@ FComputeDataProviderRenderProxy* UOptimusSkinWeightsAsVertexMaskDataProvider::Ge
 				bIsInitialized = true;
 
 				CachedSelectedBones.Reset();
-				CachedBoneIsSelectedPerSection.Reset();
+				CachedBoneIsSelectedPerSectionPerLod.Reset();
 				
 				const FReferenceSkeleton& RefSkeleton = SkeletalMesh->GetSkinnedAsset()->GetRefSkeleton();
 
@@ -371,23 +371,29 @@ FComputeDataProviderRenderProxy* UOptimusSkinWeightsAsVertexMaskDataProvider::Ge
 					}
 				}
 				
-				for (int32 SectionIndex = 0; SectionIndex < LodRenderData.RenderSections.Num(); ++SectionIndex)
+				for (int32 LODIndex = 0; LODIndex < SkeletalMeshRenderData.LODRenderData.Num(); LODIndex++)
 				{
-					TArray<uint32>& BoneIsSelectedRef = CachedBoneIsSelectedPerSection.AddDefaulted_GetRef();
+					FSkeletalMeshLODRenderData const& LocalLodRenderData = SkeletalMeshRenderData.LODRenderData[CurrentLodIndex];
+					
+					TArray<TArray<uint32>>& CachedBoneIsSelectedPerSectionRef = CachedBoneIsSelectedPerSectionPerLod.AddDefaulted_GetRef();
+					for (int32 SectionIndex = 0; SectionIndex < LocalLodRenderData.RenderSections.Num(); ++SectionIndex)
+					{
+						TArray<uint32>& BoneIsSelectedRef = CachedBoneIsSelectedPerSectionRef.AddDefaulted_GetRef();
 				
-					FSkelMeshRenderSection const& RenderSection = LodRenderData.RenderSections[SectionIndex];
-					const int32 NumBones = RenderSection.BoneMap.Num();
-					BoneIsSelectedRef.Init(false, NumBones);
+						FSkelMeshRenderSection const& RenderSection = LocalLodRenderData.RenderSections[SectionIndex];
+						const int32 NumBones = RenderSection.BoneMap.Num();
+						BoneIsSelectedRef.Init(false, NumBones);
 
-					for (int32 BoneIndex = 0; BoneIndex < NumBones; BoneIndex++)
-					{
-						int32 FinalBoneIndex = RenderSection.BoneMap[BoneIndex];
-						if (CachedSelectedBones.Contains(FinalBoneIndex))
+						for (int32 BoneIndex = 0; BoneIndex < NumBones; BoneIndex++)
 						{
-							BoneIsSelectedRef[BoneIndex] = true;
+							int32 FinalBoneIndex = RenderSection.BoneMap[BoneIndex];
+							if (CachedSelectedBones.Contains(FinalBoneIndex))
+							{
+								BoneIsSelectedRef[BoneIndex] = true;
+							}
 						}
 					}
-				}	
+				}
 			}
 
 			if (bDebugDrawIncludedBones)
@@ -421,7 +427,7 @@ FComputeDataProviderRenderProxy* UOptimusSkinWeightsAsVertexMaskDataProvider::Ge
 			
 			Proxy->SkeletalMeshObject = SkeletalMeshObject;
 			Proxy->SkinWeightProfile = SkinWeightProfile;
-			Proxy->BoneIsSelectedPerSection = CachedBoneIsSelectedPerSection;
+			Proxy->BoneIsSelectedPerSectionPerLod = CachedBoneIsSelectedPerSectionPerLod;
 		}
 	}
 
@@ -460,7 +466,7 @@ bool FOptimusSkinWeightsAsVertexMaskDataProviderProxy::IsValid(FValidationData c
 		return false;
 	}
 
-	if (BoneIsSelectedPerSection.IsEmpty())
+	if (BoneIsSelectedPerSectionPerLod.IsEmpty())
 	{
 		return false;
 	}
@@ -482,7 +488,7 @@ void FOptimusSkinWeightsAsVertexMaskDataProviderProxy::AllocateResources(FRDGBui
 	{
 		int32 NumBones = LodRenderData->RenderSections[InvocationIndex].BoneMap.Num();
 		
-		const TArray<uint32>& BoneIsSelected = BoneIsSelectedPerSection[InvocationIndex];
+		const TArray<uint32>& BoneIsSelected = BoneIsSelectedPerSectionPerLod[LodIndex][InvocationIndex];
 
 		int32 Stride = sizeof(uint32);
 		int32 ByteSize = NumBones * Stride;
diff --git a/Engine/Plugins/Animation/DeformerGraph/Source/OptimusCore/Private/DataInterfaces/OptimusDataInterfaceSkinWeightsAsVertexMask.h b/Engine/Plugins/Animation/DeformerGraph/Source/OptimusCore/Private/DataInterfaces/OptimusDataInterfaceSkinWeightsAsVertexMask.h
--- a/Engine/Plugins/Animation/DeformerGraph/Source/OptimusCore/Private/DataInterfaces/OptimusDataInterfaceSkinWeightsAsVertexMask.h
+++ b/Engine/Plugins/Animation/DeformerGraph/Source/OptimusCore/Private/DataInterfaces/OptimusDataInterfaceSkinWeightsAsVertexMask.h
@@ -108,7 +108,7 @@ public:
 
 	bool bIsInitialized = false;
 	TSet<int32> CachedSelectedBones;
-	TArray<TArray<uint32>> CachedBoneIsSelectedPerSection;
+	TArray<TArray<TArray<uint32>>> CachedBoneIsSelectedPerSectionPerLod;
 	
 	//~ Begin UComputeDataProvider Interface
 	FComputeDataProviderRenderProxy* GetRenderProxy() override;
@@ -140,7 +140,7 @@ public:
 	
 	FName SkinWeightProfile = NAME_None;
 
-	TArray<TArray<uint32>> BoneIsSelectedPerSection;
+	TArray<TArray<TArray<uint32>>> BoneIsSelectedPerSectionPerLod;
 	TArray<FRDGBufferRef> BoneIsSelectedBuffersPerSection;
 	TArray<FRDGBufferSRVRef> BoneIsSelectedBufferSRVsPerSection;
 };
