[Deformer Graph] Added error message when compiling deformer graph that contains debug draw data interface with unsupported platform project settings.

diff --git a/Engine/Plugins/Animation/DeformerGraph/Source/OptimusCore/Private/DataInterfaces/OptimusDataInterfaceDebugDraw.cpp b/Engine/Plugins/Animation/DeformerGraph/Source/OptimusCore/Private/DataInterfaces/OptimusDataInterfaceDebugDraw.cpp
--- a/Engine/Plugins/Animation/DeformerGraph/Source/OptimusCore/Private/DataInterfaces/OptimusDataInterfaceDebugDraw.cpp
+++ b/Engine/Plugins/Animation/DeformerGraph/Source/OptimusCore/Private/DataInterfaces/OptimusDataInterfaceDebugDraw.cpp
@@ -2,6 +2,7 @@
 
 #include "OptimusDataInterfaceDebugDraw.h"
 
+#include "OptimusCoreModule.h"
 #include "Components/SkinnedMeshComponent.h"
 #include "ComputeFramework/ShaderParamTypeDefinition.h"
 #include "OptimusDataDomain.h"
@@ -13,6 +14,7 @@
 
 #include UE_INLINE_GENERATED_CPP_BY_NAME(OptimusDataInterfaceDebugDraw)
 
+#define LOCTEXT_NAMESPACE "OptimusDebugDrawDataInterface"
 
 UOptimusDebugDrawDataInterface::UOptimusDebugDrawDataInterface()
 {
@@ -56,6 +58,18 @@ void UOptimusDebugDrawDataInterface::RegisterTypes()
 		EOptimusDataTypeUsageFlags::DataInterfaceOutput | EOptimusDataTypeUsageFlags::PinType);
 }
 
+TOptional<FText> UOptimusDebugDrawDataInterface::ValidateForCompile() const
+{
+	if (!bIsSupported)
+	{
+		return LOCTEXT("Unsupported Platform Setting", "Deformer Graph contains unsupported DebugDraw Data Interface and may not "
+						"function properly. Either make sure DirectX 12 and Shader Model 6 are enabled in Project Settings, "
+						"or remove the data interface from your graph completely");
+	}
+
+	return {};
+}
+
 void UOptimusDebugDrawDataInterface::GetSupportedInputs(TArray<FShaderFunctionDefinition>& OutFunctions) const
 {
 	FShaderValueTypeHandle DebugDrawType = FOptimusDataTypeRegistry::Get().FindType(FName("FDebugDraw"))->ShaderValueType;
@@ -113,6 +127,12 @@ UComputeDataProvider* UOptimusDebugDrawDataInterface::CreateDataProvider(TObject
 	UOptimusDebugDrawDataProvider* Provider = NewObject<UOptimusDebugDrawDataProvider>();
 	Provider->PrimitiveComponent = Cast<UPrimitiveComponent>(InBinding);
 	Provider->DebugDrawParameters = DebugDrawParameters;
+
+	if (!bIsSupported)
+	{
+		UE_LOG(LogOptimusCore, Error, TEXT("Deformer Graph aborted because it relies on unsupported DebugDraw DataInterface. Either remove the data interface or make sure DirectX 12 and Shader Model 6 is enabled in project settings for DebugDraw Data Interface to function, since DXC is required for shader compilation."));
+	}
+	
 	return Provider;
 }
 
@@ -205,3 +225,5 @@ void FOptimusDebugDrawDataProviderProxy::GatherDispatchData(FDispatchData const&
 		Parameters.RWEntryBuffer = CachedParameters.ShaderPrint_RWEntryBuffer;
 	}
 }
+
+#undef LOCTEXT_NAMESPACE
\ No newline at end of file
diff --git a/Engine/Plugins/Animation/DeformerGraph/Source/OptimusCore/Private/DataInterfaces/OptimusDataInterfaceDebugDraw.h b/Engine/Plugins/Animation/DeformerGraph/Source/OptimusCore/Private/DataInterfaces/OptimusDataInterfaceDebugDraw.h
--- a/Engine/Plugins/Animation/DeformerGraph/Source/OptimusCore/Private/DataInterfaces/OptimusDataInterfaceDebugDraw.h
+++ b/Engine/Plugins/Animation/DeformerGraph/Source/OptimusCore/Private/DataInterfaces/OptimusDataInterfaceDebugDraw.h
@@ -52,6 +52,7 @@ public:
 	TArray<FOptimusCDIPinDefinition> GetPinDefinitions() const override;
 	TSubclassOf<UActorComponent> GetRequiredComponentClass() const override;
 	void RegisterTypes() override;
+	TOptional<FText> ValidateForCompile() const override;
 	//~ End UOptimusComputeDataInterface Interface
 	
 	//~ Begin UComputeDataInterface Interface
diff --git a/Engine/Plugins/Animation/DeformerGraph/Source/OptimusCore/Private/Nodes/OptimusNode_DataInterface.cpp b/Engine/Plugins/Animation/DeformerGraph/Source/OptimusCore/Private/Nodes/OptimusNode_DataInterface.cpp
--- a/Engine/Plugins/Animation/DeformerGraph/Source/OptimusCore/Private/Nodes/OptimusNode_DataInterface.cpp
+++ b/Engine/Plugins/Animation/DeformerGraph/Source/OptimusCore/Private/Nodes/OptimusNode_DataInterface.cpp
@@ -91,6 +91,15 @@ TOptional<FText> UOptimusNode_DataInterface::ValidateForCompile(const FOptimusPi
 			}
 		}
 	}
+
+	if (DataInterfaceData)
+	{
+		const TOptional<FText> ErrorIfDetected = DataInterfaceData->ValidateForCompile();
+		if (ErrorIfDetected.IsSet())
+		{
+			return ErrorIfDetected;
+		}
+	}
 	
 	return {};
 }
diff --git a/Engine/Plugins/Animation/DeformerGraph/Source/OptimusCore/Public/OptimusComputeDataInterface.h b/Engine/Plugins/Animation/DeformerGraph/Source/OptimusCore/Public/OptimusComputeDataInterface.h
--- a/Engine/Plugins/Animation/DeformerGraph/Source/OptimusCore/Public/OptimusComputeDataInterface.h
+++ b/Engine/Plugins/Animation/DeformerGraph/Source/OptimusCore/Public/OptimusComputeDataInterface.h
@@ -193,6 +193,8 @@ public:
 		return true;
 	}
 
+	virtual TOptional<FText> ValidateForCompile() const { return{}; };
+
 	/// Returns all known UOptimusComputeDataInterface-derived classes.
 	static TArray<TSubclassOf<UOptimusComputeDataInterface>> GetAllComputeDataInterfaceClasses();
 
