[Deformer Graph Editor] fix debug draw not showing up in deformer graph editor viewport

diff --git a/Engine/Plugins/Animation/DeformerGraph/Source/OptimusEditor/Private/OptimusEditor.cpp b/Engine/Plugins/Animation/DeformerGraph/Source/OptimusEditor/Private/OptimusEditor.cpp
--- a/Engine/Plugins/Animation/DeformerGraph/Source/OptimusEditor/Private/OptimusEditor.cpp
+++ b/Engine/Plugins/Animation/DeformerGraph/Source/OptimusEditor/Private/OptimusEditor.cpp
@@ -52,6 +52,7 @@
 #include "ISourceCodeAccessor.h"
 #include "OptimusEditorStyle.h"
 #include "Toolkits/ToolkitManager.h"
+#include "EditorViewportClient.h"
 
 
 #define LOCTEXT_NAMESPACE "OptimusEditor"
@@ -1073,6 +1074,7 @@ void FOptimusEditor::HandleViewportCreated(
 	)
 {
 	ViewportWidget = InPersonaViewport;
+	ViewportWidget->GetViewportClient().EngineShowFlags.ShaderPrint = 1;
 }
 
 
