Add fit (F) command to the deformer graph editor

diff --git a/Engine/Plugins/Animation/DeformerGraph/Source/OptimusEditor/Private/OptimusEditor.cpp b/Engine/Plugins/Animation/DeformerGraph/Source/OptimusEditor/Private/OptimusEditor.cpp
--- a/Engine/Plugins/Animation/DeformerGraph/Source/OptimusEditor/Private/OptimusEditor.cpp
+++ b/Engine/Plugins/Animation/DeformerGraph/Source/OptimusEditor/Private/OptimusEditor.cpp
@@ -317,6 +317,15 @@ void FOptimusEditor::Compile()
 	}
 }
 
+void FOptimusEditor::ZoomToFitGraph()
+{
+	if(GraphEditorWidget)
+	{
+		// This will focus on the selected nodes, if any. If no nodes are selected, it will focus the whole graph.
+		constexpr bool bOnlySelection = true;	
+		GraphEditorWidget->ZoomToFit(bOnlySelection);
+	}
+}
 
 bool FOptimusEditor::CanCompile() const
 {
@@ -959,8 +968,11 @@ void FOptimusEditor::BindCommands()
 		Commands.Compile,
 		FExecuteAction::CreateSP(this, &FOptimusEditor::Compile),
 		FCanExecuteAction::CreateSP(this, &FOptimusEditor::CanCompile));
-}
 
+	ToolkitCommands->MapAction(
+		Commands.ZoomToFitGraph,
+		FExecuteAction::CreateSP(this, &FOptimusEditor::ZoomToFitGraph));
+}
 
 void FOptimusEditor::HandlePreviewSceneCreated(const TSharedRef<IPersonaPreviewScene>& InPreviewScene)
 {
diff --git a/Engine/Plugins/Animation/DeformerGraph/Source/OptimusEditor/Private/OptimusEditor.h b/Engine/Plugins/Animation/DeformerGraph/Source/OptimusEditor/Private/OptimusEditor.h
--- a/Engine/Plugins/Animation/DeformerGraph/Source/OptimusEditor/Private/OptimusEditor.h
+++ b/Engine/Plugins/Animation/DeformerGraph/Source/OptimusEditor/Private/OptimusEditor.h
@@ -110,6 +110,7 @@ private:
 	// ----------------------------------------------------------------------------------------
 	// Editor commands
 	void Compile();
+	void ZoomToFitGraph();
 
 	bool CanCompile() const;
 
diff --git a/Engine/Plugins/Animation/DeformerGraph/Source/OptimusEditor/Private/OptimusEditorCommands.cpp b/Engine/Plugins/Animation/DeformerGraph/Source/OptimusEditor/Private/OptimusEditorCommands.cpp
--- a/Engine/Plugins/Animation/DeformerGraph/Source/OptimusEditor/Private/OptimusEditorCommands.cpp
+++ b/Engine/Plugins/Animation/DeformerGraph/Source/OptimusEditor/Private/OptimusEditorCommands.cpp
@@ -22,7 +22,8 @@ FOptimusEditorCommands::FOptimusEditorCommands()
 
 void FOptimusEditorCommands::RegisterCommands()
 {
-	UI_COMMAND(Compile, "Compile", "Compile the current deformer graph into a compute kernel graph.", EUserInterfaceActionType::Button, FInputChord(EKeys::F7));
+	UI_COMMAND(Compile, "Compile", "Compile the current deformer graph into a compute kernel graph.", EUserInterfaceActionType::Button, FInputChord(EKeys::F7))
+	UI_COMMAND(ZoomToFitGraph, "ZoomToFitGraph", "Fit the graph in the editor viewport.", EUserInterfaceActionType::None, FInputChord(EKeys::F));
 }
 
 
diff --git a/Engine/Plugins/Animation/DeformerGraph/Source/OptimusEditor/Private/OptimusEditorCommands.h b/Engine/Plugins/Animation/DeformerGraph/Source/OptimusEditor/Private/OptimusEditorCommands.h
--- a/Engine/Plugins/Animation/DeformerGraph/Source/OptimusEditor/Private/OptimusEditorCommands.h
+++ b/Engine/Plugins/Animation/DeformerGraph/Source/OptimusEditor/Private/OptimusEditorCommands.h
@@ -16,4 +16,5 @@ public:
 	void RegisterCommands() override;
 
 	TSharedPtr<FUICommandInfo> Compile;
+	TSharedPtr<FUICommandInfo> ZoomToFitGraph;
 };
\ No newline at end of file
