Display name for DG data interface

diff --git a/Engine/Plugins/Animation/DeformerGraph/Source/OptimusCore/Private/Actions/OptimusNodeGraphActions.cpp b/Engine/Plugins/Animation/DeformerGraph/Source/OptimusCore/Private/Actions/OptimusNodeGraphActions.cpp
--- a/Engine/Plugins/Animation/DeformerGraph/Source/OptimusCore/Private/Actions/OptimusNodeGraphActions.cpp
+++ b/Engine/Plugins/Animation/DeformerGraph/Source/OptimusCore/Private/Actions/OptimusNodeGraphActions.cpp
@@ -772,7 +772,7 @@ bool FOptimusNodeGraphAction_ConnectAdderPin::Do(IOptimusPathResolver* InRoot)
 
 	UOptimusNodePin* SourcePin = InRoot->ResolvePinPath(SourcePinPath);
 
-	FName Name = SourcePin->GetFName();
+	FName Name = FName(SourcePin->GetDisplayName().ToString());
 	if (IOptimusUnnamedNodePinProvider* UnnamedPinProvider = Cast<IOptimusUnnamedNodePinProvider>(SourcePin->GetOwningNode()))
 	{
 		Name = UnnamedPinProvider->GetNameForAdderPin(SourcePin);
diff --git a/Engine/Plugins/Animation/DeformerGraph/Source/OptimusCore/Private/Nodes/OptimusNode_DataInterface.cpp b/Engine/Plugins/Animation/DeformerGraph/Source/OptimusCore/Private/Nodes/OptimusNode_DataInterface.cpp
--- a/Engine/Plugins/Animation/DeformerGraph/Source/OptimusCore/Private/Nodes/OptimusNode_DataInterface.cpp
+++ b/Engine/Plugins/Animation/DeformerGraph/Source/OptimusCore/Private/Nodes/OptimusNode_DataInterface.cpp
@@ -222,6 +222,20 @@ void UOptimusNode_DataInterface::InitializeTransientData()
 	}
 }
 
+void UOptimusNode_DataInterface::InitializePinsDisplayName()
+{
+	for (const FOptimusCDIPinDefinition& Def: DataInterfaceData->GetPinDefinitions())
+	{
+		if (ensure(!Def.PinName.IsNone()))
+		{
+			if(UOptimusNodePin* Pin = GetPinByName(Def.PinName))
+			{
+				Pin->SetDisplayName(Def.DisplayName);
+			}
+		}
+	}
+}
+
 void UOptimusNode_DataInterface::Serialize(FArchive& Ar)
 {
  	Super::Serialize(Ar);
@@ -482,7 +496,7 @@ void UOptimusNode_DataInterface::CreatePinFromDefinition(const FOptimusCDIPinDef
 				*ValueTypeHandle->ToString());
 			return;
 		}
-
+		
 		if (bSupportUndo)
 		{
 			AddPin(InDefinition.PinName, EOptimusNodePinDirection::Output, {}, PinDataType);
diff --git a/Engine/Plugins/Animation/DeformerGraph/Source/OptimusCore/Private/Nodes/OptimusNode_DataInterface.h b/Engine/Plugins/Animation/DeformerGraph/Source/OptimusCore/Private/Nodes/OptimusNode_DataInterface.h
--- a/Engine/Plugins/Animation/DeformerGraph/Source/OptimusCore/Private/Nodes/OptimusNode_DataInterface.h
+++ b/Engine/Plugins/Animation/DeformerGraph/Source/OptimusCore/Private/Nodes/OptimusNode_DataInterface.h
@@ -86,6 +86,9 @@ private:
 	void CreatePropertyPinsFromDataInterface(const UOptimusComputeDataInterface *InDataInterface, bool bSupportUndo);
 	void CreateComponentPin();
 
+	/** Allows the node to initialize all the pins display names from the definitions */
+	virtual void InitializePinsDisplayName() override;
+
 protected:
 	friend class UOptimusDeformer;
 	friend class UOptimusNodeGraph;
diff --git a/Engine/Plugins/Animation/DeformerGraph/Source/OptimusCore/Private/OptimusNode.cpp b/Engine/Plugins/Animation/DeformerGraph/Source/OptimusCore/Private/OptimusNode.cpp
--- a/Engine/Plugins/Animation/DeformerGraph/Source/OptimusCore/Private/OptimusNode.cpp
+++ b/Engine/Plugins/Animation/DeformerGraph/Source/OptimusCore/Private/OptimusNode.cpp
@@ -135,6 +135,18 @@ UOptimusNodeGraph* UOptimusNode::GetOwningGraph() const
 	return Cast<UOptimusNodeGraph>(GetOuter());
 }
 
+UOptimusNodePin* UOptimusNode::GetPinByName(const FName& PinName) const
+{
+	for (UOptimusNodePin* Pin : Pins)
+	{
+		if(Pin->GetName() == PinName)
+		{
+			return Pin;
+		}
+	}
+	return nullptr;
+}
+
 TArray<UOptimusNodePin*> UOptimusNode::GetPinsByDirection(EOptimusNodePinDirection InDirection,
 	bool bInRecursive) const
 {
@@ -408,6 +420,9 @@ void UOptimusNode::PostCreateNode()
 	{
 		TGuardValue<bool> NodeConstructionGuard(bConstructingNode, true);
 		ConstructNode();
+
+		// Update the pins display names
+		InitializePinsDisplayName();
 	}
 }
 
@@ -429,6 +444,9 @@ void UOptimusNode::PostLoad()
 
 	// Now that all serialized data is ready, let's generate any transient data
 	InitializeTransientData();
+
+	// Update the pins display names
+	InitializePinsDisplayName();
 }
 
 void UOptimusNode::PostLoadNodeSpecificData()
diff --git a/Engine/Plugins/Animation/DeformerGraph/Source/OptimusCore/Private/OptimusNodePin.cpp b/Engine/Plugins/Animation/DeformerGraph/Source/OptimusCore/Private/OptimusNodePin.cpp
--- a/Engine/Plugins/Animation/DeformerGraph/Source/OptimusCore/Private/OptimusNodePin.cpp
+++ b/Engine/Plugins/Animation/DeformerGraph/Source/OptimusCore/Private/OptimusNodePin.cpp
@@ -109,11 +109,11 @@ FName UOptimusNodePin::GetUniqueName() const
 
 FText UOptimusNodePin::GetDisplayName() const
 {
-	
 	const FOptimusDataTypeHandle BoolType = FOptimusDataTypeRegistry::Get().FindType(*FBoolProperty::StaticClass());
 	const bool bIsBool = GetDataType() == BoolType;
-	
-	return FText::FromString(FName::NameToDisplayString(GetName(), bIsBool));
+
+	const FString PinName = !DisplayName.IsNone() ? DisplayName.ToString() : GetName();
+	return FText::FromString(FName::NameToDisplayString(PinName, bIsBool));
 }
 
 
@@ -615,6 +615,16 @@ bool UOptimusNodePin::SetName(FName InName)
 	return true;
 }
 
+bool UOptimusNodePin::SetDisplayName(const FName& InDisplayName)
+{
+	if (DisplayName == InDisplayName)
+	{
+		return false;
+	}
+	DisplayName = InDisplayName;
+	
+	return true;
+}
 
 void UOptimusNodePin::Notify(EOptimusGraphNotifyType InNotifyType) const
 {
diff --git a/Engine/Plugins/Animation/DeformerGraph/Source/OptimusCore/Public/OptimusComputeDataInterface.h b/Engine/Plugins/Animation/DeformerGraph/Source/OptimusCore/Public/OptimusComputeDataInterface.h
--- a/Engine/Plugins/Animation/DeformerGraph/Source/OptimusCore/Public/OptimusComputeDataInterface.h
+++ b/Engine/Plugins/Animation/DeformerGraph/Source/OptimusCore/Public/OptimusComputeDataInterface.h
@@ -22,11 +22,13 @@ struct FOptimusCDIPinDefinition
 	FOptimusCDIPinDefinition(
 		const FName InPinName,
 		const FString InDataFunctionName,
-		const bool bInMutable = true
+		const bool bInMutable = true,
+		const FName InDisplayName = NAME_None
 		) :
 		PinName(InPinName),
 		DataFunctionName(InDataFunctionName),
-		bMutable(bInMutable)
+		bMutable(bInMutable),
+		DisplayName(InDisplayName)
 	{ }
 
 	// A single level context lookup.
@@ -35,12 +37,14 @@ struct FOptimusCDIPinDefinition
 		const FString InDataFunctionName,
 		const FName InContextName,
 		const FString InCountFunctionName,
-		const bool bInMutable = true
+		const bool bInMutable = true,
+		const FName InDisplayName = NAME_None
 		) :
 		PinName(InPinName),
 		DataFunctionName(InDataFunctionName),
 		DataDimensions{{InContextName, InCountFunctionName}},
-		bMutable(bInMutable)
+		bMutable(bInMutable),
+		DisplayName(InDisplayName)
 	{ }
 
 	FOptimusCDIPinDefinition(
@@ -49,24 +53,28 @@ struct FOptimusCDIPinDefinition
 		const FName InContextName,
 		const int32 InMultiplier,
 		const FString InCountFunctionName,
-		const bool bInMutable = true
+		const bool bInMutable = true,
+		const FName InDisplayName = NAME_None
 		) :
 		PinName(InPinName),
 		DataFunctionName(InDataFunctionName),
 		DataDimensions{{InContextName, InCountFunctionName}},
 		DomainMultiplier(FMath::Max(1, InMultiplier)),
-		bMutable(bInMutable)
+		bMutable(bInMutable),
+		DisplayName(InDisplayName)
 	{ }
 	
 	FOptimusCDIPinDefinition(
 		const FName InPinName,
 		const FString InDataFunctionName,
 		const std::initializer_list<FDimensionInfo> InContexts,
-		const bool bInMutable = true
+		const bool bInMutable = true,
+		const FName InDisplayName = NAME_None
 		) :
 		PinName(InPinName),
 		DataFunctionName(InDataFunctionName),
-		bMutable(bInMutable)
+		bMutable(bInMutable),
+		DisplayName(InDisplayName)
 	{
 		for (FDimensionInfo ContextInfo: InContexts)
 		{
@@ -122,6 +130,9 @@ struct FOptimusCDIPinDefinition
 	int32 DomainMultiplier = 1;
 
 	bool bMutable = true;
+	
+	// Name used for display
+	FName DisplayName;
 };
 
 struct FOptimusCDIPropertyPinDefinition
diff --git a/Engine/Plugins/Animation/DeformerGraph/Source/OptimusCore/Public/OptimusNode.h b/Engine/Plugins/Animation/DeformerGraph/Source/OptimusCore/Public/OptimusNode.h
--- a/Engine/Plugins/Animation/DeformerGraph/Source/OptimusCore/Public/OptimusNode.h
+++ b/Engine/Plugins/Animation/DeformerGraph/Source/OptimusCore/Public/OptimusNode.h
@@ -105,6 +105,9 @@ public:
 	/** Returns the list of all the pins on this node */
 	TArrayView<UOptimusNodePin* const> GetPins() const { return Pins; }
 
+	/** Return the pin with a specific name */
+	UOptimusNodePin* GetPinByName(const FName& PinName) const;
+
 	/** Returns the list of pins on this node filtered by direction */
 	TArray<UOptimusNodePin*> GetPinsByDirection(EOptimusNodePinDirection InDirection, bool bInRecursive) const;
 	
@@ -228,6 +231,9 @@ protected:
 	// Allows the node to initialize any transient data that can be derived from serialized properties
 	virtual void InitializeTransientData();
 
+	/** Allows the node to initialize all the pins display names from the definitions */
+	virtual void InitializePinsDisplayName() {}
+
 	virtual bool ValidateConnection(
 		const UOptimusNodePin& InThisNodesPin,
 		const UOptimusNodePin& InOtherNodesPin,
diff --git a/Engine/Plugins/Animation/DeformerGraph/Source/OptimusCore/Public/OptimusNodePin.h b/Engine/Plugins/Animation/DeformerGraph/Source/OptimusCore/Public/OptimusNodePin.h
--- a/Engine/Plugins/Animation/DeformerGraph/Source/OptimusCore/Public/OptimusNodePin.h
+++ b/Engine/Plugins/Animation/DeformerGraph/Source/OptimusCore/Public/OptimusNodePin.h
@@ -80,6 +80,9 @@ public:
 	/// E.g: SetupGraph/LinearBlendSkinning1.Direction.X
 	FString GetPinPath() const;
 
+	/** Set the display name of the node pin */
+	bool SetDisplayName(const FName& InDisplayName);
+
 	/// Returns a pin path from a string. Returns an empty array if string is invalid or empty.
 	static TArray<FName> GetPinNamePathFromString(const FStringView InPinPathString);
 
@@ -216,4 +219,7 @@ PRAGMA_ENABLE_DEPRECATION_WARNINGS
 
 	UPROPERTY()
 	TArray<TObjectPtr<UOptimusNodePin>> SubPins;
+
+	UPROPERTY(Transient)
+	FName DisplayName;
 };
diff --git a/Engine/Plugins/Animation/DeformerGraph/Source/OptimusEditor/Private/OptimusEditorClipboard.cpp b/Engine/Plugins/Animation/DeformerGraph/Source/OptimusEditor/Private/OptimusEditorClipboard.cpp
--- a/Engine/Plugins/Animation/DeformerGraph/Source/OptimusEditor/Private/OptimusEditorClipboard.cpp
+++ b/Engine/Plugins/Animation/DeformerGraph/Source/OptimusEditor/Private/OptimusEditorClipboard.cpp
@@ -421,6 +421,7 @@ bool FOptimusEditorClipboard::ProcessPostCreateObject(UObject* InRootOuter, UObj
 			UOptimusNode* Node = Cast<UOptimusNode>(InNewObject);
 
 			Node->InitializeTransientData();
+			Node->InitializePinsDisplayName();
 			
 			return Graph->AddNodeDirect(Node);
 		}
@@ -449,6 +450,7 @@ bool FOptimusEditorClipboard::ProcessPostCreateObject(UObject* InRootOuter, UObj
 				for (UOptimusNode* Node : WorkingSubGraph->GetAllNodes())
 				{
 					Node->InitializeTransientData();
+					Node->InitializePinsDisplayName();
 				}
 
 				for (UOptimusNodeGraph* SubGraphToAdd : WorkingSubGraph->GetGraphs())
