[Deformer Graph] Fixed Trigger Graph cannot be triggered from blueprint, this was broken because it was using generated unique name for the graph name, which is always different from the trigger graph name users can set in the editor.

diff --git a/Engine/Plugins/Animation/DeformerGraph/Source/OptimusCore/Private/OptimusDeformer.cpp b/Engine/Plugins/Animation/DeformerGraph/Source/OptimusCore/Private/OptimusDeformer.cpp
--- a/Engine/Plugins/Animation/DeformerGraph/Source/OptimusCore/Private/OptimusDeformer.cpp
+++ b/Engine/Plugins/Animation/DeformerGraph/Source/OptimusCore/Private/OptimusDeformer.cpp
@@ -1499,6 +1499,10 @@ bool UOptimusDeformer::Compile()
 
 	auto ClearCompiledData = [&]()
 	{
+		for (FOptimusComputeGraphInfo& GraphInfo : ComputeGraphs )
+		{
+			Optimus::RemoveObject(GraphInfo.ComputeGraph);
+		}
 		ComputeGraphs.Reset();
 		DataInterfacePropertyOverrideMap.Reset();
 		ValueMap.Reset();
@@ -2611,23 +2615,21 @@ FOptimusNodeGraphCompilationResult UOptimusDeformer::CompileNodeGraphToComputeGr
 	TArray<FOptimusComputeGraphInfo>& GraphInfos = Result.ComputeGraphInfos;
 	for (EOptimusNodeGraphType GraphType : GraphTypes)
 	{
-		FString Name = InNodeGraph->GetName();
+		FString GraphName = InNodeGraph->GetName();
 		if (GraphType != InNodeGraph->GraphType)
 		{
 			check(GraphType == EOptimusNodeGraphType::Setup);
-			Name += TEXT("_Setup");
+			// Using "$" to avoid name clash with user provided graph name, see UOptimusNodeGraph::IsValidUserGraphName
+			GraphName += TEXT("$Setup");
 		}
 		
 		FOptimusComputeGraphInfo GraphInfo;
-		FName GraphName =
-			MakeUniqueObjectName(
-				this,
-				UOptimusComputeGraph::StaticClass(),
-				*Name
-				);
-		GraphInfo.GraphName = GraphName;
+		// For trigger graph, this graph name needs to match the node graph name so that user can use the node graph name to trigger it.
+		GraphInfo.GraphName = *GraphName;
 		GraphInfo.GraphType = GraphType;
-		GraphInfo.ComputeGraph = NewObject<UOptimusComputeGraph>(this, GraphInfo.GraphName);
+		// Avoid node graph and compute graph using the same name
+		FString ComputeGraphName = GraphName + TEXT("_ComputeGraph");
+		GraphInfo.ComputeGraph = NewObject<UOptimusComputeGraph>(this, *ComputeGraphName);
 
 		if (GraphType != InNodeGraph->GraphType)
 		{
