[Deformer Graph] avoid recursive inclusion of deformer source library

diff --git a/Engine/Plugins/Animation/DeformerGraph/Source/OptimusCore/Private/OptimusSource.cpp b/Engine/Plugins/Animation/DeformerGraph/Source/OptimusCore/Private/OptimusSource.cpp
--- a/Engine/Plugins/Animation/DeformerGraph/Source/OptimusCore/Private/OptimusSource.cpp
+++ b/Engine/Plugins/Animation/DeformerGraph/Source/OptimusCore/Private/OptimusSource.cpp
@@ -12,6 +12,19 @@ void UOptimusSource::SetSource(const FString& InText)
 	Modify();
 }
 
+#if WITH_EDITOR
+
+void UOptimusSource::PostEditChangeProperty(FPropertyChangedEvent& PropertyChangedEvent)
+{
+	if (PropertyChangedEvent.GetMemberPropertyName() == GET_MEMBER_NAME_CHECKED(UOptimusSource, AdditionalSources))
+	{
+		// Avoid recursive source inclusion
+		AdditionalSources.Remove(this);
+	}
+}
+
+#endif
+
 FString UOptimusSource::GetVirtualPath() const 
 {
 	FString ShaderPathName = GetPathName();
diff --git a/Engine/Plugins/Animation/DeformerGraph/Source/OptimusCore/Public/OptimusSource.h b/Engine/Plugins/Animation/DeformerGraph/Source/OptimusCore/Public/OptimusSource.h
--- a/Engine/Plugins/Animation/DeformerGraph/Source/OptimusCore/Public/OptimusSource.h
+++ b/Engine/Plugins/Animation/DeformerGraph/Source/OptimusCore/Public/OptimusSource.h
@@ -15,7 +15,10 @@ class OPTIMUSCORE_API UOptimusSource
 
 public:
 	void SetSource(const FString& InText);
-
+#if WITH_EDITOR
+	void PostEditChangeProperty(FPropertyChangedEvent& PropertyChangedEvent) override;
+#endif
+	
 	// Begin UComputeSource interface.
 	FString GetSource() const override { return SourceText; }
 	FString GetVirtualPath() const override;
