OptimusNode_DataInterface: Fixed access violation to the interface.

diff --git a/Engine/Plugins/Animation/DeformerGraph/Source/OptimusCore/Private/Nodes/OptimusNode_DataInterface.cpp b/Engine/Plugins/Animation/DeformerGraph/Source/OptimusCore/Private/Nodes/OptimusNode_DataInterface.cpp
--- a/Engine/Plugins/Animation/DeformerGraph/Source/OptimusCore/Private/Nodes/OptimusNode_DataInterface.cpp
+++ b/Engine/Plugins/Animation/DeformerGraph/Source/OptimusCore/Private/Nodes/OptimusNode_DataInterface.cpp
@@ -224,13 +224,16 @@ void UOptimusNode_DataInterface::InitializeTransientData()
 
 void UOptimusNode_DataInterface::InitializePinsDisplayName()
 {
-	for (const FOptimusCDIPinDefinition& Def: DataInterfaceData->GetPinDefinitions())
+	if (ensure(DataInterfaceData))
 	{
-		if (ensure(!Def.PinName.IsNone()))
+		for (const FOptimusCDIPinDefinition& Def : DataInterfaceData->GetPinDefinitions())
 		{
-			if(UOptimusNodePin* Pin = GetPinByName(Def.PinName))
+			if (ensure(!Def.PinName.IsNone()))
 			{
-				Pin->SetDisplayName(Def.DisplayName);
+				if (UOptimusNodePin* Pin = GetPinByName(Def.PinName))
+				{
+					Pin->SetDisplayName(Def.DisplayName);
+				}
 			}
 		}
 	}
