[Deformer Graph] Fixed bad motion vectors when LOD changes if the skeltal mesh has any deformer active

diff --git a/Engine/Plugins/Animation/DeformerGraph/Source/OptimusCore/Private/DataInterfaces/OptimusDataInterfaceSkinnedMeshRead.cpp b/Engine/Plugins/Animation/DeformerGraph/Source/OptimusCore/Private/DataInterfaces/OptimusDataInterfaceSkinnedMeshRead.cpp
--- a/Engine/Plugins/Animation/DeformerGraph/Source/OptimusCore/Private/DataInterfaces/OptimusDataInterfaceSkinnedMeshRead.cpp
+++ b/Engine/Plugins/Animation/DeformerGraph/Source/OptimusCore/Private/DataInterfaces/OptimusDataInterfaceSkinnedMeshRead.cpp
@@ -167,7 +167,7 @@ UComputeDataProvider* UOptimusSkinnedMeshReadDataInterface::CreateDataProvider(T
 
 FComputeDataProviderRenderProxy* UOptimusSkinnedMeshReadDataProvider::GetRenderProxy()
 {
-	return new FOptimusSkinnedMeshReadDataProviderProxy(SkinnedMesh, InputMask, DeformerInstance->OutputBuffersFromPreviousInstances);
+	return new FOptimusSkinnedMeshReadDataProviderProxy(SkinnedMesh, InputMask, DeformerInstance->OutputBuffersFromPreviousInstances, &LastLodIndexCachedByRenderProxy);
 }
 
 void UOptimusSkinnedMeshReadDataProvider::SetDeformerInstance(UOptimusDeformerInstance* InInstance)
@@ -175,11 +175,14 @@ void UOptimusSkinnedMeshReadDataProvider::SetDeformerInstance(UOptimusDeformerIn
 	DeformerInstance = InInstance;
 }
 
-FOptimusSkinnedMeshReadDataProviderProxy::FOptimusSkinnedMeshReadDataProviderProxy(USkinnedMeshComponent* InSkinnedMeshComponent, uint64 InInputMask, EMeshDeformerOutputBuffer InOutputBuffersFromPreviousInstances)
+FOptimusSkinnedMeshReadDataProviderProxy::FOptimusSkinnedMeshReadDataProviderProxy(USkinnedMeshComponent* InSkinnedMeshComponent, uint64 InInputMask, EMeshDeformerOutputBuffer InOutputBuffersFromPreviousInstances, int32* InLastLodIndexPtr) 
 {
 	SkeletalMeshObject = InSkinnedMeshComponent != nullptr ? InSkinnedMeshComponent->MeshObject : nullptr;
 	InputMask = InInputMask;
 	OutputBuffersFromPreviousInstances = InOutputBuffersFromPreviousInstances;
+	LastLodIndexPtr = InLastLodIndexPtr;
+
+	check(LastLodIndexPtr != nullptr);
 }
 
 bool FOptimusSkinnedMeshReadDataProviderProxy::IsValid(FValidationData const& InValidationData) const
@@ -209,9 +212,16 @@ void FOptimusSkinnedMeshReadDataProviderProxy::AllocateResources(FRDGBuilder& Gr
 	// Allocate required buffers
 	const int32 LodIndex = SkeletalMeshObject->GetLOD();
 	
+	bool bLodJustChanged = false;
+	if (LodIndex != *LastLodIndexPtr)
+	{
+		bLodJustChanged = true;	
+		*LastLodIndexPtr = LodIndex;
+	}
+	
 	if (InputMask & static_cast<uint64>(ESkinnedMeshReadDataInterfaceInputSelectorMask::Position) )
 	{
-		FRDGBuffer* PositionBuffer = FSkeletalMeshDeformerHelpers::AllocateVertexFactoryPositionBuffer(GraphBuilder, SkeletalMeshObject, LodIndex, TEXT("OptimusSkinnedMeshPosition"));
+		FRDGBuffer* PositionBuffer = FSkeletalMeshDeformerHelpers::AllocateVertexFactoryPositionBuffer(GraphBuilder, SkeletalMeshObject, LodIndex, bLodJustChanged,TEXT("OptimusSkinnedMeshPosition"));
 		PositionBufferUAV = GraphBuilder.CreateUAV(PositionBuffer, PF_R32_FLOAT, ERDGUnorderedAccessViewFlags::SkipBarrier);
 	}
 	else
diff --git a/Engine/Plugins/Animation/DeformerGraph/Source/OptimusCore/Private/DataInterfaces/OptimusDataInterfaceSkinnedMeshRead.h b/Engine/Plugins/Animation/DeformerGraph/Source/OptimusCore/Private/DataInterfaces/OptimusDataInterfaceSkinnedMeshRead.h
--- a/Engine/Plugins/Animation/DeformerGraph/Source/OptimusCore/Private/DataInterfaces/OptimusDataInterfaceSkinnedMeshRead.h
+++ b/Engine/Plugins/Animation/DeformerGraph/Source/OptimusCore/Private/DataInterfaces/OptimusDataInterfaceSkinnedMeshRead.h
@@ -60,7 +60,10 @@ public:
 	UPROPERTY(EditAnywhere, BlueprintReadWrite, Category = Binding)
 	TObjectPtr<USkinnedMeshComponent> SkinnedMesh = nullptr;
 
-	uint64 InputMask;
+	uint64 InputMask = 0;
+
+	// Served as persistent storage for the provider proxy, should not be used by the data provider itself
+	int32 LastLodIndexCachedByRenderProxy = 0;
 
 	//~ Begin UComputeDataProvider Interface
 	FComputeDataProviderRenderProxy* GetRenderProxy() override;
@@ -78,7 +81,7 @@ private:
 class FOptimusSkinnedMeshReadDataProviderProxy : public FComputeDataProviderRenderProxy
 {
 public:
-	FOptimusSkinnedMeshReadDataProviderProxy(USkinnedMeshComponent* InSkinnedMeshComponent, uint64 InInputMask, EMeshDeformerOutputBuffer InOutputBuffersWithValidData);
+	FOptimusSkinnedMeshReadDataProviderProxy(USkinnedMeshComponent* InSkinnedMeshComponent, uint64 InInputMask, EMeshDeformerOutputBuffer InOutputBuffersWithValidData, int32* InLastLodIndexPtr);
 
 	//~ Begin FComputeDataProviderRenderProxy Interface
 	bool IsValid(FValidationData const& InValidationData) const override;
@@ -90,8 +93,9 @@ public:
 private:
 	using FParameters = FSkinnedMeshReadDataInterfaceParameters;
 
-	FSkeletalMeshObject* SkeletalMeshObject;
-	uint64 InputMask;
+	FSkeletalMeshObject* SkeletalMeshObject = nullptr;
+	uint64 InputMask = 0;
+	int32* LastLodIndexPtr = nullptr;
 	EMeshDeformerOutputBuffer OutputBuffersFromPreviousInstances;
 
 	// Using UAV here because we might be reading and writing to these buffers in the same kernel
diff --git a/Engine/Plugins/Animation/DeformerGraph/Source/OptimusCore/Private/DataInterfaces/OptimusDataInterfaceSkinnedMeshWrite.cpp b/Engine/Plugins/Animation/DeformerGraph/Source/OptimusCore/Private/DataInterfaces/OptimusDataInterfaceSkinnedMeshWrite.cpp
--- a/Engine/Plugins/Animation/DeformerGraph/Source/OptimusCore/Private/DataInterfaces/OptimusDataInterfaceSkinnedMeshWrite.cpp
+++ b/Engine/Plugins/Animation/DeformerGraph/Source/OptimusCore/Private/DataInterfaces/OptimusDataInterfaceSkinnedMeshWrite.cpp
@@ -153,14 +153,17 @@ EMeshDeformerOutputBuffer UOptimusSkinnedMeshWriteDataInterface::GetOutputBuffer
 
 FComputeDataProviderRenderProxy* UOptimusSkinnedMeshWriteDataProvider::GetRenderProxy()
 {
-	return new FOptimusSkinnedMeshWriteDataProviderProxy(SkinnedMesh, OutputMask);
+	return new FOptimusSkinnedMeshWriteDataProviderProxy(SkinnedMesh, OutputMask, &LastLodIndexCachedByRenderProxy);
 }
 
 
-FOptimusSkinnedMeshWriteDataProviderProxy::FOptimusSkinnedMeshWriteDataProviderProxy(USkinnedMeshComponent* InSkinnedMeshComponent, uint64 InOutputMask)
+FOptimusSkinnedMeshWriteDataProviderProxy::FOptimusSkinnedMeshWriteDataProviderProxy(USkinnedMeshComponent* InSkinnedMeshComponent, uint64 InOutputMask, int32* InLastLodIndexPtr)
 {
 	SkeletalMeshObject = InSkinnedMeshComponent != nullptr ? InSkinnedMeshComponent->MeshObject : nullptr;
 	OutputMask = InOutputMask;
+	LastLodIndexPtr = InLastLodIndexPtr;
+	
+	check(LastLodIndexPtr != nullptr);
 }
 
 bool FOptimusSkinnedMeshWriteDataProviderProxy::IsValid(FValidationData const& InValidationData) const
@@ -190,9 +193,16 @@ void FOptimusSkinnedMeshWriteDataProviderProxy::AllocateResources(FRDGBuilder& G
 	// Allocate required buffers
 	const int32 LodIndex = SkeletalMeshObject->GetLOD();
 
+	bool bLodJustChanged = false;
+	if (LodIndex != *LastLodIndexPtr)
+	{
+		bLodJustChanged = true;	
+		*LastLodIndexPtr = LodIndex;
+	}
+
 	if (OutputMask & static_cast<uint64>(ESkinnedMeshWriteDataInterfaceOutputSelectorMask::Position))
 	{
-		PositionBuffer = FSkeletalMeshDeformerHelpers::AllocateVertexFactoryPositionBuffer(GraphBuilder, SkeletalMeshObject, LodIndex, TEXT("OptimusSkinnedMeshPosition"));
+		PositionBuffer = FSkeletalMeshDeformerHelpers::AllocateVertexFactoryPositionBuffer(GraphBuilder, SkeletalMeshObject, LodIndex, bLodJustChanged,TEXT("OptimusSkinnedMeshPosition"));
 		PositionBufferUAV = GraphBuilder.CreateUAV(PositionBuffer, PF_R32_FLOAT, ERDGUnorderedAccessViewFlags::SkipBarrier);
 	}
 	else
diff --git a/Engine/Plugins/Animation/DeformerGraph/Source/OptimusCore/Private/DataInterfaces/OptimusDataInterfaceSkinnedMeshWrite.h b/Engine/Plugins/Animation/DeformerGraph/Source/OptimusCore/Private/DataInterfaces/OptimusDataInterfaceSkinnedMeshWrite.h
--- a/Engine/Plugins/Animation/DeformerGraph/Source/OptimusCore/Private/DataInterfaces/OptimusDataInterfaceSkinnedMeshWrite.h
+++ b/Engine/Plugins/Animation/DeformerGraph/Source/OptimusCore/Private/DataInterfaces/OptimusDataInterfaceSkinnedMeshWrite.h
@@ -59,7 +59,10 @@ public:
 	UPROPERTY(EditAnywhere, BlueprintReadWrite, Category = Binding)
 	TObjectPtr<USkinnedMeshComponent> SkinnedMesh = nullptr;
 
-	uint64 OutputMask;
+	uint64 OutputMask = 0;
+
+	// Served as persistent storage for the provider proxy, should not be used by the data provider itself
+	int32 LastLodIndexCachedByRenderProxy = 0;
 
 	//~ Begin UComputeDataProvider Interface
 	FComputeDataProviderRenderProxy* GetRenderProxy() override;
@@ -69,7 +72,7 @@ public:
 class FOptimusSkinnedMeshWriteDataProviderProxy : public FComputeDataProviderRenderProxy
 {
 public:
-	FOptimusSkinnedMeshWriteDataProviderProxy(USkinnedMeshComponent* InSkinnedMeshComponent, uint64 InOutputMask);
+	FOptimusSkinnedMeshWriteDataProviderProxy(USkinnedMeshComponent* InSkinnedMeshComponent, uint64 InOutputMask, int32* InLastLodIndexPtr);
 
 	//~ Begin FComputeDataProviderRenderProxy Interface
 	bool IsValid(FValidationData const& InValidationData) const override;
@@ -80,8 +83,9 @@ public:
 private:
 	using FParameters = FSkinedMeshWriteDataInterfaceParameters;
 
-	FSkeletalMeshObject* SkeletalMeshObject;
-	uint64 OutputMask;
+	FSkeletalMeshObject* SkeletalMeshObject = nullptr;
+	uint64 OutputMask = 0;
+	int32* LastLodIndexPtr = nullptr; 
 
 	FRDGBuffer* PositionBuffer = nullptr;
 	FRDGBufferUAV* PositionBufferUAV = nullptr;
diff --git a/Engine/Source/Runtime/Engine/Private/SkeletalMeshDeformerHelpers.cpp b/Engine/Source/Runtime/Engine/Private/SkeletalMeshDeformerHelpers.cpp
--- a/Engine/Source/Runtime/Engine/Private/SkeletalMeshDeformerHelpers.cpp
+++ b/Engine/Source/Runtime/Engine/Private/SkeletalMeshDeformerHelpers.cpp
@@ -81,7 +81,7 @@ FSkeletalMeshDeformerHelpers::FClothBuffers FSkeletalMeshDeformerHelpers::GetClo
 	return Ret;
 }
 
-FRDGBuffer* FSkeletalMeshDeformerHelpers::AllocateVertexFactoryPositionBuffer(FRDGBuilder& GraphBuilder, FSkeletalMeshObject* InMeshObject, int32 InLodIndex, TCHAR const* InBufferName)
+FRDGBuffer* FSkeletalMeshDeformerHelpers::AllocateVertexFactoryPositionBuffer(FRDGBuilder& GraphBuilder, FSkeletalMeshObject* InMeshObject, int32 InLodIndex, bool bInLodJustChanged, TCHAR const* InBufferName)
 {
 	if (InMeshObject->IsCPUSkinned())
 	{
@@ -112,9 +112,20 @@ FRDGBuffer* FSkeletalMeshDeformerHelpers::AllocateVertexFactoryPositionBuffer(FR
 		PositionBuffer = GraphBuilder.CreateBuffer(FRDGBufferDesc::CreateBufferDesc(PosBufferBytesPerElement, NumVertices * 3), InBufferName, ERDGBufferFlags::None);
 		DeformerGeometry.Position = GraphBuilder.ConvertToExternalBuffer(PositionBuffer);
 		DeformerGeometry.PositionSRV = DeformerGeometry.Position->GetOrCreateSRV(GraphBuilder.RHICmdList, FRHIBufferSRVCreateInfo(PF_R32_FLOAT));
+
+		// Avoid using position buffer from the last time this LOD was active to compute motion vectors,
+		// the position delta between that previous position (could be from any time ago) and the current position can be any crazy value that is not meaningful
+		// instead lets just set the motion vector to zero here
+		if (bInLodJustChanged)
+		{
+			DeformerGeometry.PrevPosition = DeformerGeometry.Position;
+			DeformerGeometry.PrevPositionSRV = DeformerGeometry.PositionSRV;
+		}
+		
 		DeformerGeometry.PositionUpdatedFrame = Frame;
 		GraphBuilder.SetBufferAccessFinal(PositionBuffer, ERHIAccess::VertexOrIndexBuffer | ERHIAccess::SRVMask);
 
+		
 #if RHI_RAYTRACING
 		// Update ray tracing geometry whenever we recreate the position buffer.
 		FSkeletalMeshRenderData& SkelMeshRenderData = MeshObjectGPU->GetSkeletalMeshRenderData();
diff --git a/Engine/Source/Runtime/Engine/Public/SkeletalMeshDeformerHelpers.h b/Engine/Source/Runtime/Engine/Public/SkeletalMeshDeformerHelpers.h
--- a/Engine/Source/Runtime/Engine/Public/SkeletalMeshDeformerHelpers.h
+++ b/Engine/Source/Runtime/Engine/Public/SkeletalMeshDeformerHelpers.h
@@ -57,7 +57,7 @@ public:
 	 * Ownership is handled by the MeshObject.
 	 * If we call this more than once for the same MeshObject in the same frame then we return the allocation from the first call.
 	 */
-	ENGINE_API static FRDGBuffer* AllocateVertexFactoryPositionBuffer(FRDGBuilder& GraphBuilder, FSkeletalMeshObject* InMeshObject, int32 InLodIndex, TCHAR const* InBufferName);
+	ENGINE_API static FRDGBuffer* AllocateVertexFactoryPositionBuffer(FRDGBuilder& GraphBuilder, FSkeletalMeshObject* InMeshObject, int32 InLodIndex, bool bInLodJustChanged, TCHAR const* InBufferName);
 
 	/**
 	 * Allocate and bind a new tangent buffer and return it for writing.
