[Deformer Graph & CoimputeFramework] Enabled android support for deformer graph and compute framework

diff --git a/Engine/Config/Android/DataDrivenPlatformInfo.ini b/Engine/Config/Android/DataDrivenPlatformInfo.ini
--- a/Engine/Config/Android/DataDrivenPlatformInfo.ini
+++ b/Engine/Config/Android/DataDrivenPlatformInfo.ini
@@ -99,6 +99,7 @@ bSupportsDualSourceBlending = true
 bSupportsDxc=true
 bSupportsIndependentSamplers=true
 bIsHlslcc=false
+bSupportsComputeFramework = true
 
 FriendlyName=LOCTEXT("FriendlyShaderPlatformName_AndroidVulkan", "Android Vulkan Mobile")
 
@@ -136,3 +137,4 @@ MinimumWaveSize=4
 MaximumWaveSize=128
 bSupportsWave64=true
 FriendlyName=LOCTEXT("FriendlyShaderPlatformName_AndroidVulkanSM5", "Android Vulkan SM5 Mobile")
+bSupportsComputeFramework = true
diff --git a/Engine/Plugins/Animation/DeformerGraph/DeformerGraph.uplugin b/Engine/Plugins/Animation/DeformerGraph/DeformerGraph.uplugin
--- a/Engine/Plugins/Animation/DeformerGraph/DeformerGraph.uplugin
+++ b/Engine/Plugins/Animation/DeformerGraph/DeformerGraph.uplugin
@@ -20,7 +20,6 @@
 			"Type": "Runtime",
 			"LoadingPhase": "PostConfigInit",
 			"PlatformDenyList": [
-				"Android",
 				"IOS"
 			]
 		},
@@ -29,7 +28,6 @@
 			"Type": "Runtime",
 			"LoadingPhase": "PreDefault",
 			"PlatformDenyList": [
-				"Android",
 				"IOS"
 			]
 		},
@@ -38,7 +36,6 @@
 			"Type": "UncookedOnly",
 			"LoadingPhase": "PreDefault",
 			"PlatformDenyList": [
-				"Android",
 				"IOS"
 			]
 		},
