2-Support for  adding deformer instance ourtside of control rig

diff --git a/Engine/Plugins/Animation/ControlRig/Source/ControlRig/Public/ControlRig.h b/Engine/Plugins/Animation/ControlRig/Source/ControlRig/Public/ControlRig.h
--- a/Engine/Plugins/Animation/ControlRig/Source/ControlRig/Public/ControlRig.h
+++ b/Engine/Plugins/Animation/ControlRig/Source/ControlRig/Public/ControlRig.h
@@ -18,6 +18,8 @@
 #include "Animation/AnimInstance.h"
 #include "Animation/AnimInstanceProxy.h"
 #include "Animation/AttributesRuntime.h"
+#include "Animation/MeshDeformerInstance.h"
+#include "Animation/MeshDeformerProducer.h"
 #include "Rigs/RigModuleDefines.h"
 
 #if WITH_EDITOR
@@ -47,7 +49,7 @@ CONTROLRIG_API DECLARE_LOG_CATEGORY_EXTERN(LogControlRig, Log, All);
 
 /** Runs logic for mapping input data to transforms (the "Rig") */
 UCLASS(Blueprintable, Abstract, editinlinenew)
-class CONTROLRIG_API UControlRig : public URigVMHost, public INodeMappingProviderInterface, public IRigHierarchyProvider
+class CONTROLRIG_API UControlRig : public URigVMHost, public INodeMappingProviderInterface, public IRigHierarchyProvider, public IMeshDeformerProducer
 {
 	GENERATED_UCLASS_BODY()
 
@@ -205,11 +207,15 @@ public:
 	virtual void BeginDestroy() override;
 	// END UObject interface
 
-	DECLARE_EVENT_OneParam(UControlRig, FControlRigBeginDestroyEvent, class UControlRig*);
-	FControlRigBeginDestroyEvent& OnBeginDestroy() {return BeginDestroyEvent;};
+	using FControlRigBeginDestroyEvent = FMeshDeformerBeginDestroyEvent;
+
+	//~ Begin IMeshDeformerProducer interface
+	virtual FMeshDeformerBeginDestroyEvent& OnBeginDestroy() override {return BeginDestroyEvent;};
+	//~ End IMeshDeformerProducer interface
+	
 private:
 	/** Broadcasts a notification just before the controlrig is destroyed. */
-	FControlRigBeginDestroyEvent BeginDestroyEvent;
+	FMeshDeformerBeginDestroyEvent BeginDestroyEvent;
 	
 public:	
 	
diff --git a/Engine/Plugins/Animation/DeformerGraph/Source/OptimusCore/Private/ControlRig/RigUnit_Optimus.cpp b/Engine/Plugins/Animation/DeformerGraph/Source/OptimusCore/Private/ControlRig/RigUnit_Optimus.cpp
--- a/Engine/Plugins/Animation/DeformerGraph/Source/OptimusCore/Private/ControlRig/RigUnit_Optimus.cpp
+++ b/Engine/Plugins/Animation/DeformerGraph/Source/OptimusCore/Private/ControlRig/RigUnit_Optimus.cpp
@@ -297,7 +297,7 @@ FRigUnit_AddOptimusDeformer_Execute()
 			// This might change in the future depending on how per-LOD instance is implemented
 			if (UOptimusDeformerDynamicInstanceManager* DeformerInstanceManager = Cast<UOptimusDeformerDynamicInstanceManager>(ComponentToProcess->GetMeshDeformerInstanceForLOD(0)))
 			{
-				DeformerInstanceManager->EnqueueRigDeformer(DeformerInstanceGuid, SettingsTraitPtr->ExecutionPhase, SettingsTraitPtr->ExecutionGroup);
+				DeformerInstanceManager->EnqueueProducerDeformer(DeformerInstanceGuid, SettingsTraitPtr->ExecutionPhase, SettingsTraitPtr->ExecutionGroup);
 			}	
 		}
 		
@@ -345,9 +345,9 @@ FRigUnit_AddOptimusDeformer_Execute()
 										if (UOptimusDeformerDynamicInstanceManager* DeformerInstanceManager = Cast<UOptimusDeformerDynamicInstanceManager>(ComponentToProcess->GetMeshDeformerInstanceForLOD(0)))
 										{
 											UControlRig* Rig = WeakRig.Get();
-											if (!DeformerInstanceManager->GetRigDeformer(DeformerInstanceGuid))
+											if (!DeformerInstanceManager->GetDeformerInstance(DeformerInstanceGuid))
 											{
-												DeformerInstanceManager->AddRigDeformer(Rig, DeformerInstanceGuid, DeformerGraph);
+												DeformerInstanceManager->AddProducerDeformer(Rig, DeformerInstanceGuid, DeformerGraph);
 											}
 										}
 									}
@@ -367,7 +367,7 @@ FRigUnit_AddOptimusDeformer_Execute()
 					{
 						if (DeformerInstanceGuid.IsValid())
 						{
-							if (UOptimusDeformerInstance* DeformerInstance = DeformerInstanceManager->GetRigDeformer(DeformerInstanceGuid))
+							if (UOptimusDeformerInstance* DeformerInstance = DeformerInstanceManager->GetDeformerInstance(DeformerInstanceGuid))
 							{
 								VariableTrait->SetValue(DeformerInstance);
 							}
diff --git a/Engine/Plugins/Animation/DeformerGraph/Source/OptimusCore/Private/OptimusDeformerDynamicInstanceManager.cpp b/Engine/Plugins/Animation/DeformerGraph/Source/OptimusCore/Private/OptimusDeformerDynamicInstanceManager.cpp
--- a/Engine/Plugins/Animation/DeformerGraph/Source/OptimusCore/Private/OptimusDeformerDynamicInstanceManager.cpp
+++ b/Engine/Plugins/Animation/DeformerGraph/Source/OptimusCore/Private/OptimusDeformerDynamicInstanceManager.cpp
@@ -14,7 +14,7 @@ void UOptimusDeformerDynamicInstanceManager::AllocateResources()
 	
 	DefaultInstance->AllocateResources();
 
-	for (TPair<FGuid, TObjectPtr<UOptimusDeformerInstance>>& GuidToInstance : GuidToRigDeformerInstanceMap)
+	for (TPair<FGuid, TObjectPtr<UOptimusDeformerInstance>>& GuidToInstance : GuidToInstanceMap)
 	{
 		GuidToInstance.Value->AllocateResources();
 	}
@@ -26,7 +26,7 @@ void UOptimusDeformerDynamicInstanceManager::ReleaseResources()
 	
 	DefaultInstance->ReleaseResources();
 	
-	for (TPair<FGuid, TObjectPtr<UOptimusDeformerInstance>>& GuidToInstance : GuidToRigDeformerInstanceMap)
+	for (TPair<FGuid, TObjectPtr<UOptimusDeformerInstance>>& GuidToInstance : GuidToInstanceMap)
 	{
 		GuidToInstance.Value->ReleaseResources();
 	}
@@ -36,15 +36,16 @@ void UOptimusDeformerDynamicInstanceManager::EnqueueWork(FEnqueueWorkDesc const&
 {
 	// Runs during UWorld::SendAllEndOfFrameUpdates
 
-	for (FGuid Guid : RigDeformerInstancePendingInit)
+	for (FGuid Guid : GuidsPendingInit)
 	{
-		if (TObjectPtr<UOptimusDeformerInstance>* InstancePtr = GuidToRigDeformerInstanceMap.Find(Guid))
+		if (TObjectPtr<UOptimusDeformerInstance>* InstancePtr = GuidToInstanceMap.Find(Guid))
 		{
+			
 			(*InstancePtr)->AllocateResources();
 		}
 	}
 
-	RigDeformerInstancePendingInit.Reset();
+	GuidsPendingInit.Reset();
 	
 	// Enqueue work
 
@@ -72,7 +73,7 @@ void UOptimusDeformerDynamicInstanceManager::EnqueueWork(FEnqueueWorkDesc const&
 
 	for (EOptimusDeformerExecutionPhase Phase : Phases)
 	{
-		if (TMap<int32, TArray<FGuid>>* ExecutionGroupQueueMapPtr = InstanceQueueMap.Find(Phase))
+		if (TMap<int32, TArray<FGuid>>* ExecutionGroupQueueMapPtr = ExecutionQueueMap.Find(Phase))
 		{
 			TArray<int32> SortedExecutionGroups;
 			ExecutionGroupQueueMapPtr->GenerateKeyArray(SortedExecutionGroups);
@@ -84,7 +85,7 @@ void UOptimusDeformerDynamicInstanceManager::EnqueueWork(FEnqueueWorkDesc const&
 				int32 LastOverrideExecutionGroup = SortedExecutionGroups.Last();
 				FGuid LastOverrideInstanceGuid = (*ExecutionGroupQueueMapPtr)[LastOverrideExecutionGroup].Last();
 
-				if (TObjectPtr<UOptimusDeformerInstance>* InstancePtr = GuidToRigDeformerInstanceMap.Find(LastOverrideInstanceGuid))
+				if (TObjectPtr<UOptimusDeformerInstance>* InstancePtr = GuidToInstanceMap.Find(LastOverrideInstanceGuid))
 				{
 					EnqueueInstance((*InstancePtr));
 				}
@@ -95,7 +96,7 @@ void UOptimusDeformerDynamicInstanceManager::EnqueueWork(FEnqueueWorkDesc const&
 				{
 					for (FGuid Guid : (*ExecutionGroupQueueMapPtr)[ExecutionGroup])
 					{
-						if (TObjectPtr<UOptimusDeformerInstance>* InstancePtr = GuidToRigDeformerInstanceMap.Find(Guid))
+						if (TObjectPtr<UOptimusDeformerInstance>* InstancePtr = GuidToInstanceMap.Find(Guid))
 						{
 							EnqueueInstance((*InstancePtr));
 						}
@@ -110,7 +111,7 @@ void UOptimusDeformerDynamicInstanceManager::EnqueueWork(FEnqueueWorkDesc const&
 		}
 	}
 
-	InstanceQueueMap.Reset();
+	ExecutionQueueMap.Reset();
 }
 
 EMeshDeformerOutputBuffer UOptimusDeformerDynamicInstanceManager::GetOutputBuffers() const
@@ -126,77 +127,81 @@ UMeshDeformerInstance* UOptimusDeformerDynamicInstanceManager::GetInstanceForSou
 	return DefaultInstance;
 }
 
-
-void UOptimusDeformerDynamicInstanceManager::OnControlRigBeginDestroy(UControlRig* InControlRig)
+void UOptimusDeformerDynamicInstanceManager::OnObjectBeginDestroy(IMeshDeformerProducer* InObject)
 {
-	if (TArray<FGuid>* Instances = RigToInstanceGuidsMap.Find(InControlRig))
+	if (TArray<FGuid>* Instances = ProducerToGuidsMap.Find(Cast<UObject>(InObject)))
 	{
 		TArray<FGuid> InstancesToRemove = *Instances;
 		for (FGuid Guid : InstancesToRemove)
 		{
-			if (TObjectPtr<UOptimusDeformerInstance>* InstancePtr = GuidToRigDeformerInstanceMap.Find(Guid))
+			if (TObjectPtr<UOptimusDeformerInstance>* InstancePtr = GuidToInstanceMap.Find(Guid))
 			{
 				(*InstancePtr)->ReleaseResources();
 			}
 
-			GuidToRigDeformerInstanceMap.Remove(Guid);
+			GuidToInstanceMap.Remove(Guid);
 		}
 	}
 	
-	RigToInstanceGuidsMap.Remove(InControlRig);
+	ProducerToGuidsMap.Remove(Cast<UObject>(InObject));
+	
+	InObject->OnBeginDestroy().RemoveAll(this);
 	
-	InControlRig->OnBeginDestroy().RemoveAll(this);
 }
 
 void UOptimusDeformerDynamicInstanceManager::BeginDestroy()
 {
-	TArray<TWeakObjectPtr<UControlRig>> Rigs;
-	RigToInstanceGuidsMap.GenerateKeyArray(Rigs);
+	TArray<TWeakObjectPtr<UObject>> Objects;
+	ProducerToGuidsMap.GenerateKeyArray(Objects);
 
 	// Release resources should have been called, so just unregister callbacks for good measure
 	
-	for (TWeakObjectPtr<UControlRig>& Rig : Rigs)
+	for (TWeakObjectPtr<UObject>& Object : Objects)
 	{
-		if (Rig.IsValid())
+		if (Object.IsValid())
 		{
-			Rig->OnBeginDestroy().RemoveAll(this);
+			if(IMeshDeformerProducer* ManagedObject = Cast<IMeshDeformerProducer>(Object))
+			{
+				ManagedObject->OnBeginDestroy().RemoveAll(this);
+			}
 		}
 	}
 	
 	Super::BeginDestroy();
 }
 
-
-void UOptimusDeformerDynamicInstanceManager::AddRigDeformer(UControlRig* InControlRig, FGuid InInstanceGuid, UOptimusDeformer* InDeformer)
+void UOptimusDeformerDynamicInstanceManager::AddProducerDeformer(IMeshDeformerProducer* InProducer, FGuid InInstanceGuid, UOptimusDeformer* InDeformer)
 {
 	check(IsInGameThread());
 	
-	if (ensure(!GuidToRigDeformerInstanceMap.Contains(InInstanceGuid)))
+	if (ensure(!GuidToInstanceMap.Contains(InInstanceGuid)))
 	{
-		UOptimusDeformerInstance* DeformerInstance = InDeformer->CreateOptimusInstance(CastChecked<USkeletalMeshComponent>(GetOuter()), nullptr);
-		GuidToRigDeformerInstanceMap.Add(InInstanceGuid, DeformerInstance);
-		RigDeformerInstancePendingInit.Add(InInstanceGuid);
+		UOptimusDeformerInstance* DeformerInstance = InDeformer->CreateOptimusInstance(CastChecked<UMeshComponent>(GetOuter()), nullptr);
+		GuidToInstanceMap.Add(InInstanceGuid, DeformerInstance);
+		GuidsPendingInit.Add(InInstanceGuid);
 
-		if (TArray<FGuid>* GuidArray = RigToInstanceGuidsMap.Find(InControlRig))
+		if (TArray<FGuid>* GuidArray = ProducerToGuidsMap.Find(Cast<UObject>(InProducer)))
 		{
 			GuidArray->Add(InInstanceGuid);
 		}
 		else
 		{
-			RigToInstanceGuidsMap.Add(InControlRig, {InInstanceGuid});
+			ProducerToGuidsMap.Add(Cast<UObject>(InProducer), {InInstanceGuid});
 			
 			// First time for this control rig, register some callbacks as well
-			check(!InControlRig->OnBeginDestroy().IsBoundToObject(this));
+			check(!InProducer->OnBeginDestroy().IsBoundToObject(this));
+			
 			// Assuming owning component of the rig cannot change
-			InControlRig->OnBeginDestroy().AddUObject(this, &UOptimusDeformerDynamicInstanceManager::OnControlRigBeginDestroy);
+			UOptimusDeformerDynamicInstanceManager* DeformerInstanceManager = this;
+			InProducer->OnBeginDestroy().AddUObject(this, &UOptimusDeformerDynamicInstanceManager::OnObjectBeginDestroy);
 
 		}
 	}
 }
 
-UOptimusDeformerInstance* UOptimusDeformerDynamicInstanceManager::GetRigDeformer(FGuid InInstanceGuid)
+UOptimusDeformerInstance* UOptimusDeformerDynamicInstanceManager::GetDeformerInstance(FGuid InInstanceGuid)
 {
-	TObjectPtr<UOptimusDeformerInstance>* InstancePtr = GuidToRigDeformerInstanceMap.Find(InInstanceGuid);
+	TObjectPtr<UOptimusDeformerInstance>* InstancePtr = GuidToInstanceMap.Find(InInstanceGuid);
 	
 	if (InstancePtr)
 	{
@@ -206,13 +211,12 @@ UOptimusDeformerInstance* UOptimusDeformerDynamicInstanceManager::GetRigDeformer
 	return nullptr;
 }
 
-
-void UOptimusDeformerDynamicInstanceManager::EnqueueRigDeformer(FGuid InInstanceGuid, EOptimusDeformerExecutionPhase InExecutionPhase, int32 InExecutionGroup)
+void UOptimusDeformerDynamicInstanceManager::EnqueueProducerDeformer(FGuid InInstanceGuid, EOptimusDeformerExecutionPhase InExecutionPhase, int32 InExecutionGroup)
 {
 	// Typically called from anim thread, but there shouldn't be concurrent access to this queue. All rigs running on the current mesh should run sequentially.
 	
-	
-	TArray<FGuid>& InstanceQueueRef = InstanceQueueMap.FindOrAdd(InExecutionPhase).FindOrAdd(InExecutionGroup);
+	FScopeLock ScopeLock(&EnqueueCriticalSection);
+	TArray<FGuid>& InstanceQueueRef = ExecutionQueueMap.FindOrAdd(InExecutionPhase).FindOrAdd(InExecutionGroup);
 	
 	// If we ever get duplicates, it means extra unnecessary instances were added via extra control rig evaluations triggered by user actions like moving a control.
 	// So let's invalidate those instances.
diff --git a/Engine/Plugins/Animation/DeformerGraph/Source/OptimusCore/Public/OptimusDeformerDynamicInstanceManager.h b/Engine/Plugins/Animation/DeformerGraph/Source/OptimusCore/Public/OptimusDeformerDynamicInstanceManager.h
--- a/Engine/Plugins/Animation/DeformerGraph/Source/OptimusCore/Public/OptimusDeformerDynamicInstanceManager.h
+++ b/Engine/Plugins/Animation/DeformerGraph/Source/OptimusCore/Public/OptimusDeformerDynamicInstanceManager.h
@@ -3,9 +3,9 @@
 #pragma once
 
 #include "Animation/MeshDeformerInstance.h"
+#include "Animation/MeshDeformerProducer.h"
 #include "OptimusDeformerDynamicInstanceManager.generated.h"
 
-
 class UControlRig;
 class UOptimusDeformerInstance;
 class UOptimusDeformer;
@@ -42,28 +42,45 @@ public:
 	/** InstanceManager is an intermediate instance, call this function to get the instance for the deformer that created this instance manager */
 	UMeshDeformerInstance* GetInstanceForSourceDeformer() override;
 
-	/** Remove associated deformer instances when the rig is removed */
-	void OnControlRigBeginDestroy(UControlRig* InControlRig);
-
+	/** Begin destroying the managaer */
 	void BeginDestroy() override;
 
-	void AddRigDeformer(UControlRig* InControlRig, FGuid InInstanceGuid, UOptimusDeformer* InDeformer);
-	UOptimusDeformerInstance* GetRigDeformer(FGuid InInstanceGuid);
-	void EnqueueRigDeformer(FGuid InInstanceGuid, EOptimusDeformerExecutionPhase InExecutionPhase, int32 InExecutionGroup);
+	/** Remove associated deformer instances when the object is removed */
+	void OnObjectBeginDestroy(IMeshDeformerProducer* InObject);
+
+	/** Add a producer defomer to the manager (from game thread) */
+	void AddProducerDeformer(IMeshDeformerProducer* InProducer, FGuid InInstanceGuid, UOptimusDeformer* InDeformer);
+
+	/** Get the object deformer given a guid */
+	UOptimusDeformerInstance* GetDeformerInstance(FGuid InInstanceGuid);
 	
+	/** Enqueue an object deformer instance to be executed later (from anim/physics... threads )*/
+	void EnqueueProducerDeformer(FGuid InInstanceGuid, EOptimusDeformerExecutionPhase InExecutionPhase, int32 InExecutionGroup);
+
 	UPROPERTY()
 	TObjectPtr<UOptimusDeformerInstance> DefaultInstance;
 
+	UE_DEPRECATED(5.6, "This property is going to be deleted and replaced by the private member GuidToInstanceMap")
 	UPROPERTY()
 	TMap<FGuid, TObjectPtr<UOptimusDeformerInstance>> GuidToRigDeformerInstanceMap;
 
-	TMap<TWeakObjectPtr<UControlRig>, TArray<FGuid>> RigToInstanceGuidsMap;
+private:
+
+	/** Guid to deformer instance map*/
+	UPROPERTY()
+	TMap<FGuid, TObjectPtr<UOptimusDeformerInstance>> GuidToInstanceMap;
+	
+	/** Object ptr to deformer guids map */
+	TMap<TWeakObjectPtr<UObject>, TArray<FGuid>> ProducerToGuidsMap;
+
+	/** Freshly created deformer instances should be initialized before dispatch */
+	TArray<FGuid> GuidsPendingInit;
 
-	// Freshly created deformer instances should be initialized before dispatch 
-	TArray<FGuid> RigDeformerInstancePendingInit;
+	/** Instances per execution group per execution phase */
+	TMap<EOptimusDeformerExecutionPhase, TMap<int32, TArray<FGuid>>> ExecutionQueueMap;
 
-	// Instances per execution group per execution phase
-	TMap<EOptimusDeformerExecutionPhase, TMap<int32, TArray<FGuid>>> InstanceQueueMap;
+	/** Enqueue critical section to be thread safe */
+	FCriticalSection EnqueueCriticalSection;
 };
 
 
diff --git a/Engine/Plugins/Animation/DeformerGraph/Source/OptimusCore/Public/OptimusDeformerInstance.h b/Engine/Plugins/Animation/DeformerGraph/Source/OptimusCore/Public/OptimusDeformerInstance.h
--- a/Engine/Plugins/Animation/DeformerGraph/Source/OptimusCore/Public/OptimusDeformerInstance.h
+++ b/Engine/Plugins/Animation/DeformerGraph/Source/OptimusCore/Public/OptimusDeformerInstance.h
@@ -167,7 +167,7 @@ protected:
  * It contains the per instance deformer variable state and local state for each of the graphs in the deformer.
  */
 UCLASS(BlueprintType)
-class UOptimusDeformerInstance : public UMeshDeformerInstance
+class OPTIMUSCORE_API UOptimusDeformerInstance : public UMeshDeformerInstance
 {
 	GENERATED_BODY()
 
diff --git a/Engine/Source/Runtime/Engine/Classes/Animation/MeshDeformerProducer.h b/Engine/Source/Runtime/Engine/Classes/Animation/MeshDeformerProducer.h
new file mode 100644
--- /dev/null
+++ b/Engine/Source/Runtime/Engine/Classes/Animation/MeshDeformerProducer.h
@@ -0,0 +1,26 @@
+﻿// Copyright Epic Games, Inc. All Rights Reserved.
+
+#pragma once
+
+#include "CoreMinimal.h"
+#include "UObject/Interface.h"
+#include "MeshDeformerProducer.generated.h"
+
+class UMeshDeformerInstance;
+
+/** Interface for all the UObject able to add directly mesh deformer instances to the manager */
+UINTERFACE(MinimalAPI)
+class UMeshDeformerProducer: public UInterface
+{
+	GENERATED_BODY()
+};
+
+class IMeshDeformerProducer
+{
+	GENERATED_BODY()
+
+public:
+	DECLARE_EVENT_OneParam(UMeshDeformerInstance, FMeshDeformerBeginDestroyEvent, class IMeshDeformerProducer*);
+	virtual FMeshDeformerBeginDestroyEvent& OnBeginDestroy() = 0;
+};
+
diff --git a/Engine/Source/Runtime/Experimental/Dataflow/Engine/Private/Dataflow/DataflowObject.cpp b/Engine/Source/Runtime/Experimental/Dataflow/Engine/Private/Dataflow/DataflowObject.cpp
--- a/Engine/Source/Runtime/Experimental/Dataflow/Engine/Private/Dataflow/DataflowObject.cpp
+++ b/Engine/Source/Runtime/Experimental/Dataflow/Engine/Private/Dataflow/DataflowObject.cpp
@@ -47,6 +47,14 @@ UDataflow::UDataflow(const FObjectInitializer& ObjectInitializer)
 	, Dataflow(new UE::Dataflow::FGraph())
 {}
 
+void UDataflow::BeginDestroy()
+{
+	BeginDestroyEvent.Broadcast(this);
+	BeginDestroyEvent.Clear();
+	
+	Super::BeginDestroy();
+}
+
 void UDataflow::EvaluateTerminalNodeByName(FName NodeName, UObject* Asset)
 {
 	ensureAlwaysMsgf(false, TEXT("Deprecated use the dataflow blueprint library from now on"));
diff --git a/Engine/Source/Runtime/Experimental/Dataflow/Engine/Public/Dataflow/DataflowObject.h b/Engine/Source/Runtime/Experimental/Dataflow/Engine/Public/Dataflow/DataflowObject.h
--- a/Engine/Source/Runtime/Experimental/Dataflow/Engine/Public/Dataflow/DataflowObject.h
+++ b/Engine/Source/Runtime/Experimental/Dataflow/Engine/Public/Dataflow/DataflowObject.h
@@ -4,6 +4,8 @@
 
 #include "Containers/Array.h"
 #include "CoreMinimal.h"
+#include "Animation/MeshDeformerInstance.h"
+#include "Animation/MeshDeformerProducer.h"
 #include "Dataflow/DataflowPreview.h"
 #include "Dataflow/DataflowCore.h"
 #include "EdGraph/EdGraph.h"
@@ -67,7 +69,7 @@ enum class EDataflowType : uint8
 *
 */
 UCLASS(BlueprintType, customconstructor, MinimalAPI)
-class UDataflow : public UEdGraph
+class UDataflow : public UEdGraph, public IMeshDeformerProducer
 {
 	GENERATED_UCLASS_BODY()
 
@@ -91,6 +93,14 @@ public:
 	/** Simulation tag to use in the node registry */
 	static const inline FString SimulationTag = TEXT("DataflowSimulationTag");
 
+	//~ Begin UObject interface
+	virtual void BeginDestroy() override;
+	//~ End UObject interface
+
+	//~ Begin IMeshDeformerProducer interface
+	DATAFLOWENGINE_API virtual FMeshDeformerBeginDestroyEvent& OnBeginDestroy() override {return BeginDestroyEvent;};
+	//~ End IMeshDeformerProducer interface
+
 public:
 	UPROPERTY(EditAnywhere, Category = "Evaluation")
 	bool bActive = true;
@@ -180,6 +190,10 @@ private:
 	/** Used to disable per-node serialization when serializing a transaction */
 	bool bEnablePerNodeTransactionSerialization = true;
 #endif
+
+private:
+	/** Broadcasts a notification just before the dataflow is destroyed. */
+	FMeshDeformerBeginDestroyEvent BeginDestroyEvent;
 };
 
 
