[Deformer Graph] remove spaces instead of replacing space with underscore when sanitize pin names

diff --git a/Engine/Plugins/Animation/DeformerGraph/Source/OptimusCore/Private/OptimusHelpers.cpp b/Engine/Plugins/Animation/DeformerGraph/Source/OptimusCore/Private/OptimusHelpers.cpp
--- a/Engine/Plugins/Animation/DeformerGraph/Source/OptimusCore/Private/OptimusHelpers.cpp
+++ b/Engine/Plugins/Animation/DeformerGraph/Source/OptimusCore/Private/OptimusHelpers.cpp
@@ -47,8 +47,12 @@ FName Optimus::FUniqueNameGenerator::GetUniqueName(FName InName)
 
 FName Optimus::GetSanitizedNameForHlsl(FName InName)
 {
-	// Sanitize the name
 	FString Name = InName.ToString();
+
+	// Remove spaces
+	Name.ReplaceInline(TEXT(" "), TEXT(""));
+
+	// Replace other bad characters with "_" 
 	for (int32 i = 0; i < Name.Len(); ++i)
 	{
 		TCHAR& C = Name[i];
diff --git a/Engine/Plugins/Animation/DeformerGraph/Source/OptimusCore/Public/OptimusNodePin.h b/Engine/Plugins/Animation/DeformerGraph/Source/OptimusCore/Public/OptimusNodePin.h
--- a/Engine/Plugins/Animation/DeformerGraph/Source/OptimusCore/Public/OptimusNodePin.h
+++ b/Engine/Plugins/Animation/DeformerGraph/Source/OptimusCore/Public/OptimusNodePin.h
@@ -220,6 +220,5 @@ PRAGMA_ENABLE_DEPRECATION_WARNINGS
 	UPROPERTY()
 	TArray<TObjectPtr<UOptimusNodePin>> SubPins;
 
-	UPROPERTY(Transient)
 	FName DisplayName;
 };
