[Deformer Graph] Fixed array constant nodes of struct type (int2, int3, fquat ... etc) not having pins

diff --git a/Engine/Plugins/Animation/DeformerGraph/Source/OptimusCore/Private/OptimusDataTypeRegistry.cpp b/Engine/Plugins/Animation/DeformerGraph/Source/OptimusCore/Private/OptimusDataTypeRegistry.cpp
--- a/Engine/Plugins/Animation/DeformerGraph/Source/OptimusCore/Private/OptimusDataTypeRegistry.cpp
+++ b/Engine/Plugins/Animation/DeformerGraph/Source/OptimusCore/Private/OptimusDataTypeRegistry.cpp
@@ -1455,7 +1455,7 @@ FOptimusDataTypeHandle FOptimusDataTypeRegistry::FindType(const FProperty& InPro
 	else if (const FArrayProperty* ArrayProperty = CastField<const FArrayProperty>(&InProperty))
 	{
 		const FProperty* PropertyForType = ArrayProperty->Inner;
-		return FindArrayType(*PropertyForType->GetClass());
+		return FindArrayType(*PropertyForType);
 	}
 	else
 	{
