Fix private UOptimusNode_DataInterface reference in public API function (new Interface to expose functionality)

diff --git a/Engine/Plugins/Animation/DeformerGraph/Source/OptimusCore/Private/DataInterfaces/OptimusDataInterfaceAdvancedSkeleton.cpp b/Engine/Plugins/Animation/DeformerGraph/Source/OptimusCore/Private/DataInterfaces/OptimusDataInterfaceAdvancedSkeleton.cpp
--- a/Engine/Plugins/Animation/DeformerGraph/Source/OptimusCore/Private/DataInterfaces/OptimusDataInterfaceAdvancedSkeleton.cpp
+++ b/Engine/Plugins/Animation/DeformerGraph/Source/OptimusCore/Private/DataInterfaces/OptimusDataInterfaceAdvancedSkeleton.cpp
@@ -298,11 +298,14 @@ void UOptimusAdvancedSkeletonDataInterface::Initialize()
 }
 
 
-void UOptimusAdvancedSkeletonDataInterface::RegisterPropertyChangeDelegatesForOwningNode(UOptimusNode_DataInterface* InNode)
+void UOptimusAdvancedSkeletonDataInterface::RegisterPropertyChangeDelegatesForOwningNode(UOptimusNode* InNode)
 {
-	OnPinDefinitionChangedDelegate.BindUObject(InNode, &UOptimusNode_DataInterface::RecreatePinsFromPinDefinitions);
-	OnPinDefinitionRenamedDelegate.BindUObject(InNode, &UOptimusNode_DataInterface::RenamePinFromPinDefinition);
-	OnDisplayNameChangedDelegate.BindUObject(InNode, &UOptimusNode_DataInterface::UpdateDisplayNameFromDataInterface);
+	if(InNode)
+	{
+		OnPinDefinitionChangedDelegate.BindUObject(InNode, &UOptimusNode::RecreatePinsFromPinDefinitions);
+		OnPinDefinitionRenamedDelegate.BindUObject(InNode, &UOptimusNode::RenamePinFromPinDefinition);
+		OnDisplayNameChangedDelegate.BindUObject(InNode, &UOptimusNode::UpdateDisplayNameFromDataInterface);
+	}
 }
 
 void UOptimusAdvancedSkeletonDataInterface::OnDataTypeChanged(FName InTypeName)
diff --git a/Engine/Plugins/Animation/DeformerGraph/Source/OptimusCore/Private/DataInterfaces/OptimusDataInterfaceAdvancedSkeleton.h b/Engine/Plugins/Animation/DeformerGraph/Source/OptimusCore/Private/DataInterfaces/OptimusDataInterfaceAdvancedSkeleton.h
--- a/Engine/Plugins/Animation/DeformerGraph/Source/OptimusCore/Private/DataInterfaces/OptimusDataInterfaceAdvancedSkeleton.h
+++ b/Engine/Plugins/Animation/DeformerGraph/Source/OptimusCore/Private/DataInterfaces/OptimusDataInterfaceAdvancedSkeleton.h
@@ -122,7 +122,7 @@ public:
 
 	void Initialize() override;
 	bool CanPinDefinitionChange() override {return true;};
-	void RegisterPropertyChangeDelegatesForOwningNode(UOptimusNode_DataInterface* InNode) override;
+	void RegisterPropertyChangeDelegatesForOwningNode(UOptimusNode* InNode) override;
 
 	void OnDataTypeChanged(FName InTypeName) override;
 	//~ End UOptimusComputeDataInterface Interface
diff --git a/Engine/Plugins/Animation/DeformerGraph/Source/OptimusCore/Private/DataInterfaces/OptimusDataInterfaceAnimAttribute.cpp b/Engine/Plugins/Animation/DeformerGraph/Source/OptimusCore/Private/DataInterfaces/OptimusDataInterfaceAnimAttribute.cpp
--- a/Engine/Plugins/Animation/DeformerGraph/Source/OptimusCore/Private/DataInterfaces/OptimusDataInterfaceAnimAttribute.cpp
+++ b/Engine/Plugins/Animation/DeformerGraph/Source/OptimusCore/Private/DataInterfaces/OptimusDataInterfaceAnimAttribute.cpp
@@ -233,10 +233,13 @@ void UOptimusAnimAttributeDataInterface::Initialize()
 			FOptimusDataTypeRegistry::Get().FindType(*FFloatProperty::StaticClass()) );
 }
 
-void UOptimusAnimAttributeDataInterface::RegisterPropertyChangeDelegatesForOwningNode(UOptimusNode_DataInterface* InNode)
+void UOptimusAnimAttributeDataInterface::RegisterPropertyChangeDelegatesForOwningNode(UOptimusNode* InNode)
 {
-	OnPinDefinitionChangedDelegate.BindUObject(InNode, &UOptimusNode_DataInterface::RecreatePinsFromPinDefinitions);
-	OnPinDefinitionRenamedDelegate.BindUObject(InNode, &UOptimusNode_DataInterface::RenamePinFromPinDefinition);
+	if(InNode)
+	{
+		OnPinDefinitionChangedDelegate.BindUObject(InNode, &UOptimusNode::RecreatePinsFromPinDefinitions);
+		OnPinDefinitionRenamedDelegate.BindUObject(InNode, &UOptimusNode::RenamePinFromPinDefinition);
+	}
 }
 
 
diff --git a/Engine/Plugins/Animation/DeformerGraph/Source/OptimusCore/Private/DataInterfaces/OptimusDataInterfaceAnimAttribute.h b/Engine/Plugins/Animation/DeformerGraph/Source/OptimusCore/Private/DataInterfaces/OptimusDataInterfaceAnimAttribute.h
--- a/Engine/Plugins/Animation/DeformerGraph/Source/OptimusCore/Private/DataInterfaces/OptimusDataInterfaceAnimAttribute.h
+++ b/Engine/Plugins/Animation/DeformerGraph/Source/OptimusCore/Private/DataInterfaces/OptimusDataInterfaceAnimAttribute.h
@@ -128,7 +128,7 @@ public:
 	
 	void Initialize() override;
 	bool CanPinDefinitionChange() override {return true;};
-	void RegisterPropertyChangeDelegatesForOwningNode(UOptimusNode_DataInterface* InNode) override;
+	void RegisterPropertyChangeDelegatesForOwningNode(UOptimusNode* InNode) override;
 	
 	//~ End UOptimusComputeDataInterface Interface
 	
diff --git a/Engine/Plugins/Animation/DeformerGraph/Source/OptimusCore/Private/DataInterfaces/OptimusDataInterfaceSkinWeightsAsVertexMask.cpp b/Engine/Plugins/Animation/DeformerGraph/Source/OptimusCore/Private/DataInterfaces/OptimusDataInterfaceSkinWeightsAsVertexMask.cpp
--- a/Engine/Plugins/Animation/DeformerGraph/Source/OptimusCore/Private/DataInterfaces/OptimusDataInterfaceSkinWeightsAsVertexMask.cpp
+++ b/Engine/Plugins/Animation/DeformerGraph/Source/OptimusCore/Private/DataInterfaces/OptimusDataInterfaceSkinWeightsAsVertexMask.cpp
@@ -77,11 +77,14 @@ TSubclassOf<UActorComponent> UOptimusSkinWeightsAsVertexMaskDataInterface::GetRe
 	return USkeletalMeshComponent::StaticClass();
 }
 
-void UOptimusSkinWeightsAsVertexMaskDataInterface::RegisterPropertyChangeDelegatesForOwningNode(UOptimusNode_DataInterface* InNode)
+void UOptimusSkinWeightsAsVertexMaskDataInterface::RegisterPropertyChangeDelegatesForOwningNode(UOptimusNode* InNode)
 {
-	OnPinDefinitionChangedDelegate.BindUObject(InNode, &UOptimusNode_DataInterface::RecreatePinsFromPinDefinitions);
-	OnPinDefinitionRenamedDelegate.BindUObject(InNode, &UOptimusNode_DataInterface::RenamePinFromPinDefinition);
-	OnDisplayNameChangedDelegate.BindUObject(InNode, &UOptimusNode_DataInterface::UpdateDisplayNameFromDataInterface);
+	if(InNode)
+	{
+		OnPinDefinitionChangedDelegate.BindUObject(InNode, &UOptimusNode::RecreatePinsFromPinDefinitions);
+		OnPinDefinitionRenamedDelegate.BindUObject(InNode, &UOptimusNode::RenamePinFromPinDefinition);
+		OnDisplayNameChangedDelegate.BindUObject(InNode, &UOptimusNode::UpdateDisplayNameFromDataInterface);
+	}
 }
 
 void UOptimusSkinWeightsAsVertexMaskDataInterface::GetSupportedInputs(TArray<FShaderFunctionDefinition>& OutFunctions) const
diff --git a/Engine/Plugins/Animation/DeformerGraph/Source/OptimusCore/Private/DataInterfaces/OptimusDataInterfaceSkinWeightsAsVertexMask.h b/Engine/Plugins/Animation/DeformerGraph/Source/OptimusCore/Private/DataInterfaces/OptimusDataInterfaceSkinWeightsAsVertexMask.h
--- a/Engine/Plugins/Animation/DeformerGraph/Source/OptimusCore/Private/DataInterfaces/OptimusDataInterfaceSkinWeightsAsVertexMask.h
+++ b/Engine/Plugins/Animation/DeformerGraph/Source/OptimusCore/Private/DataInterfaces/OptimusDataInterfaceSkinWeightsAsVertexMask.h
@@ -30,7 +30,7 @@ public:
 	TArray<FOptimusCDIPropertyPinDefinition> GetPropertyPinDefinitions() const override;
 	TSubclassOf<UActorComponent> GetRequiredComponentClass() const override;
 
-	void RegisterPropertyChangeDelegatesForOwningNode(UOptimusNode_DataInterface* InNode) override;
+	void RegisterPropertyChangeDelegatesForOwningNode(UOptimusNode* InNode) override;
 	//~ End UOptimusComputeDataInterface Interface
 	
 
diff --git a/Engine/Plugins/Animation/DeformerGraph/Source/OptimusCore/Private/Nodes/OptimusNode_DataInterface.h b/Engine/Plugins/Animation/DeformerGraph/Source/OptimusCore/Private/Nodes/OptimusNode_DataInterface.h
--- a/Engine/Plugins/Animation/DeformerGraph/Source/OptimusCore/Private/Nodes/OptimusNode_DataInterface.h
+++ b/Engine/Plugins/Animation/DeformerGraph/Source/OptimusCore/Private/Nodes/OptimusNode_DataInterface.h
@@ -32,18 +32,17 @@ public:
 
 	virtual void SetDataInterfaceClass(TSubclassOf<UOptimusComputeDataInterface> InDataInterfaceClass);
 	bool IsComponentSourceCompatible(const UOptimusComponentSource* InComponentSource) const;
-
-	void RecreatePinsFromPinDefinitions();
-	void RenamePinFromPinDefinition(FName InOld, FName InNew);
-	void UpdateDisplayNameFromDataInterface();
-
+	
 	// -- UOptimusNode overrides
 	FName GetNodeCategory() const override 
 	{
 		return CategoryName::DataInterfaces;
 	}
 	void InitializeTransientData() override;
-
+	void RecreatePinsFromPinDefinitions() override;
+	void RenamePinFromPinDefinition(FName InOld, FName InNew) override;
+	void UpdateDisplayNameFromDataInterface() override;
+	
 	// -- UObject overrides
 	void Serialize(FArchive& Ar) override;
 
diff --git a/Engine/Plugins/Animation/DeformerGraph/Source/OptimusCore/Public/OptimusComputeDataInterface.h b/Engine/Plugins/Animation/DeformerGraph/Source/OptimusCore/Public/OptimusComputeDataInterface.h
--- a/Engine/Plugins/Animation/DeformerGraph/Source/OptimusCore/Public/OptimusComputeDataInterface.h
+++ b/Engine/Plugins/Animation/DeformerGraph/Source/OptimusCore/Public/OptimusComputeDataInterface.h
@@ -11,8 +11,7 @@
 
 
 class UActorComponent;
-class UOptimusNode_DataInterface;
-
+class UOptimusNode;
 
 struct FOptimusCDIPinDefinition
 {
@@ -180,7 +179,7 @@ public:
 	virtual bool CanPinDefinitionChange() { return false; };
 
 	/// Register delegates for data interface node to update when the data interface changes
-	virtual void RegisterPropertyChangeDelegatesForOwningNode(UOptimusNode_DataInterface* InNode) {};
+	virtual void RegisterPropertyChangeDelegatesForOwningNode(UOptimusNode* InNode) {};
 	
 	/**
 	 * @return Returns the component type that this data interface operates on.
diff --git a/Engine/Plugins/Animation/DeformerGraph/Source/OptimusCore/Public/OptimusNode.h b/Engine/Plugins/Animation/DeformerGraph/Source/OptimusCore/Public/OptimusNode.h
--- a/Engine/Plugins/Animation/DeformerGraph/Source/OptimusCore/Public/OptimusNode.h
+++ b/Engine/Plugins/Animation/DeformerGraph/Source/OptimusCore/Public/OptimusNode.h
@@ -55,6 +55,15 @@ public:
 
 	/** Returns true if the node can be deleted by the user */
 	virtual bool CanUserDeleteNode() const { return true; }
+	
+	/** Recreate the pins from the definition */
+	virtual void RecreatePinsFromPinDefinitions() {};
+
+	/** Rename the pin from the definition */
+	virtual void RenamePinFromPinDefinition(FName InOld, FName InNew) {};
+
+	/** Update the display name */
+	virtual void UpdateDisplayNameFromDataInterface() {};
 
 	/**
 	 * Returns the node class name. This name is immutable for the given node class.
