[Deformer Graph] Hide sub pin for array and struct types since they are not supported at the moment

diff --git a/Engine/Plugins/Animation/DeformerGraph/Source/OptimusCore/Private/OptimusDataType.cpp b/Engine/Plugins/Animation/DeformerGraph/Source/OptimusCore/Private/OptimusDataType.cpp
--- a/Engine/Plugins/Animation/DeformerGraph/Source/OptimusCore/Private/OptimusDataType.cpp
+++ b/Engine/Plugins/Animation/DeformerGraph/Source/OptimusCore/Private/OptimusDataType.cpp
@@ -152,3 +152,8 @@ int32 FOptimusDataType::GetArrayElementShaderValueSize(int32 InArrayIndex) const
 	return INDEX_NONE;
 }
 
+bool FOptimusDataType::IsArrayType() const
+{ 
+	return TypeName.ToString().StartsWith("TArray<");
+}
+
diff --git a/Engine/Plugins/Animation/DeformerGraph/Source/OptimusCore/Private/OptimusDataTypeRegistry.cpp b/Engine/Plugins/Animation/DeformerGraph/Source/OptimusCore/Private/OptimusDataTypeRegistry.cpp
--- a/Engine/Plugins/Animation/DeformerGraph/Source/OptimusCore/Private/OptimusDataTypeRegistry.cpp
+++ b/Engine/Plugins/Animation/DeformerGraph/Source/OptimusCore/Private/OptimusDataTypeRegistry.cpp
@@ -241,7 +241,8 @@ void FOptimusDataTypeRegistry::RegisterBuiltinTypes()
 	// NOTE: The pin categories should match the PC_* ones in EdGraphSchema_K2.cpp for the 
 	// fundamental types.
 	// FIXME: Turn this into an array and separate out to own file.
-	constexpr bool bShowElements = true;
+	// Data subpin in general is not currently supported, disable for now
+	constexpr bool bShowElements = false; // true;
 	constexpr bool bHideElements = false;
 
 
diff --git a/Engine/Plugins/Animation/DeformerGraph/Source/OptimusCore/Public/OptimusDataType.h b/Engine/Plugins/Animation/DeformerGraph/Source/OptimusCore/Public/OptimusDataType.h
--- a/Engine/Plugins/Animation/DeformerGraph/Source/OptimusCore/Public/OptimusDataType.h
+++ b/Engine/Plugins/Animation/DeformerGraph/Source/OptimusCore/Public/OptimusDataType.h
@@ -81,6 +81,8 @@ struct OPTIMUSCORE_API FOptimusDataType
 	// Returns the element size of each array typed shader struct member
 	int32 GetArrayElementShaderValueSize(int32 InArrayIndex) const;
 
+	bool IsArrayType() const;
+	
 	UPROPERTY()
 	FName TypeName;
 
diff --git a/Engine/Plugins/Animation/DeformerGraph/Source/OptimusEditor/Private/Widgets/SOptimusEditorGraphNode.cpp b/Engine/Plugins/Animation/DeformerGraph/Source/OptimusEditor/Private/Widgets/SOptimusEditorGraphNode.cpp
--- a/Engine/Plugins/Animation/DeformerGraph/Source/OptimusEditor/Private/Widgets/SOptimusEditorGraphNode.cpp
+++ b/Engine/Plugins/Animation/DeformerGraph/Source/OptimusEditor/Private/Widgets/SOptimusEditorGraphNode.cpp
@@ -333,7 +333,7 @@ public:
 	void Construct(
 		const FArguments& InArgs,
 		TSharedRef<SGraphPin> InPinWidget,
-		bool bIsValue,
+		bool bShowValueWidget,
 		TSharedPtr<ITableRow> InOptionalOwnerRow)
 	{
 		const UEdGraphPin* InGraphPin = InPinWidget->GetPinObj();
@@ -354,7 +354,7 @@ public:
 		TSharedRef<SWidget> LabelContent = LabelWidget;
 		const TSharedRef<SWidget> PinContent = InPinWidget;
 		
-		if (bIsLeaf && bIsInput && bIsValue)
+		if (bIsLeaf && bIsInput && bShowValueWidget)
 		{
 			TSharedPtr<SWidget> InputValueWidget = InPinWidget->GetValueWidget();
 
@@ -994,7 +994,10 @@ TSharedRef<ITableRow> SOptimusEditorGraphNode::MakeTableRowWidget(
 	const TSharedRef<STableViewBase>& OwnerTable
 	)
 {
-	const bool bIsValue = (InModelPin->GetDataDomain().IsSingleton() && InModelPin->GetPropertyFromPin() != nullptr);
+	const bool bShowValueWidget =
+		InModelPin->GetDataDomain().IsSingleton() &&
+		InModelPin->GetPropertyFromPin() != nullptr &&
+		!InModelPin->GetDataType()->IsArrayType();
 	
 	const UOptimusEditorGraphNode* EditorGraphNode = GetEditorGraphNode();
 	TSharedPtr<SGraphPin> PinWidget;
@@ -1022,7 +1025,7 @@ TSharedRef<ITableRow> SOptimusEditorGraphNode::MakeTableRowWidget(
 	}
 
 	RowWidget->SetContent(
-		SNew(SOptimusEditorGraphPinWidget, PinWidget.ToSharedRef(), bIsValue, RowWidget)
+		SNew(SOptimusEditorGraphPinWidget, PinWidget.ToSharedRef(), bShowValueWidget, RowWidget)
 		.PinLabel(this, &SOptimusEditorGraphNode::GetPinLabel, WeakPin)
 	);
 	
