[Deformer Graph] Improved node tool tips.

diff --git a/Engine/Plugins/Animation/DeformerGraph/Source/OptimusCore/Private/DataInterfaces/OptimusDataInterfaceAdvancedSkeleton.h b/Engine/Plugins/Animation/DeformerGraph/Source/OptimusCore/Private/DataInterfaces/OptimusDataInterfaceAdvancedSkeleton.h
--- a/Engine/Plugins/Animation/DeformerGraph/Source/OptimusCore/Private/DataInterfaces/OptimusDataInterfaceAdvancedSkeleton.h
+++ b/Engine/Plugins/Animation/DeformerGraph/Source/OptimusCore/Private/DataInterfaces/OptimusDataInterfaceAdvancedSkeleton.h
@@ -99,7 +99,9 @@ struct FOptimusAnimAttributeBufferArray
 };
 
 
-/** Compute Framework Data Interface for skeletal data. */
+/* Compute Framework Data Interface for skeletal data. */
+
+/* Skeleton Data Interface provides data related to skinning. */ 
 UCLASS(Category = ComputeFramework)
 class OPTIMUSCORE_API UOptimusAdvancedSkeletonDataInterface : public UOptimusComputeDataInterface
 {
diff --git a/Engine/Plugins/Animation/DeformerGraph/Source/OptimusCore/Private/DataInterfaces/OptimusDataInterfaceDebugDraw.h b/Engine/Plugins/Animation/DeformerGraph/Source/OptimusCore/Private/DataInterfaces/OptimusDataInterfaceDebugDraw.h
--- a/Engine/Plugins/Animation/DeformerGraph/Source/OptimusCore/Private/DataInterfaces/OptimusDataInterfaceDebugDraw.h
+++ b/Engine/Plugins/Animation/DeformerGraph/Source/OptimusCore/Private/DataInterfaces/OptimusDataInterfaceDebugDraw.h
@@ -38,6 +38,8 @@ struct FOptimusDebugDrawParameters
 };
 
 /** Compute Framework Data Interface for writing skinned mesh. */
+
+/* Debug Draw Data Interface provides access to a set of debug drawing hlsl functions */
 UCLASS(Category = ComputeFramework)
 class OPTIMUSCORE_API UOptimusDebugDrawDataInterface : public UOptimusComputeDataInterface
 {
diff --git a/Engine/Plugins/Animation/DeformerGraph/Source/OptimusCore/Private/DataInterfaces/OptimusDataInterfaceHalfEdge.h b/Engine/Plugins/Animation/DeformerGraph/Source/OptimusCore/Private/DataInterfaces/OptimusDataInterfaceHalfEdge.h
--- a/Engine/Plugins/Animation/DeformerGraph/Source/OptimusCore/Private/DataInterfaces/OptimusDataInterfaceHalfEdge.h
+++ b/Engine/Plugins/Animation/DeformerGraph/Source/OptimusCore/Private/DataInterfaces/OptimusDataInterfaceHalfEdge.h
@@ -17,6 +17,9 @@ class USkinnedMeshComponent;
 
 
 /** Compute Framework Data Interface for reading mesh half edge data. */
+/** Half Edge Data Interface provides vertex connectivity info, even across material sections, see ComputeNormalsTangents deformer function for example usage.
+ *  Skeletal mesh assets using this data interface should have "BuildHalfEdgeBuffers" turned on under LODInfo settings
+ */
 UCLASS(Category = ComputeFramework)
 class OPTIMUSCORE_API UOptimusHalfEdgeDataInterface : public UOptimusComputeDataInterface
 {
diff --git a/Engine/Plugins/Animation/DeformerGraph/Source/OptimusCore/Private/DataInterfaces/OptimusDataInterfaceSkinWeightsAsVertexMask.h b/Engine/Plugins/Animation/DeformerGraph/Source/OptimusCore/Private/DataInterfaces/OptimusDataInterfaceSkinWeightsAsVertexMask.h
--- a/Engine/Plugins/Animation/DeformerGraph/Source/OptimusCore/Private/DataInterfaces/OptimusDataInterfaceSkinWeightsAsVertexMask.h
+++ b/Engine/Plugins/Animation/DeformerGraph/Source/OptimusCore/Private/DataInterfaces/OptimusDataInterfaceSkinWeightsAsVertexMask.h
@@ -16,7 +16,7 @@ class FSkeletalMeshLODRenderData;
 struct FReferenceSkeleton;
 
 
-/** Compute Framework Data Interface for skeletal data. */
+/** Compute Framework Data Interface for merging skin weights of one or more bones into a per-vertex float map. */
 UCLASS(Category = ComputeFramework)
 class OPTIMUSCORE_API UOptimusSkinWeightsAsVertexMaskDataInterface : public UOptimusComputeDataInterface
 {
diff --git a/Engine/Plugins/Animation/DeformerGraph/Source/OptimusCore/Private/DataInterfaces/OptimusDataInterfaceSkinnedMesh.h b/Engine/Plugins/Animation/DeformerGraph/Source/OptimusCore/Private/DataInterfaces/OptimusDataInterfaceSkinnedMesh.h
--- a/Engine/Plugins/Animation/DeformerGraph/Source/OptimusCore/Private/DataInterfaces/OptimusDataInterfaceSkinnedMesh.h
+++ b/Engine/Plugins/Animation/DeformerGraph/Source/OptimusCore/Private/DataInterfaces/OptimusDataInterfaceSkinnedMesh.h
@@ -11,7 +11,7 @@ class FSkeletalMeshObject;
 class FSkinnedMeshDataInterfaceParameters;
 class USkinnedMeshComponent;
 
-/** Compute Framework Data Interface for reading skeletal mesh. */
+/** Compute Framework Data Interface for reading data from the skeletal mesh asset */
 UCLASS(Category = ComputeFramework)
 class OPTIMUSCORE_API UOptimusSkinnedMeshDataInterface : public UOptimusComputeDataInterface
 {
diff --git a/Engine/Plugins/Animation/DeformerGraph/Source/OptimusCore/Private/DataInterfaces/OptimusDataInterfaceSkinnedMeshRead.h b/Engine/Plugins/Animation/DeformerGraph/Source/OptimusCore/Private/DataInterfaces/OptimusDataInterfaceSkinnedMeshRead.h
--- a/Engine/Plugins/Animation/DeformerGraph/Source/OptimusCore/Private/DataInterfaces/OptimusDataInterfaceSkinnedMeshRead.h
+++ b/Engine/Plugins/Animation/DeformerGraph/Source/OptimusCore/Private/DataInterfaces/OptimusDataInterfaceSkinnedMeshRead.h
@@ -15,7 +15,7 @@ class FSkinnedMeshReadDataInterfaceParameters;
 class USkinnedMeshComponent;
 enum class EMeshDeformerOutputBuffer : uint8;
 
-/** Compute Framework Data Interface for writing skinned mesh. */
+/** Compute Framework Data Interface for reading the current state of skinned mesh, which may have been deformed by deformers run earlier */
 UCLASS(Category = ComputeFramework)
 class OPTIMUSCORE_API UOptimusSkinnedMeshReadDataInterface :
 	public UOptimusComputeDataInterface
diff --git a/Engine/Plugins/Animation/DeformerGraph/Source/OptimusCore/Private/Nodes/OptimusNode_DataInterface.cpp b/Engine/Plugins/Animation/DeformerGraph/Source/OptimusCore/Private/Nodes/OptimusNode_DataInterface.cpp
--- a/Engine/Plugins/Animation/DeformerGraph/Source/OptimusCore/Private/Nodes/OptimusNode_DataInterface.cpp
+++ b/Engine/Plugins/Animation/DeformerGraph/Source/OptimusCore/Private/Nodes/OptimusNode_DataInterface.cpp
@@ -16,6 +16,7 @@
 #include "Serialization/MemoryReader.h"
 #include "Serialization/MemoryWriter.h"
 
+
 #include UE_INLINE_GENERATED_CPP_BY_NAME(OptimusNode_DataInterface)
 
 
@@ -135,6 +136,15 @@ void UOptimusNode_DataInterface::OnDataTypeChanged(FName InTypeName)
 	DataInterfaceData->OnDataTypeChanged(InTypeName);
 }
 
+FText UOptimusNode_DataInterface::GetTooltipText() const
+{
+#if WITH_EDITORONLY_DATA
+	return DataInterfaceData->GetClass()->GetToolTipText();
+#else
+	return {};
+#endif
+}
+
 void UOptimusNode_DataInterface::ExportState(FArchive& Ar) const
 {
 	Super::ExportState(Ar);
diff --git a/Engine/Plugins/Animation/DeformerGraph/Source/OptimusCore/Private/Nodes/OptimusNode_DataInterface.h b/Engine/Plugins/Animation/DeformerGraph/Source/OptimusCore/Private/Nodes/OptimusNode_DataInterface.h
--- a/Engine/Plugins/Animation/DeformerGraph/Source/OptimusCore/Private/Nodes/OptimusNode_DataInterface.h
+++ b/Engine/Plugins/Animation/DeformerGraph/Source/OptimusCore/Private/Nodes/OptimusNode_DataInterface.h
@@ -66,6 +66,7 @@ protected:
 	TOptional<FText> ValidateForCompile(const FOptimusPinTraversalContext& InContext) const override;
 	void PostLoadNodeSpecificData() override;
 	void OnDataTypeChanged(FName InTypeName) override;
+	FText GetTooltipText() const override;
 
 	void ExportState(FArchive& Ar) const override;
 	void ImportState(FArchive& Ar) override;
diff --git a/Engine/Plugins/Animation/DeformerGraph/Source/OptimusCore/Private/OptimusNode.cpp b/Engine/Plugins/Animation/DeformerGraph/Source/OptimusCore/Private/OptimusNode.cpp
--- a/Engine/Plugins/Animation/DeformerGraph/Source/OptimusCore/Private/OptimusNode.cpp
+++ b/Engine/Plugins/Animation/DeformerGraph/Source/OptimusCore/Private/OptimusNode.cpp
@@ -91,6 +91,10 @@ bool UOptimusNode::SetDisplayName(FText InDisplayName)
 	return true;
 }
 
+FText UOptimusNode::GetTooltipText() const
+{
+	return {};
+}
 
 
 bool UOptimusNode::SetGraphPosition(const FVector2D& InPosition)
diff --git a/Engine/Plugins/Animation/DeformerGraph/Source/OptimusCore/Public/OptimusNode.h b/Engine/Plugins/Animation/DeformerGraph/Source/OptimusCore/Public/OptimusNode.h
--- a/Engine/Plugins/Animation/DeformerGraph/Source/OptimusCore/Public/OptimusNode.h
+++ b/Engine/Plugins/Animation/DeformerGraph/Source/OptimusCore/Public/OptimusNode.h
@@ -86,6 +86,12 @@ public:
 	 */ 
 	bool SetDisplayName(FText InDisplayName);
 
+	/**
+	 * Returns the tool tip.
+	 * @return The tool tip to show to the user.
+	 */
+	virtual FText GetTooltipText() const;
+	
 	/**
 	 * Sets the position in the graph UI that the node should be shown at.
 	 * @param InPosition The coordinates of the node's position.
diff --git a/Engine/Plugins/Animation/DeformerGraph/Source/OptimusEditor/Private/OptimusEditorGraphNode.cpp b/Engine/Plugins/Animation/DeformerGraph/Source/OptimusEditor/Private/OptimusEditorGraphNode.cpp
--- a/Engine/Plugins/Animation/DeformerGraph/Source/OptimusEditor/Private/OptimusEditorGraphNode.cpp
+++ b/Engine/Plugins/Animation/DeformerGraph/Source/OptimusEditor/Private/OptimusEditorGraphNode.cpp
@@ -375,6 +375,16 @@ FLinearColor UOptimusEditorGraphNode::GetNodeTitleColor() const
 	return Super::GetNodeTitleColor();
 }
 
+FText UOptimusEditorGraphNode::GetTooltipText() const
+{
+	if (ensure(ModelNode))
+	{
+		return ModelNode->GetTooltipText();
+	}
+
+	return {};
+}
+
 
 bool UOptimusEditorGraphNode::CreateGraphPinFromModelPin(
 	const UOptimusNodePin* InModelPin,
diff --git a/Engine/Plugins/Animation/DeformerGraph/Source/OptimusEditor/Private/OptimusEditorGraphNode.h b/Engine/Plugins/Animation/DeformerGraph/Source/OptimusEditor/Private/OptimusEditorGraphNode.h
--- a/Engine/Plugins/Animation/DeformerGraph/Source/OptimusEditor/Private/OptimusEditorGraphNode.h
+++ b/Engine/Plugins/Animation/DeformerGraph/Source/OptimusEditor/Private/OptimusEditorGraphNode.h
@@ -49,6 +49,7 @@ public:
 	FText GetNodeTitle(ENodeTitleType::Type InTitleType) const override;
 	void GetNodeContextMenuActions(UToolMenu* InMenu, UGraphNodeContextMenuContext* InContext) const override;
 	FLinearColor GetNodeTitleColor() const override;
+	FText GetTooltipText() const override;
 	
 	// FIXME: Move to private and add accessor function.
 	UPROPERTY()
diff --git a/Engine/Plugins/Animation/DeformerGraph/Source/OptimusEditor/Private/Widgets/SOptimusEditorGraphNode.cpp b/Engine/Plugins/Animation/DeformerGraph/Source/OptimusEditor/Private/Widgets/SOptimusEditorGraphNode.cpp
--- a/Engine/Plugins/Animation/DeformerGraph/Source/OptimusEditor/Private/Widgets/SOptimusEditorGraphNode.cpp
+++ b/Engine/Plugins/Animation/DeformerGraph/Source/OptimusEditor/Private/Widgets/SOptimusEditorGraphNode.cpp
@@ -15,6 +15,7 @@
 
 #include "Editor.h"
 #include "GraphEditorSettings.h"
+#include "IDocumentation.h"
 #include "OptimusComponentSource.h"
 #include "SGraphPin.h"
 #include "Widgets/Images/SImage.h"
@@ -884,6 +885,20 @@ void SOptimusEditorGraphNode::Tick(
 	}
 }
 
+void SOptimusEditorGraphNode::UpdateGraphNode()
+{
+	if (GraphNode)
+	{
+		// Assign a tool tip before SGraphNode::UpdateGraphNode because
+		// by default SGraphNode adds the raw name of the graph node to
+		// the tool tip if the EdGraph is transient, which we don't want
+		FText ToolTipText = GraphNode->GetTooltipText();
+		TSharedRef<SToolTip> ToolTip = IDocumentation::Get()->CreateToolTip( ToolTipText, NULL, GraphNode->GetDocumentationLink(), GraphNode->GetDocumentationExcerptName() );
+		SetToolTip(ToolTip);
+	}
+	SGraphNode::UpdateGraphNode();
+}
+
 
 UOptimusEditorGraphNode* SOptimusEditorGraphNode::GetEditorGraphNode() const
 {
diff --git a/Engine/Plugins/Animation/DeformerGraph/Source/OptimusEditor/Private/Widgets/SOptimusEditorGraphNode.h b/Engine/Plugins/Animation/DeformerGraph/Source/OptimusEditor/Private/Widgets/SOptimusEditorGraphNode.h
--- a/Engine/Plugins/Animation/DeformerGraph/Source/OptimusEditor/Private/Widgets/SOptimusEditorGraphNode.h
+++ b/Engine/Plugins/Animation/DeformerGraph/Source/OptimusEditor/Private/Widgets/SOptimusEditorGraphNode.h
@@ -32,6 +32,7 @@ public:
 	TSharedPtr<SGraphPin> GetHoveredPin(const FGeometry& MyGeometry, const FPointerEvent& MouseEvent) const override;
 	void RefreshErrorInfo() override;
 	void Tick(const FGeometry& AllottedGeometry, const double InCurrentTime, const float InDeltaTime) override;
+	void UpdateGraphNode() override;
 
 private:
 	UOptimusEditorGraphNode *GetEditorGraphNode() const;
