[Deformer Graph] Made sure FRotator is not allowed to be set as an anim attribut type.

diff --git a/Engine/Plugins/Animation/DeformerGraph/Source/OptimusCore/Private/OptimusDataTypeRegistry.cpp b/Engine/Plugins/Animation/DeformerGraph/Source/OptimusCore/Private/OptimusDataTypeRegistry.cpp
--- a/Engine/Plugins/Animation/DeformerGraph/Source/OptimusCore/Private/OptimusDataTypeRegistry.cpp
+++ b/Engine/Plugins/Animation/DeformerGraph/Source/OptimusCore/Private/OptimusDataTypeRegistry.cpp
@@ -431,7 +431,7 @@ void FOptimusDataTypeRegistry::RegisterBuiltinTypes()
 	    FShaderValueType::Get(EShaderFundamentalType::Float, 3, 3),
 	    {},
 	    bShowElements,
-	    EOptimusDataTypeUsageFlags::Resource | EOptimusDataTypeUsageFlags::Variable | EOptimusDataTypeUsageFlags::DataInterfaceOutput | EOptimusDataTypeUsageFlags::PinType | EOptimusDataTypeUsageFlags::PerBoneAnimAttribute);
+	    EOptimusDataTypeUsageFlags::Resource | EOptimusDataTypeUsageFlags::Variable | EOptimusDataTypeUsageFlags::DataInterfaceOutput | EOptimusDataTypeUsageFlags::PinType);
 
 	// FTransform -> float4x4
 	Registry.RegisterType(
