[Deformer Graph] Fixed having multiple variable nodes of the same variable cause shader compilation error around redefinition of functions from the graph data interface

diff --git a/Engine/Plugins/Animation/DeformerGraph/Source/OptimusCore/Private/OptimusDeformer.cpp b/Engine/Plugins/Animation/DeformerGraph/Source/OptimusCore/Private/OptimusDeformer.cpp
--- a/Engine/Plugins/Animation/DeformerGraph/Source/OptimusCore/Private/OptimusDeformer.cpp
+++ b/Engine/Plugins/Animation/DeformerGraph/Source/OptimusCore/Private/OptimusDeformer.cpp
@@ -2272,7 +2272,7 @@ FOptimusNodeGraphCompilationResult UOptimusDeformer::CompileNodeGraphToComputeGr
 	// Create the graph data interface and fill it with the value nodes.
 	UOptimusGraphDataInterface* GraphDataInterface = NewObject<UOptimusGraphDataInterface>(this);
 	DataInterfaceToBindingIndexMap.Add(GraphDataInterface) = GraphDataComponentBinding->GetIndex();
-	
+
 	TArray<FOptimusGraphVariableDescription> ValueNodeDescriptions;
 	ValueNodeDescriptions.Reserve(ActiveValueNodes.Num());	
 
@@ -2282,23 +2282,20 @@ FOptimusNodeGraphCompilationResult UOptimusDeformer::CompileNodeGraphToComputeGr
 		FOptimusValueIdentifier ValueId = ValueProvider->GetValueIdentifier();
 		EOptimusValueUsage ValueUsage = ValueNodeUsageMap[ValueNode];
 
-		FOptimusValueDescription Description;
+		FOptimusValueDescription& Description = NodeGraphValueMap.FindOrAdd(ValueId);
 		Description.DataType = ValueProvider->GetValueDataType();
-		Description.ValueUsage = ValueUsage;
+		Description.ValueUsage |= ValueUsage;
 
-		if (EnumHasAnyFlags(Description.ValueUsage, EOptimusValueUsage::CPU))
+		if (EnumHasAnyFlags(Description.ValueUsage, EOptimusValueUsage::CPU) && !Description.Value.IsInitialized())
 		{
 			Description.Value = ValueProvider->GetValue();
 		}
-		else if (EnumHasAnyFlags(Description.ValueUsage, EOptimusValueUsage::GPU))
+		if (EnumHasAnyFlags(Description.ValueUsage, EOptimusValueUsage::GPU) && !Description.ShaderValue.IsValid())
 		{
 			Description.ShaderValue = ValueProvider->GetValue().GetShaderValue(ValueProvider->GetValueDataType());
-		}
-		
-		NodeGraphValueMap.Add(ValueId, MoveTemp(Description));	
 
-		if (EnumHasAnyFlags(ValueUsage, EOptimusValueUsage::GPU))
-		{
+			// Add GPU values to the graph data interface too, which receives value from the deformer instance
+			// and we only want to add each value once, even when there are multiple value nodes all referencing the same value (mainly variable nodes)
 			FOptimusGraphVariableDescription ValueNodeDescription;;
 			ValueNodeDescription.ValueId = ValueProvider->GetValueIdentifier();
 			ValueNodeDescription.Name = ValueNodeDescription.ValueId.Name.ToString();
