[SkinnedMeshComponent & Deformer Graph] Added an async API to readback render geometry for a mesh deformer deformed skeletal mesh and generate a FMeshDescription

diff --git a/Engine/Plugins/Animation/DeformerGraph/Source/OptimusCore/OptimusCore.Build.cs b/Engine/Plugins/Animation/DeformerGraph/Source/OptimusCore/OptimusCore.Build.cs
--- a/Engine/Plugins/Animation/DeformerGraph/Source/OptimusCore/OptimusCore.Build.cs
+++ b/Engine/Plugins/Animation/DeformerGraph/Source/OptimusCore/OptimusCore.Build.cs
@@ -27,7 +27,9 @@ namespace UnrealBuildTool.Rules
 					"Renderer",
 					"RHI",
 					"RigVM",
-					"ControlRig"
+					"ControlRig",
+					"StaticMeshDescription",
+					"MeshDescription"
 				}
 			);
 
diff --git a/Engine/Plugins/Animation/DeformerGraph/Source/OptimusCore/Private/DataInterfaces/OptimusDataInterfaceSkinnedMeshWrite.cpp b/Engine/Plugins/Animation/DeformerGraph/Source/OptimusCore/Private/DataInterfaces/OptimusDataInterfaceSkinnedMeshWrite.cpp
--- a/Engine/Plugins/Animation/DeformerGraph/Source/OptimusCore/Private/DataInterfaces/OptimusDataInterfaceSkinnedMeshWrite.cpp
+++ b/Engine/Plugins/Animation/DeformerGraph/Source/OptimusCore/Private/DataInterfaces/OptimusDataInterfaceSkinnedMeshWrite.cpp
@@ -15,6 +15,15 @@
 #include "DataDrivenShaderPlatformInfo.h"
 #include "Animation/MeshDeformerInstance.h"
 
+#if WITH_EDITORONLY_DATA
+#include "OptimusDeformerInstance.h"
+#include "Components/SkeletalMeshComponent.h"
+#include "Rendering/SkeletalMeshModel.h"
+#include "MeshAttributes.h"
+#include "StaticMeshAttributes.h"
+#include "Async/ParallelFor.h"
+#endif // WITH_EDITORONLY_DATA
+
 #include UE_INLINE_GENERATED_CPP_BY_NAME(OptimusDataInterfaceSkinnedMeshWrite)
 
 
@@ -153,18 +162,240 @@ EMeshDeformerOutputBuffer UOptimusSkinnedMeshWriteDataInterface::GetOutputBuffer
 
 FComputeDataProviderRenderProxy* UOptimusSkinnedMeshWriteDataProvider::GetRenderProxy()
 {
-	return new FOptimusSkinnedMeshWriteDataProviderProxy(SkinnedMesh, OutputMask, &LastLodIndexCachedByRenderProxy);
+	FOptimusSkinnedMeshWriteDataProviderProxy* Proxy = new FOptimusSkinnedMeshWriteDataProviderProxy();
+	
+	Proxy->SkeletalMeshObject = SkinnedMesh? SkinnedMesh->MeshObject : nullptr;
+	Proxy->OutputMask = OutputMask;
+	Proxy->LastLodIndexPtr = &LastLodIndexCachedByRenderProxy;
+
+#if WITH_EDITORONLY_DATA
+	SetupGeometryReadbackIfNeeded(Proxy);
+#endif
+	
+	return Proxy;
 }
 
+void UOptimusSkinnedMeshWriteDataProvider::SetDeformerInstance(UOptimusDeformerInstance* InInstance)
+{
+#if WITH_EDITORONLY_DATA
+	DeformerInstance = InInstance;
+#endif // WITH_EDITORONLY_DATA
+}
+#if WITH_EDITORONLY_DATA
+void UOptimusSkinnedMeshWriteDataProvider::SetupGeometryReadbackIfNeeded(FOptimusSkinnedMeshWriteDataProviderProxy* Proxy)
+{
+	// Process geometry readback requests if there is one
+	if (DeformerInstance->OnGeometryReadbackCompletedCallbacks.Num() > 0)
+	{
+		Proxy->RequestGeometryReadback_RenderThread = [
+			WeakThis = TWeakObjectPtr(this),
+			FrameNumber = GFrameNumber,
+			GeometryReadbackCallbacks = DeformerInstance->OnGeometryReadbackCompletedCallbacks
+		](
+			TArray<FComputeDataProviderRenderProxy::FReadbackData>& OutBufferReadbacks,
+			FSkeletalMeshObject* MeshObject,
+			FRDGBuffer* PositionBuffer,
+			FRDGBuffer* TangentBuffer,
+			FRDGBuffer* ColorBuffer
+		)
+		{
+			if (UOptimusSkinnedMeshWriteDataProvider* DataProvider = WeakThis.Get())
+			{
+				// Each proxy allocates a FGeometryReadback which stores the result of its readback and is processed async at a later time
+				TSharedRef<FGeometryReadback> GeometryReadback = MakeShared<FGeometryReadback>();
+				DataProvider->GeometryReadbacks_RenderThread.Enqueue(GeometryReadback);
+
+				GeometryReadback->FrameNumber = FrameNumber;
+				GeometryReadback->GeometryReadbackCallbacks = GeometryReadbackCallbacks;
+				GeometryReadback->LodIndex = MeshObject->GetLOD();
+
+				auto SetupBufferReadback = [WeakThis, &OutBufferReadbacks, FrameNumber](FRDGBuffer* GPUBuffer, FBufferReadback& BufferReadback)
+				{
+					if (GPUBuffer)
+					{
+						BufferReadback.bShouldReadback = true;
+						BufferReadback.OnReadbackCompleted_RenderThread = [WeakThis, &ReadbackData = BufferReadback.ReadbackData, FrameNumber](const void* InData, int InNumBytes)
+						{
+							if (ensure(InNumBytes > 0))
+							{
+								check(ReadbackData.Num() == 0);
+								ReadbackData.SetNumUninitialized(InNumBytes);
+								FPlatformMemory::Memcpy(ReadbackData.GetData(), InData, InNumBytes);
+							}
+							
+							if (UOptimusSkinnedMeshWriteDataProvider* DataProvider = WeakThis.Get())
+							{
+								// Each time a new buffer readback completes, we want to check if the geometry readback it
+								// belongs to is ready for processing. (All its buffer readbacks are ready)
+								DataProvider->ProcessCompletedGeometryReadback_RenderThread();	
+							}
+						};
+
+						// Raw buffer readback request to the compute worker
+						FComputeDataProviderRenderProxy::FReadbackData& BufferReadbackRequestRef = OutBufferReadbacks.AddDefaulted_GetRef();
+						BufferReadbackRequestRef.Buffer = GPUBuffer;
+						BufferReadbackRequestRef.NumBytes = GPUBuffer->Desc.GetSize();
+						BufferReadbackRequestRef.ReadbackCallback_RenderThread = &BufferReadback.OnReadbackCompleted_RenderThread;
+					}
+				};
+
+				SetupBufferReadback(PositionBuffer, GeometryReadback->Position);
+				SetupBufferReadback(TangentBuffer, GeometryReadback->Tangent);
+				SetupBufferReadback(ColorBuffer, GeometryReadback->Color);
+			}
+		};
+
+		DeformerInstance->OnGeometryReadbackCompletedCallbacks.Reset();
+	}	
+}
+#endif // WITH_EDITORONLY_DATA
 
-FOptimusSkinnedMeshWriteDataProviderProxy::FOptimusSkinnedMeshWriteDataProviderProxy(USkinnedMeshComponent* InSkinnedMeshComponent, uint64 InOutputMask, int32* InLastLodIndexPtr)
+#if WITH_EDITORONLY_DATA
+void UOptimusSkinnedMeshWriteDataProvider::ProcessCompletedGeometryReadback_RenderThread()
 {
-	SkeletalMeshObject = InSkinnedMeshComponent != nullptr ? InSkinnedMeshComponent->MeshObject : nullptr;
-	OutputMask = InOutputMask;
-	LastLodIndexPtr = InLastLodIndexPtr;
-	
-	check(LastLodIndexPtr != nullptr);
+	// Try to process readbacks in the order they are added,
+	// later readbacks won't be processed until earlier readbacks have been popped and processed
+	while (!GeometryReadbacks_RenderThread.IsEmpty())
+	{
+		TSharedPtr<FGeometryReadback> Readback = *GeometryReadbacks_RenderThread.Peek();
+
+		// Stop processing if a buffer is not ready
+		if (Readback->Position.bShouldReadback && Readback->Position.ReadbackData.Num() == 0)
+		{
+			break;
+		}
+		if (Readback->Tangent.bShouldReadback && Readback->Tangent.ReadbackData.Num() == 0)
+		{
+			break;
+		}
+		if (Readback->Color.bShouldReadback && Readback->Color.ReadbackData.Num() == 0)
+		{
+			break;
+		}
+
+		// Ready for processing
+		GeometryReadbacks_RenderThread.Pop();
+		
+		auto ReadbackGeometryToMeshDescription = [WeakMesh = TWeakObjectPtr(SkinnedMesh), Readback]()
+		{
+			if (const USkeletalMeshComponent* SkeletalMeshComponent = Cast<USkeletalMeshComponent>(WeakMesh.Get()))
+			{
+				const int32 LODIndex = Readback->LodIndex;
+				const TArray<uint8>& Positions = Readback->Position.ReadbackData; 
+				const TArray<uint8>& NormalsTangents = Readback->Tangent.ReadbackData; 
+				const TArray<uint8>& Colors = Readback->Color.ReadbackData; 
+				
+				const USkeletalMesh* SkeletalMesh = SkeletalMeshComponent->GetSkeletalMeshAsset();
+				FMeshDescription MeshDescription;
+				SkeletalMesh->CloneMeshDescription(LODIndex, MeshDescription);
+
+				const TArray<int32>& RenderToImportedMap = SkeletalMesh->GetImportedModel()->LODModels[LODIndex].MeshToImportVertexMap;
+				// Avoid conditional build during parallel for loop below
+				MeshDescription.BuildVertexIndexers();
+
+				constexpr int32 SizeOfVertPosition = FSkeletalMeshDeformerHelpers::PosBufferElementMultiplier * FSkeletalMeshDeformerHelpers::PosBufferBytesPerElement; // float3, 3 * float, 12 Bytes per vert
+				constexpr int32 SizeOfVertTangents = FSkeletalMeshDeformerHelpers::TangentBufferElementMultiplier * FSkeletalMeshDeformerHelpers::TangentBufferBytesPerElement; // 2 * half4, which is 4 * 2 bytes, 16 bytes per vert
+				constexpr int32 SizeOfVertColor = FSkeletalMeshDeformerHelpers::ColorBufferBytesPerElement; // 1 * 4 * 8bit, 4 bytes per vert
+
+				const int32 NumVertPositions = Positions.Num() / SizeOfVertPosition;
+				const int32 NumVertNormalsTangents= NormalsTangents.Num() / SizeOfVertTangents;
+				const int32 NumVertColors = Colors.Num() / SizeOfVertColor;
+
+				// Make sure sizes of all buffers make sense
+				if (ensure(NumVertPositions == NumVertNormalsTangents || NumVertPositions == 0 || NumVertNormalsTangents == 0) && 
+					ensure(NumVertNormalsTangents == NumVertColors || NumVertNormalsTangents == 0 || NumVertColors == 0) &&
+					ensure(NumVertPositions == NumVertColors || NumVertPositions == 0 || NumVertColors == 0))
+				{
+					const int32 NumVertices =
+						NumVertPositions != 0 ? NumVertPositions :
+						NumVertNormalsTangents != 0 ? NumVertNormalsTangents :
+						NumVertColors != 0 ? NumVertColors : 0;
+				
+					// Parallel for each render vert
+					ParallelFor(NumVertices,[&MeshDescription, NumVertices, &RenderToImportedMap, &Positions, &NormalsTangents, &Colors](int32 RenderVertIndex)
+					{
+						const FVertexID ImportedVertexID(RenderToImportedMap[RenderVertIndex]);
+
+						// Positions
+						if (Positions.Num() > 0)
+						{
+							const float* Ptr = reinterpret_cast<const float*>(Positions.GetData());
+						
+							const float PositionX = Ptr[RenderVertIndex*3 + 0];
+							const float PositionY = Ptr[RenderVertIndex*3 + 1];
+							const float PositionZ = Ptr[RenderVertIndex*3 + 2];
+
+							MeshDescription.VertexAttributes().SetAttribute(ImportedVertexID, MeshAttribute::Vertex::Position, 0, FVector3f(PositionX,PositionY,PositionZ));
+						}
+
+						// Normals Tangents
+						const TArrayView<const FVertexInstanceID> VertexInstances = MeshDescription.GetVertexVertexInstanceIDs(ImportedVertexID);
+						if (NormalsTangents.Num() > 0)
+						{
+							const FPackedRGBA16N* Ptr = reinterpret_cast<const FPackedRGBA16N*>(NormalsTangents.GetData());
+							const FVector3f Tangent = Ptr[RenderVertIndex*2 + 0].ToFVector3f();
+							const FVector4f Normal = Ptr[RenderVertIndex*2 + 1].ToFVector4f();
+						
+							for (const FVertexInstanceID VertexInstanceID : VertexInstances)
+							{
+								MeshDescription.VertexInstanceAttributes().SetAttribute<FVector3f>(
+									VertexInstanceID,
+									MeshAttribute::VertexInstance::Normal,
+									0,
+									FVector3f(Normal));
+							
+								MeshDescription.VertexInstanceAttributes().SetAttribute<FVector3f>(
+									VertexInstanceID,
+									MeshAttribute::VertexInstance::Tangent,
+									0,
+									Tangent);
+
+								MeshDescription.VertexInstanceAttributes().SetAttribute<float>(
+									VertexInstanceID,
+									MeshAttribute::VertexInstance::BinormalSign,
+									0,
+									Normal.W);
+							}
+						}
+
+						// Colors
+						if (Colors.Num() > 0)
+						{
+							check(sizeof(FColor) == 4);
+						
+							const FColor* Ptr = reinterpret_cast<const FColor*>(Colors.GetData());
+
+							for (const FVertexInstanceID VertexInstanceID : VertexInstances)
+							{
+								MeshDescription.VertexInstanceAttributes().SetAttribute<FVector4f>(
+									VertexInstanceID,
+									MeshAttribute::VertexInstance::Color,
+									0,
+									FVector4f(Ptr[RenderVertIndex]));
+							}
+						}
+					});
+
+					for (TFunction<void(const FMeshDescription&)> Callback : Readback->GeometryReadbackCallbacks)
+					{
+						Callback(MeshDescription);	
+					}
+				}
+			}	
+		};
+
+		// Readback data should be processed sequentially on a background thread, so using last task as prerequisite
+		if (LastReadbackProcessingTask_RenderThread)
+		{
+			LastReadbackProcessingTask_RenderThread = FFunctionGraphTask::CreateAndDispatchWhenReady(ReadbackGeometryToMeshDescription, TStatId(), LastReadbackProcessingTask_RenderThread, ENamedThreads::AnyBackgroundThreadNormalTask);	
+		}
+		else
+		{
+			LastReadbackProcessingTask_RenderThread = FFunctionGraphTask::CreateAndDispatchWhenReady(ReadbackGeometryToMeshDescription, TStatId(), NULL, ENamedThreads::AnyBackgroundThreadNormalTask);	
+		}
+	}
 }
+#endif // WITH_EDITORONLY_DATA
 
 bool FOptimusSkinnedMeshWriteDataProviderProxy::IsValid(FValidationData const& InValidationData) const
 {
@@ -207,6 +438,14 @@ void FOptimusSkinnedMeshWriteDataProviderProxy::AllocateResources(FRDGBuilder& G
 	}
 	else
 	{
+#if WITH_EDITORONLY_DATA
+		if (RequestGeometryReadback_RenderThread)
+		{
+			PositionBuffer = FSkeletalMeshDeformerHelpers::GetAllocatedPositionBuffer(GraphBuilder, SkeletalMeshObject, LodIndex);
+		}
+#endif // WITH_EDITORONLY_DATA
+		
+		// Unused UAV
 		PositionBufferUAV = GraphBuilder.CreateUAV(GraphBuilder.RegisterExternalBuffer(GWhiteVertexBufferWithRDG->Buffer), PF_R32_FLOAT);
 	}
 
@@ -221,6 +460,14 @@ void FOptimusSkinnedMeshWriteDataProviderProxy::AllocateResources(FRDGBuilder& G
 	}
 	else
 	{
+#if WITH_EDITORONLY_DATA
+		if (RequestGeometryReadback_RenderThread)
+		{
+			TangentBuffer = FSkeletalMeshDeformerHelpers::GetAllocatedTangentBuffer(GraphBuilder, SkeletalMeshObject, LodIndex);
+		}
+#endif // WITH_EDITORONLY_DATA
+
+		// Unused UAV
 		TangentBufferUAV = GraphBuilder.CreateUAV(GraphBuilder.RegisterExternalBuffer(GWhiteVertexBufferWithRDG->Buffer), TangentsFormat);
 	}
 
@@ -232,6 +479,14 @@ void FOptimusSkinnedMeshWriteDataProviderProxy::AllocateResources(FRDGBuilder& G
 	}
 	else
 	{
+#if WITH_EDITORONLY_DATA
+		if (RequestGeometryReadback_RenderThread)
+		{
+			ColorBuffer = FSkeletalMeshDeformerHelpers::GetAllocatedColorBuffer(GraphBuilder, SkeletalMeshObject, LodIndex);	
+		}
+#endif // WITH_EDITORONLY_DATA
+	
+		// Unused UAV
 		ColorBufferUAV = GraphBuilder.CreateUAV(GraphBuilder.RegisterExternalBuffer(GWhiteVertexBufferWithRDG->Buffer), PF_A32B32G32R32F);
 	}
 
@@ -257,3 +512,13 @@ void FOptimusSkinnedMeshWriteDataProviderProxy::GatherDispatchData(FDispatchData
 		Parameters.ColorBufferUAV = ColorBufferUAV;
 	}
 }
+
+void FOptimusSkinnedMeshWriteDataProviderProxy::GetReadbackData(TArray<FReadbackData>& OutReadbackData) const
+{
+#if WITH_EDITORONLY_DATA
+	if (RequestGeometryReadback_RenderThread)
+	{
+		RequestGeometryReadback_RenderThread(OutReadbackData, SkeletalMeshObject, PositionBuffer, TangentBuffer, ColorBuffer);
+	}
+#endif // WITH_EDITORONLY_DATA
+}
diff --git a/Engine/Plugins/Animation/DeformerGraph/Source/OptimusCore/Private/DataInterfaces/OptimusDataInterfaceSkinnedMeshWrite.h b/Engine/Plugins/Animation/DeformerGraph/Source/OptimusCore/Private/DataInterfaces/OptimusDataInterfaceSkinnedMeshWrite.h
--- a/Engine/Plugins/Animation/DeformerGraph/Source/OptimusCore/Private/DataInterfaces/OptimusDataInterfaceSkinnedMeshWrite.h
+++ b/Engine/Plugins/Animation/DeformerGraph/Source/OptimusCore/Private/DataInterfaces/OptimusDataInterfaceSkinnedMeshWrite.h
@@ -2,10 +2,15 @@
 
 #pragma once
 
+#include "IOptimusDeformerInstanceAccessor.h"
 #include "IOptimusOutputBufferWriter.h"
 #include "OptimusComputeDataInterface.h"
 #include "ComputeFramework/ComputeDataProvider.h"
 
+#if WITH_EDITORONLY_DATA
+#include "Containers/Queue.h"
+#endif // WITH_EDITORONLY_DATA
+
 #include "OptimusDataInterfaceSkinnedMeshWrite.generated.h"
 
 class FRDGBuffer;
@@ -14,6 +19,10 @@ class FSkeletalMeshObject;
 class FSkinedMeshWriteDataInterfaceParameters;
 class USkinnedMeshComponent;
 
+#if WITH_EDITORONLY_DATA
+struct FMeshDescription;
+#endif // WITH_EDITORONLY_DATA
+
 /** Compute Framework Data Interface for writing skinned mesh. */
 UCLASS(Category = ComputeFramework)
 class OPTIMUSCORE_API UOptimusSkinnedMeshWriteDataInterface :
@@ -40,6 +49,7 @@ public:
 	void GetShaderHash(FString& InOutKey) const override;
 	void GetHLSL(FString& OutHLSL, FString const& InDataInterfaceName) const override;
 	UComputeDataProvider* CreateDataProvider(TObjectPtr<UObject> InBinding, uint64 InInputMask, uint64 InOutputMask) const override;
+	bool GetRequiresReadback() const override { return true; }
 	//~ End UComputeDataInterface Interface
 
 	//~ Begin IOptimusOutputBufferWriter Interface
@@ -49,9 +59,13 @@ private:
 	static TCHAR const* TemplateFilePath;
 };
 
+
+
 /** Compute Framework Data Provider for writing skinned mesh. */
 UCLASS(BlueprintType, editinlinenew, Category = ComputeFramework)
-class UOptimusSkinnedMeshWriteDataProvider : public UComputeDataProvider
+class UOptimusSkinnedMeshWriteDataProvider :
+	public UComputeDataProvider,
+	public IOptimusDeformerInstanceAccessor
 {
 	GENERATED_BODY()
 
@@ -61,31 +75,78 @@ public:
 
 	uint64 OutputMask = 0;
 
-	// Served as persistent storage for the provider proxy, should not be used by the data provider itself
-	int32 LastLodIndexCachedByRenderProxy = 0;
-
 	//~ Begin UComputeDataProvider Interface
 	FComputeDataProviderRenderProxy* GetRenderProxy() override;
 	//~ End UComputeDataProvider Interface
+
+	//~ Begin IOptimusDeformerInstanceAccessor Interface
+	void SetDeformerInstance(UOptimusDeformerInstance* InInstance) override;
+	//~ End IOptimusDeformerInstanceAccessor Interface
+
+private:
+	// Served as persistent storage for the provider proxy, should not be used by the data provider itself
+	int32 LastLodIndexCachedByRenderProxy = 0;
+
+#if WITH_EDITORONLY_DATA
+	UPROPERTY()
+	TObjectPtr<UOptimusDeformerInstance> DeformerInstance = nullptr;
+
+	void SetupGeometryReadbackIfNeeded(class FOptimusSkinnedMeshWriteDataProviderProxy* Proxy);
+	void ProcessCompletedGeometryReadback_RenderThread();
+	
+	struct FBufferReadback
+	{
+		bool bShouldReadback = false;
+		TArray<uint8> ReadbackData;
+		FComputeDataProviderRenderProxy::FReadbackCallback OnReadbackCompleted_RenderThread;
+	};
+	
+	struct FGeometryReadback
+	{
+		uint64 FrameNumber = 0;
+		TArray<TFunction<void(const FMeshDescription&)>> GeometryReadbackCallbacks;
+		
+		int32 LodIndex = 0;
+		
+		FBufferReadback Position;
+		FBufferReadback Tangent;
+		FBufferReadback Color;
+	};
+	
+	// Served as persistent storage for the proxy, render thread only
+	TQueue<TSharedPtr<FGeometryReadback>> GeometryReadbacks_RenderThread;
+
+	FGraphEventRef LastReadbackProcessingTask_RenderThread;
+#endif // WITH_EDITORONLY_DATA
 };
 
 class FOptimusSkinnedMeshWriteDataProviderProxy : public FComputeDataProviderRenderProxy
 {
 public:
-	FOptimusSkinnedMeshWriteDataProviderProxy(USkinnedMeshComponent* InSkinnedMeshComponent, uint64 InOutputMask, int32* InLastLodIndexPtr);
-
 	//~ Begin FComputeDataProviderRenderProxy Interface
 	bool IsValid(FValidationData const& InValidationData) const override;
 	void AllocateResources(FRDGBuilder& GraphBuilder) override;
 	void GatherDispatchData(FDispatchData const& InDispatchData) override;
+	void GetReadbackData(TArray<FReadbackData>& OutReadbackData) const override;
 	//~ End FComputeDataProviderRenderProxy Interface
 
-private:
-	using FParameters = FSkinedMeshWriteDataInterfaceParameters;
-
 	FSkeletalMeshObject* SkeletalMeshObject = nullptr;
 	uint64 OutputMask = 0;
-	int32* LastLodIndexPtr = nullptr; 
+	int32* LastLodIndexPtr = nullptr;
+
+#if WITH_EDITORONLY_DATA
+	TFunction<
+		void(
+			TArray<FReadbackData>&,
+			FSkeletalMeshObject*,
+			FRDGBuffer*,
+			FRDGBuffer*,
+			FRDGBuffer*
+		)> RequestGeometryReadback_RenderThread;
+#endif // WITH_EDITORONLY_DATA
+	
+private:
+	using FParameters = FSkinedMeshWriteDataInterfaceParameters;
 
 	FRDGBuffer* PositionBuffer = nullptr;
 	FRDGBufferUAV* PositionBufferUAV = nullptr;
diff --git a/Engine/Plugins/Animation/DeformerGraph/Source/OptimusCore/Private/OptimusDeformerDynamicInstanceManager.cpp b/Engine/Plugins/Animation/DeformerGraph/Source/OptimusCore/Private/OptimusDeformerDynamicInstanceManager.cpp
--- a/Engine/Plugins/Animation/DeformerGraph/Source/OptimusCore/Private/OptimusDeformerDynamicInstanceManager.cpp
+++ b/Engine/Plugins/Animation/DeformerGraph/Source/OptimusCore/Private/OptimusDeformerDynamicInstanceManager.cpp
@@ -49,22 +49,11 @@ void UOptimusDeformerDynamicInstanceManager::EnqueueWork(FEnqueueWorkDesc const&
 	
 	// Enqueue work
 
-	// Making sure instances in the queue are dispatched sequentially
-	uint8 NumComputeGraphsPossiblyEnqueued = 0;
-	// Used to inform later instances whether specific buffers have valid data in them
-	EMeshDeformerOutputBuffer OutputBuffers = EMeshDeformerOutputBuffer::None;
-	
-	auto EnqueueInstance = [&NumComputeGraphsPossiblyEnqueued, &InDesc, &OutputBuffers](UOptimusDeformerInstance* Instance)
-	{
-		Instance->OutputBuffersFromPreviousInstances = OutputBuffers;
-		OutputBuffers |= Instance->GetOutputBuffers();
-		
-		Instance->GraphSortPriorityOffset = NumComputeGraphsPossiblyEnqueued;
-		NumComputeGraphsPossiblyEnqueued += Instance->ComputeGraphExecInfos.Num();
-		
-		Instance->EnqueueWork(InDesc);
-	};
 
+
+	TArray<UOptimusDeformerInstance*> SortedInstances;
+	SortedInstances.Reserve(GuidToInstanceMap.Num());
+	
 	static const EOptimusDeformerExecutionPhase Phases[] = {
 		EOptimusDeformerExecutionPhase::BeforeDefaultDeformer,
 		EOptimusDeformerExecutionPhase::OverrideDefaultDeformer,
@@ -87,7 +76,7 @@ void UOptimusDeformerDynamicInstanceManager::EnqueueWork(FEnqueueWorkDesc const&
 
 				if (TObjectPtr<UOptimusDeformerInstance>* InstancePtr = GuidToInstanceMap.Find(LastOverrideInstanceGuid))
 				{
-					EnqueueInstance((*InstancePtr));
+					SortedInstances.Add((*InstancePtr));
 				}
 			}
 			else
@@ -98,7 +87,7 @@ void UOptimusDeformerDynamicInstanceManager::EnqueueWork(FEnqueueWorkDesc const&
 					{
 						if (TObjectPtr<UOptimusDeformerInstance>* InstancePtr = GuidToInstanceMap.Find(Guid))
 						{
-							EnqueueInstance((*InstancePtr));
+							SortedInstances.Add((*InstancePtr));
 						}
 					}
 				}
@@ -107,11 +96,44 @@ void UOptimusDeformerDynamicInstanceManager::EnqueueWork(FEnqueueWorkDesc const&
 		else if (Phase == EOptimusDeformerExecutionPhase::OverrideDefaultDeformer)
 		{
 			// Use the default instance if there is no override
-			EnqueueInstance(DefaultInstance);
+			SortedInstances.Add(DefaultInstance);
+		}
+	}
+	
+	// Making sure instances in the queue are dispatched sequentially
+	uint8 NumComputeGraphsPossiblyEnqueued = 0;
+	// Used to inform later instances whether specific buffers have valid data in them
+	EMeshDeformerOutputBuffer OutputBuffers = EMeshDeformerOutputBuffer::None;
+	
+	for (int32 Index = 0; Index < SortedInstances.Num(); Index++)
+	{
+		UOptimusDeformerInstance* Instance = SortedInstances[Index];
+		Instance->OutputBuffersFromPreviousInstances = OutputBuffers;
+		OutputBuffers |= Instance->GetOutputBuffers();
+		
+		Instance->GraphSortPriorityOffset = NumComputeGraphsPossiblyEnqueued;
+		NumComputeGraphsPossiblyEnqueued += Instance->ComputeGraphExecInfos.Num();
+
+#if WITH_EDITORONLY_DATA
+		// Readback the deformer geometry after the last deformer executes
+		if (Index == SortedInstances.Num() - 1)
+		{
+			for (const TFunction<void(const FMeshDescription&)>& Callback : OnGeometryReadbackCompletedCallbacks)
+			{
+				Instance->RequestReadbackDeformerGeometry(Callback);
+			}
 		}
+#endif // WITH_EDITORONLY_DATA
+		
+		Instance->EnqueueWork(InDesc);
 	}
+	
 
 	ExecutionQueueMap.Reset();
+#if WITH_EDITORONLY_DATA
+	// Avoid unbounded accumulation of readback requests
+	OnGeometryReadbackCompletedCallbacks.Reset();
+#endif // WITH_EDITORONLY_DATA
 }
 
 EMeshDeformerOutputBuffer UOptimusDeformerDynamicInstanceManager::GetOutputBuffers() const
@@ -121,6 +143,15 @@ EMeshDeformerOutputBuffer UOptimusDeformerDynamicInstanceManager::GetOutputBuffe
 	
 	return Result;
 }
+#if WITH_EDITORONLY_DATA
+bool UOptimusDeformerDynamicInstanceManager::RequestReadbackDeformerGeometry(TFunction<void(const FMeshDescription&)> InOnReadbackCompleted_AnyThread)
+{
+	// These requests are forwarded to the last deformer instance that runs, see EnqueueWork
+	OnGeometryReadbackCompletedCallbacks.Add(InOnReadbackCompleted_AnyThread);
+	
+	return true;
+}
+#endif // WITH_EDITORONLY_DATA
 
 UMeshDeformerInstance* UOptimusDeformerDynamicInstanceManager::GetInstanceForSourceDeformer()
 {
diff --git a/Engine/Plugins/Animation/DeformerGraph/Source/OptimusCore/Private/OptimusDeformerInstance.cpp b/Engine/Plugins/Animation/DeformerGraph/Source/OptimusCore/Private/OptimusDeformerInstance.cpp
--- a/Engine/Plugins/Animation/DeformerGraph/Source/OptimusCore/Private/OptimusDeformerInstance.cpp
+++ b/Engine/Plugins/Animation/DeformerGraph/Source/OptimusCore/Private/OptimusDeformerInstance.cpp
@@ -575,6 +575,10 @@ void UOptimusDeformerInstance::EnqueueWork(UMeshDeformerInstance::FEnqueueWorkDe
 		// If we succesfully enqueued to the Immediate group then flush all work on that group now.
 		ComputeFramework::FlushWork(InDesc.Scene, ExecutionGroupName);
 	}
+#if WITH_EDITORONLY_DATA
+	// Avoid unbounded accumulation of callbacks;
+	OnGeometryReadbackCompletedCallbacks.Reset();
+#endif // WITH_EDITORONLY_DATA
 }
 
 EMeshDeformerOutputBuffer UOptimusDeformerInstance::GetOutputBuffers() const
@@ -592,6 +596,17 @@ EMeshDeformerOutputBuffer UOptimusDeformerInstance::GetOutputBuffers() const
 	return Result;
 }
 
+#if WITH_EDITORONLY_DATA
+bool UOptimusDeformerInstance::RequestReadbackDeformerGeometry(TFunction<void(const FMeshDescription&)> InOnReadbackCompleted_AnyThread)
+{
+	// These requests will be picked up by UOptimusSkinnedMeshWriteDataProvider if there is one, ignored otherwise
+	// See UOptimusSkinnedMeshWriteDataProvider::SetupGeometryReadbackIfNeeded()
+	OnGeometryReadbackCompletedCallbacks.Add(InOnReadbackCompleted_AnyThread);
+
+	return true;
+}
+#endif // WITH_EDITORONLY_DATA
+
 namespace
 {
 	template <typename T>
diff --git a/Engine/Plugins/Animation/DeformerGraph/Source/OptimusCore/Public/OptimusDeformerDynamicInstanceManager.h b/Engine/Plugins/Animation/DeformerGraph/Source/OptimusCore/Public/OptimusDeformerDynamicInstanceManager.h
--- a/Engine/Plugins/Animation/DeformerGraph/Source/OptimusCore/Public/OptimusDeformerDynamicInstanceManager.h
+++ b/Engine/Plugins/Animation/DeformerGraph/Source/OptimusCore/Public/OptimusDeformerDynamicInstanceManager.h
@@ -39,6 +39,11 @@ public:
 	/** Return the buffers that this deformer can potentially write to */
 	EMeshDeformerOutputBuffer GetOutputBuffers() const override;
 
+#if WITH_EDITORONLY_DATA
+	/** Readback the deformed geometry as mesh description */
+	bool RequestReadbackDeformerGeometry(TFunction<void(const FMeshDescription&)> InOnReadbackCompleted_AnyThread) override;
+#endif // WITH_EDITORONLY_DATA
+	
 	/** InstanceManager is an intermediate instance, call this function to get the instance for the deformer that created this instance manager */
 	UMeshDeformerInstance* GetInstanceForSourceDeformer() override;
 
@@ -81,6 +86,11 @@ private:
 
 	/** Enqueue critical section to be thread safe */
 	FCriticalSection EnqueueCriticalSection;
+
+#if WITH_EDITORONLY_DATA
+	/** Readback requests for the current frame */
+	TArray<TFunction<void(const FMeshDescription&)>> OnGeometryReadbackCompletedCallbacks;
+#endif // WITH_EDITORONLY_DATA
 };
 
 
diff --git a/Engine/Plugins/Animation/DeformerGraph/Source/OptimusCore/Public/OptimusDeformerInstance.h b/Engine/Plugins/Animation/DeformerGraph/Source/OptimusCore/Public/OptimusDeformerInstance.h
--- a/Engine/Plugins/Animation/DeformerGraph/Source/OptimusCore/Public/OptimusDeformerInstance.h
+++ b/Engine/Plugins/Animation/DeformerGraph/Source/OptimusCore/Public/OptimusDeformerInstance.h
@@ -310,6 +310,11 @@ public:
 	/** Used to see which buffers have valid data produced by dispatched instances and are safe to access for the current instance */	
 	EMeshDeformerOutputBuffer OutputBuffersFromPreviousInstances = EMeshDeformerOutputBuffer::None;
 	
+#if WITH_EDITORONLY_DATA
+	/** Readback requests for the current frame */
+	TArray<TFunction<void(const FMeshDescription&)>> OnGeometryReadbackCompletedCallbacks;
+#endif // WITH_EDITORONLY_DATA
+	
 protected:
 	friend class UOptimusDeformerDynamicInstanceManager;
 
@@ -317,6 +322,10 @@ protected:
 	void ReleaseResources() override;
 	void EnqueueWork(FEnqueueWorkDesc const& InDesc) override;
 	EMeshDeformerOutputBuffer GetOutputBuffers() const override;
+#if WITH_EDITORONLY_DATA
+	bool RequestReadbackDeformerGeometry(TFunction<void(const FMeshDescription&)> InOnReadbackCompleted_AnyThread) override;
+#endif // WITH_EDITORONLY_DATA
+	
 	UMeshDeformerInstance* GetInstanceForSourceDeformer() override {return this;};
 
 private:
diff --git a/Engine/Source/Runtime/Engine/Classes/Animation/MeshDeformerInstance.h b/Engine/Source/Runtime/Engine/Classes/Animation/MeshDeformerInstance.h
--- a/Engine/Source/Runtime/Engine/Classes/Animation/MeshDeformerInstance.h
+++ b/Engine/Source/Runtime/Engine/Classes/Animation/MeshDeformerInstance.h
@@ -6,6 +6,9 @@
 #include "MeshDeformerInstance.generated.h"
 
 class FSceneInterface;
+#if WITH_EDITORONLY_DATA
+struct FMeshDescription;
+#endif // WITH_EDITORONLY_DATA
 
 enum class EMeshDeformerOutputBuffer : uint8
 {
@@ -81,6 +84,11 @@ public:
 	/** Return the buffers that this deformer can potentially write to */
 	ENGINE_API virtual EMeshDeformerOutputBuffer GetOutputBuffers() const PURE_VIRTUAL(, return EMeshDeformerOutputBuffer::None; );
 
+#if WITH_EDITORONLY_DATA
+	/** Reads back the deformed geometry and generates a mesh description */
+	ENGINE_API virtual bool RequestReadbackDeformerGeometry(TFunction<void(const FMeshDescription&)> InOnReadbackCompleted_AnyThread) PURE_VIRTUAL(, return false; );
+#endif // WITH_EDITORONLY_DATA
+	
 	/** Returns the specific instance that directly represents the source deformer, this is needed as a deformer may create intermediate instances that aren't
 	 * necessarily user-facing.
 	 */
diff --git a/Engine/Source/Runtime/Engine/Classes/Components/SkinnedMeshComponent.h b/Engine/Source/Runtime/Engine/Classes/Components/SkinnedMeshComponent.h
--- a/Engine/Source/Runtime/Engine/Classes/Components/SkinnedMeshComponent.h
+++ b/Engine/Source/Runtime/Engine/Classes/Components/SkinnedMeshComponent.h
@@ -36,6 +36,7 @@ class USkinnedAsset;
 struct FExternalMorphSet;
 struct FExternalMorphWeightData;
 struct FSkinWeightProfileStack;
+struct FMeshDescription;
 
 namespace Nanite
 {
@@ -345,6 +346,19 @@ public:
 	/** Max LOD at which to update or apply the MeshDeformer. */
 	ENGINE_API int32 GetMeshDeformerMaxLOD() const;
 
+#if WITH_EDITORONLY_DATA
+private:
+	TArray<TFunction<void(const FMeshDescription&)>> RenderGeometryReadbackCallbacks;
+public:
+	/**
+	 * Readback the render geometry async, a FMeshDescription will be generated on a background thread for additional operation
+	 * Note: currently it only supports mesh deformed by a MeshDeformer
+	 */
+	ENGINE_API bool RequestReadbackRenderGeometry(TFunction<void(const FMeshDescription&)> InOnReadbackCompleted_AnyThread);
+#endif // WITH_EDITORONLY_DATA
+	
+public:
+	
 	/** const getters for previous transform idea */
 	const TArray<uint8>& GetPreviousBoneVisibilityStates() const
 	{
diff --git a/Engine/Source/Runtime/Engine/Private/Components/SkinnedMeshComponent.cpp b/Engine/Source/Runtime/Engine/Private/Components/SkinnedMeshComponent.cpp
--- a/Engine/Source/Runtime/Engine/Private/Components/SkinnedMeshComponent.cpp
+++ b/Engine/Source/Runtime/Engine/Private/Components/SkinnedMeshComponent.cpp
@@ -1284,6 +1284,13 @@ void USkinnedMeshComponent::SendRenderDynamicData_Concurrent()
 		UMeshDeformerInstance* DeformerInstanceForLOD = GetMeshDeformerInstanceForLOD(UseLOD);
 		if (DeformerInstanceForLOD)
 		{
+#if WITH_EDITORONLY_DATA
+			for (TFunction<void(const FMeshDescription&)> ReadbackCallback : RenderGeometryReadbackCallbacks)
+			{
+				DeformerInstanceForLOD->RequestReadbackDeformerGeometry(ReadbackCallback);
+			}
+#endif // WITH_EDITORONLY_DATA
+			
 			UMeshDeformerInstance::FEnqueueWorkDesc Desc;
 			Desc.Scene = GetScene();
 			Desc.OwnerName = GetSkinnedAsset() != nullptr ? GetSkinnedAsset()->GetFName() : GetFName();
@@ -1295,6 +1302,11 @@ void USkinnedMeshComponent::SendRenderDynamicData_Concurrent()
 	}
 
 	Super::SendRenderDynamicData_Concurrent();
+
+#if WITH_EDITORONLY_DATA
+	// Avoid unbound accumulation of readback requests, each readback request must be picked up by the deformer instance within the same frame or ignored otherwise
+	RenderGeometryReadbackCallbacks.Reset();
+#endif // WITH_EDITORONLY_DATA 
 }
 
 void USkinnedMeshComponent::ClearMotionVector()
@@ -2591,6 +2603,27 @@ int32 USkinnedMeshComponent::GetMeshDeformerMaxLOD() const
 	return MaxLod >= 0 ? MaxLod : GetNumLODs() - 1;
 }
 
+#if WITH_EDITORONLY_DATA
+bool USkinnedMeshComponent::RequestReadbackRenderGeometry(TFunction<void(const FMeshDescription&)> InOnReadbackCompleted_AnyThread)
+{
+	if (!IsInGameThread())
+	{
+		return false;
+	}
+
+	// At the moment only support mesh deformed by mesh deformer
+	if (!HasMeshDeformer())
+	{
+		return false;
+	}
+	
+	// These readback requests will be picked up by the deformer instance during SendRenderDynamicData
+	RenderGeometryReadbackCallbacks.Add(InOnReadbackCompleted_AnyThread);
+	
+	return true;
+}
+#endif //WITH_EDITORONLY_DATA
+
 FSkeletalMeshRenderData* USkinnedMeshComponent::GetSkeletalMeshRenderData() const
 {
 	if (MeshObject)
diff --git a/Engine/Source/Runtime/Engine/Private/SkeletalMeshDeformerHelpers.cpp b/Engine/Source/Runtime/Engine/Private/SkeletalMeshDeformerHelpers.cpp
--- a/Engine/Source/Runtime/Engine/Private/SkeletalMeshDeformerHelpers.cpp
+++ b/Engine/Source/Runtime/Engine/Private/SkeletalMeshDeformerHelpers.cpp
@@ -81,6 +81,63 @@ FSkeletalMeshDeformerHelpers::FClothBuffers FSkeletalMeshDeformerHelpers::GetClo
 	return Ret;
 }
 
+FRDGBuffer* FSkeletalMeshDeformerHelpers::GetAllocatedPositionBuffer(FRDGBuilder& GraphBuilder, FSkeletalMeshObject* InMeshObject, int32 InLodIndex)
+{
+	if (InMeshObject->IsCPUSkinned())
+	{
+		return nullptr;
+	}
+
+	FSkeletalMeshObjectGPUSkin* MeshObjectGPU = static_cast<FSkeletalMeshObjectGPUSkin*>(InMeshObject);
+
+	FMeshDeformerGeometry& DeformerGeometry = MeshObjectGPU->GetDeformerGeometry(InLodIndex);
+	
+	FRDGBuffer* PositionBuffer = nullptr;
+	if (DeformerGeometry.Position)
+	{
+		PositionBuffer = GraphBuilder.RegisterExternalBuffer(DeformerGeometry.Position);
+	}
+	return PositionBuffer;
+}
+
+FRDGBuffer* FSkeletalMeshDeformerHelpers::GetAllocatedTangentBuffer(FRDGBuilder& GraphBuilder, FSkeletalMeshObject* InMeshObject, int32 InLodIndex)
+{
+	if (InMeshObject->IsCPUSkinned())
+	{
+		return nullptr;
+	}
+
+	FSkeletalMeshObjectGPUSkin* MeshObjectGPU = static_cast<FSkeletalMeshObjectGPUSkin*>(InMeshObject);
+
+	FMeshDeformerGeometry& DeformerGeometry = MeshObjectGPU->GetDeformerGeometry(InLodIndex);
+	
+	FRDGBuffer* TangentBuffer = nullptr;
+	if (DeformerGeometry.Tangent)
+	{
+		TangentBuffer = GraphBuilder.RegisterExternalBuffer(DeformerGeometry.Tangent);
+	}
+	return TangentBuffer;
+}
+
+FRDGBuffer* FSkeletalMeshDeformerHelpers::GetAllocatedColorBuffer(FRDGBuilder& GraphBuilder, FSkeletalMeshObject* InMeshObject, int32 InLodIndex)
+{
+	if (InMeshObject->IsCPUSkinned())
+	{
+		return nullptr;
+	}
+
+	FSkeletalMeshObjectGPUSkin* MeshObjectGPU = static_cast<FSkeletalMeshObjectGPUSkin*>(InMeshObject);
+
+	FMeshDeformerGeometry& DeformerGeometry = MeshObjectGPU->GetDeformerGeometry(InLodIndex);
+	
+	FRDGBuffer* ColorBuffer = nullptr;
+	if (DeformerGeometry.Color)
+	{
+		ColorBuffer = GraphBuilder.RegisterExternalBuffer(DeformerGeometry.Color);
+	}
+	return ColorBuffer;
+}
+
 FRDGBuffer* FSkeletalMeshDeformerHelpers::AllocateVertexFactoryPositionBuffer(FRDGBuilder& GraphBuilder, FSkeletalMeshObject* InMeshObject, int32 InLodIndex, bool bInLodJustChanged, TCHAR const* InBufferName)
 {
 	if (InMeshObject->IsCPUSkinned())
@@ -108,8 +165,7 @@ FRDGBuffer* FSkeletalMeshDeformerHelpers::AllocateVertexFactoryPositionBuffer(FR
 		DeformerGeometry.PrevPosition = DeformerGeometry.Position;
 		DeformerGeometry.PrevPositionSRV = DeformerGeometry.PositionSRV;
 
-		const uint32 PosBufferBytesPerElement = 4;
-		PositionBuffer = GraphBuilder.CreateBuffer(FRDGBufferDesc::CreateBufferDesc(PosBufferBytesPerElement, NumVertices * 3), InBufferName, ERDGBufferFlags::None);
+		PositionBuffer = GraphBuilder.CreateBuffer(FRDGBufferDesc::CreateBufferDesc(PosBufferBytesPerElement, NumVertices * PosBufferElementMultiplier), InBufferName, ERDGBufferFlags::None);
 		DeformerGeometry.Position = GraphBuilder.ConvertToExternalBuffer(PositionBuffer);
 		DeformerGeometry.PositionSRV = DeformerGeometry.Position->GetOrCreateSRV(GraphBuilder.RHICmdList, FRHIBufferSRVCreateInfo(PF_R32_FLOAT));
 
@@ -171,8 +227,7 @@ FRDGBuffer* FSkeletalMeshDeformerHelpers::AllocateVertexFactoryTangentBuffer(FRD
 	}
 	else
 	{
-		const uint32 TangentBufferBytesPerElement = 8;
-		TangentBuffer = GraphBuilder.CreateBuffer(FRDGBufferDesc::CreateBufferDesc(TangentBufferBytesPerElement, NumVertices * 2), InBufferName, ERDGBufferFlags::None);
+		TangentBuffer = GraphBuilder.CreateBuffer(FRDGBufferDesc::CreateBufferDesc(TangentBufferBytesPerElement, NumVertices * TangentBufferElementMultiplier), InBufferName, ERDGBufferFlags::None);
 		DeformerGeometry.Tangent = GraphBuilder.ConvertToExternalBuffer(TangentBuffer);
 		const EPixelFormat TangentsFormat = IsOpenGLPlatform(GMaxRHIShaderPlatform) ? PF_R16G16B16A16_SINT : PF_R16G16B16A16_SNORM;
 		DeformerGeometry.TangentSRV = DeformerGeometry.Tangent->GetOrCreateSRV(GraphBuilder.RHICmdList, FRHIBufferSRVCreateInfo(TangentsFormat));
@@ -207,7 +262,6 @@ FRDGBuffer* FSkeletalMeshDeformerHelpers::AllocateVertexFactoryColorBuffer(FRDGB
 	}
 	else
 	{
-		const uint32 ColorBufferBytesPerElement = 4;
 		ColorBuffer = GraphBuilder.CreateBuffer(FRDGBufferDesc::CreateBufferDesc(ColorBufferBytesPerElement, NumVertices), InBufferName, ERDGBufferFlags::None);
 		DeformerGeometry.Color = GraphBuilder.ConvertToExternalBuffer(ColorBuffer);
 		DeformerGeometry.ColorSRV = DeformerGeometry.Color->GetOrCreateSRV(GraphBuilder.RHICmdList, FRHIBufferSRVCreateInfo(PF_R8G8B8A8));
diff --git a/Engine/Source/Runtime/Engine/Public/SkeletalMeshDeformerHelpers.h b/Engine/Source/Runtime/Engine/Public/SkeletalMeshDeformerHelpers.h
--- a/Engine/Source/Runtime/Engine/Public/SkeletalMeshDeformerHelpers.h
+++ b/Engine/Source/Runtime/Engine/Public/SkeletalMeshDeformerHelpers.h
@@ -13,7 +13,12 @@ class FSkeletalMeshObject;
 class FSkeletalMeshDeformerHelpers
 {
 public:
-
+	static constexpr uint32 PosBufferBytesPerElement = 4;
+	static constexpr uint32 PosBufferElementMultiplier= 3;
+	static constexpr uint32 TangentBufferBytesPerElement = 8;
+	static constexpr uint32 TangentBufferElementMultiplier= 2;
+	static constexpr uint32 ColorBufferBytesPerElement = 4;
+	
 #pragma region GetInternals
 
 	/** Get direct access to bone matrix buffer SRV. */
@@ -48,6 +53,15 @@ public:
 		uint32 InFrameNumber,
 		bool bInPreviousFrame);
 
+	/** Returns the allocated writable position buffer if one has been allocated */	
+	ENGINE_API static FRDGBuffer* GetAllocatedPositionBuffer(FRDGBuilder& GraphBuilder, FSkeletalMeshObject* InMeshObject, int32 InLodIndex);
+
+	/** Returns the allocated writable tangent buffer if one has been allocated */	
+	ENGINE_API static FRDGBuffer* GetAllocatedTangentBuffer(FRDGBuilder& GraphBuilder, FSkeletalMeshObject* InMeshObject, int32 InLodIndex);
+
+	/** Returns the allocated writable color buffer if one has been allocated */	
+	ENGINE_API static FRDGBuffer* GetAllocatedColorBuffer(FRDGBuilder& GraphBuilder, FSkeletalMeshObject* InMeshObject, int32 InLodIndex);
+	
 #pragma endregion GetInternals
 
 #pragma region SetInternals
