[Deformer Graph] Make sure to clear stale readbacks when aborting deformer instances, which can happen during recreate render state in the first few frames of PIE, causing stale readbacks to completely block the processing of new readbacks

diff --git a/Engine/Plugins/Animation/DeformerGraph/Source/OptimusCore/Private/DataInterfaces/OptimusDataInterfaceSkinnedMeshWrite.cpp b/Engine/Plugins/Animation/DeformerGraph/Source/OptimusCore/Private/DataInterfaces/OptimusDataInterfaceSkinnedMeshWrite.cpp
--- a/Engine/Plugins/Animation/DeformerGraph/Source/OptimusCore/Private/DataInterfaces/OptimusDataInterfaceSkinnedMeshWrite.cpp
+++ b/Engine/Plugins/Animation/DeformerGraph/Source/OptimusCore/Private/DataInterfaces/OptimusDataInterfaceSkinnedMeshWrite.cpp
@@ -175,22 +175,32 @@ FComputeDataProviderRenderProxy* UOptimusSkinnedMeshWriteDataProvider::GetRender
 	return Proxy;
 }
 
-void UOptimusSkinnedMeshWriteDataProvider::SetDeformerInstance(UOptimusDeformerInstance* InInstance)
+bool UOptimusSkinnedMeshWriteDataProvider::RequestReadbackDeformerGeometry(TFunction<void(const FMeshDescription&)> InOnReadbackCompleted_AnyThread)
 {
 #if WITH_EDITORONLY_DATA
-	DeformerInstance = InInstance;
+	OnGeometryReadbackCompletedCallbacks.Add(InOnReadbackCompleted_AnyThread);
+#endif // WITH_EDITORONLY_DATA
+
+	return true;
+}
+
+void UOptimusSkinnedMeshWriteDataProvider::Abort_RenderThread()
+{
+#if WITH_EDITORONLY_DATA
+	// When deformer is aborting, we also want to clear our pending readbacks because they will no longer be fulfilled
+	GeometryReadbacks_RenderThread.Empty();
 #endif // WITH_EDITORONLY_DATA
 }
 #if WITH_EDITORONLY_DATA
 void UOptimusSkinnedMeshWriteDataProvider::SetupGeometryReadbackIfNeeded(FOptimusSkinnedMeshWriteDataProviderProxy* Proxy)
 {
 	// Process geometry readback requests if there is one
-	if (DeformerInstance->OnGeometryReadbackCompletedCallbacks.Num() > 0)
+	if (OnGeometryReadbackCompletedCallbacks.Num() > 0)
 	{
 		Proxy->RequestGeometryReadback_RenderThread = [
 			WeakThis = TWeakObjectPtr(this),
 			FrameNumber = GFrameNumber,
-			GeometryReadbackCallbacks = DeformerInstance->OnGeometryReadbackCompletedCallbacks
+			GeometryReadbackCallbacks = OnGeometryReadbackCompletedCallbacks
 		](
 			TArray<FComputeDataProviderRenderProxy::FReadbackData>& OutBufferReadbacks,
 			FSkeletalMeshObject* MeshObject,
@@ -245,7 +255,7 @@ void UOptimusSkinnedMeshWriteDataProvider::SetupGeometryReadbackIfNeeded(FOptimu
 			}
 		};
 
-		DeformerInstance->OnGeometryReadbackCompletedCallbacks.Reset();
+		OnGeometryReadbackCompletedCallbacks.Reset();
 	}	
 }
 #endif // WITH_EDITORONLY_DATA
@@ -260,15 +270,9 @@ void UOptimusSkinnedMeshWriteDataProvider::ProcessCompletedGeometryReadback_Rend
 		TSharedPtr<FGeometryReadback> Readback = *GeometryReadbacks_RenderThread.Peek();
 
 		// Stop processing if a buffer is not ready
-		if (Readback->Position.bShouldReadback && Readback->Position.ReadbackData.Num() == 0)
-		{
-			break;
-		}
-		if (Readback->Tangent.bShouldReadback && Readback->Tangent.ReadbackData.Num() == 0)
-		{
-			break;
-		}
-		if (Readback->Color.bShouldReadback && Readback->Color.ReadbackData.Num() == 0)
+		if ((Readback->Position.bShouldReadback && Readback->Position.ReadbackData.Num() == 0) ||
+			(Readback->Tangent.bShouldReadback && Readback->Tangent.ReadbackData.Num() == 0) ||
+			(Readback->Color.bShouldReadback && Readback->Color.ReadbackData.Num() == 0))
 		{
 			break;
 		}
diff --git a/Engine/Plugins/Animation/DeformerGraph/Source/OptimusCore/Private/DataInterfaces/OptimusDataInterfaceSkinnedMeshWrite.h b/Engine/Plugins/Animation/DeformerGraph/Source/OptimusCore/Private/DataInterfaces/OptimusDataInterfaceSkinnedMeshWrite.h
--- a/Engine/Plugins/Animation/DeformerGraph/Source/OptimusCore/Private/DataInterfaces/OptimusDataInterfaceSkinnedMeshWrite.h
+++ b/Engine/Plugins/Animation/DeformerGraph/Source/OptimusCore/Private/DataInterfaces/OptimusDataInterfaceSkinnedMeshWrite.h
@@ -2,7 +2,7 @@
 
 #pragma once
 
-#include "IOptimusDeformerInstanceAccessor.h"
+#include "IOptimusDeformerGeometryReadbackProvider.h"
 #include "IOptimusOutputBufferWriter.h"
 #include "OptimusComputeDataInterface.h"
 #include "ComputeFramework/ComputeDataProvider.h"
@@ -19,10 +19,6 @@ class FSkeletalMeshObject;
 class FSkinedMeshWriteDataInterfaceParameters;
 class USkinnedMeshComponent;
 
-#if WITH_EDITORONLY_DATA
-struct FMeshDescription;
-#endif // WITH_EDITORONLY_DATA
-
 /** Compute Framework Data Interface for writing skinned mesh. */
 UCLASS(Category = ComputeFramework)
 class OPTIMUSCORE_API UOptimusSkinnedMeshWriteDataInterface :
@@ -65,7 +61,7 @@ private:
 UCLASS(BlueprintType, editinlinenew, Category = ComputeFramework)
 class UOptimusSkinnedMeshWriteDataProvider :
 	public UComputeDataProvider,
-	public IOptimusDeformerInstanceAccessor
+	public IOptimusDeformerGeometryReadbackProvider
 {
 	GENERATED_BODY()
 
@@ -78,18 +74,21 @@ public:
 	//~ Begin UComputeDataProvider Interface
 	FComputeDataProviderRenderProxy* GetRenderProxy() override;
 	//~ End UComputeDataProvider Interface
-
-	//~ Begin IOptimusDeformerInstanceAccessor Interface
-	void SetDeformerInstance(UOptimusDeformerInstance* InInstance) override;
-	//~ End IOptimusDeformerInstanceAccessor Interface
-
+	
+	
+	//~ Begin IOptimusDeformerGeometryReadbackProvider Interface
+	bool RequestReadbackDeformerGeometry(TFunction<void(const FMeshDescription&)> InOnReadbackCompleted_AnyThread) override;
+	void Abort_RenderThread() override;
+	//~ End IOptimusDeformerGeometryReadbackProvider Interface
+	
+	
 private:
 	// Served as persistent storage for the provider proxy, should not be used by the data provider itself
 	int32 LastLodIndexCachedByRenderProxy = 0;
 
 #if WITH_EDITORONLY_DATA
-	UPROPERTY()
-	TObjectPtr<UOptimusDeformerInstance> DeformerInstance = nullptr;
+	/** Readback requests for the current frame */
+	TArray<TFunction<void(const FMeshDescription&)>> OnGeometryReadbackCompletedCallbacks;
 
 	void SetupGeometryReadbackIfNeeded(class FOptimusSkinnedMeshWriteDataProviderProxy* Proxy);
 	void ProcessCompletedGeometryReadback_RenderThread();
diff --git a/Engine/Plugins/Animation/DeformerGraph/Source/OptimusCore/Private/OptimusDeformerInstance.cpp b/Engine/Plugins/Animation/DeformerGraph/Source/OptimusCore/Private/OptimusDeformerInstance.cpp
--- a/Engine/Plugins/Animation/DeformerGraph/Source/OptimusCore/Private/OptimusDeformerInstance.cpp
+++ b/Engine/Plugins/Animation/DeformerGraph/Source/OptimusCore/Private/OptimusDeformerInstance.cpp
@@ -5,6 +5,7 @@
 #include "Components/MeshComponent.h"
 #include "ComputeFramework/ComputeFramework.h"
 #include "ComputeWorkerInterface.h"
+#include "IOptimusDeformerGeometryReadbackProvider.h"
 #include "IOptimusPersistentBufferProvider.h"
 #include "IOptimusValueProvider.h"
 #include "DataInterfaces/OptimusDataInterfaceGraph.h"
@@ -432,6 +433,15 @@ void UOptimusDeformerInstance::SetupFromDeformer(UOptimusDeformer* InDeformer)
 			{
 				InstanceAccessor->SetDeformerInstance(this);
 			}
+#if WITH_EDITORONLY_DATA
+			if (Info.GraphType == EOptimusNodeGraphType::Update)
+			{
+				if (IOptimusDeformerGeometryReadbackProvider* GeometryReadbackProvider = Cast<IOptimusDeformerGeometryReadbackProvider>(DataProvider))
+				{
+					WeakGeometryReadbackProvider = GeometryReadbackProvider;
+				}
+			}
+#endif
 		}
 
 		// Schedule the setup graph to run.
@@ -488,9 +498,9 @@ void UOptimusDeformerInstance::AllocateResources()
 
 void UOptimusDeformerInstance::ReleaseResources()
 {
-	if (Scene || BufferPool)
+	if (Scene || BufferPool || WeakGeometryReadbackProvider.IsValid())
 	{
-		ENQUEUE_RENDER_COMMAND(OptimusReleaseResources)([BufferPool=MoveTemp(BufferPool), Scene = Scene, OwnerPointer = this] (FRHICommandListImmediate& InCmdList)
+		ENQUEUE_RENDER_COMMAND(OptimusReleaseResources)([BufferPool=MoveTemp(BufferPool), Scene = Scene, OwnerPointer = this, WeakGeometryReadbackProvider = WeakGeometryReadbackProvider] (FRHICommandListImmediate& InCmdList)
 		{
 			if (Scene)
 			{
@@ -501,6 +511,11 @@ void UOptimusDeformerInstance::ReleaseResources()
 			{
 				BufferPool->ReleaseResources();
 			}
+
+			if (WeakGeometryReadbackProvider.IsValid())
+			{
+				WeakGeometryReadbackProvider->Abort_RenderThread();
+			}
 		});
 	}
 }
@@ -575,10 +590,6 @@ void UOptimusDeformerInstance::EnqueueWork(UMeshDeformerInstance::FEnqueueWorkDe
 		// If we succesfully enqueued to the Immediate group then flush all work on that group now.
 		ComputeFramework::FlushWork(InDesc.Scene, ExecutionGroupName);
 	}
-#if WITH_EDITORONLY_DATA
-	// Avoid unbounded accumulation of callbacks;
-	OnGeometryReadbackCompletedCallbacks.Reset();
-#endif // WITH_EDITORONLY_DATA
 }
 
 EMeshDeformerOutputBuffer UOptimusDeformerInstance::GetOutputBuffers() const
@@ -599,11 +610,13 @@ EMeshDeformerOutputBuffer UOptimusDeformerInstance::GetOutputBuffers() const
 #if WITH_EDITORONLY_DATA
 bool UOptimusDeformerInstance::RequestReadbackDeformerGeometry(TFunction<void(const FMeshDescription&)> InOnReadbackCompleted_AnyThread)
 {
-	// These requests will be picked up by UOptimusSkinnedMeshWriteDataProvider if there is one, ignored otherwise
-	// See UOptimusSkinnedMeshWriteDataProvider::SetupGeometryReadbackIfNeeded()
-	OnGeometryReadbackCompletedCallbacks.Add(InOnReadbackCompleted_AnyThread);
+	if (WeakGeometryReadbackProvider.IsValid())
+	{
+		WeakGeometryReadbackProvider->RequestReadbackDeformerGeometry(InOnReadbackCompleted_AnyThread);
+		return true;
+	}
 
-	return true;
+	return false;
 }
 #endif // WITH_EDITORONLY_DATA
 
diff --git a/Engine/Plugins/Animation/DeformerGraph/Source/OptimusCore/Public/IOptimusDeformerGeometryReadbackProvider.h b/Engine/Plugins/Animation/DeformerGraph/Source/OptimusCore/Public/IOptimusDeformerGeometryReadbackProvider.h
new file mode 100644
--- /dev/null
+++ b/Engine/Plugins/Animation/DeformerGraph/Source/OptimusCore/Public/IOptimusDeformerGeometryReadbackProvider.h
@@ -0,0 +1,28 @@
+﻿// Copyright Epic Games, Inc. All Rights Reserved.
+
+#pragma once
+
+#include "UObject/Interface.h"
+
+#include "IOptimusDeformerGeometryReadbackProvider.generated.h"
+
+
+class UOptimusDeformerInstance;
+struct FMeshDescription;
+
+UINTERFACE()
+class OPTIMUSCORE_API UOptimusDeformerGeometryReadbackProvider :
+	public UInterface
+{
+	GENERATED_BODY()
+};
+
+
+class IOptimusDeformerGeometryReadbackProvider
+{
+	GENERATED_BODY()
+
+public:
+	virtual bool RequestReadbackDeformerGeometry(TFunction<void(const FMeshDescription&)> InOnReadbackCompleted_AnyThread) = 0;
+	virtual void Abort_RenderThread() = 0;
+};
diff --git a/Engine/Plugins/Animation/DeformerGraph/Source/OptimusCore/Public/OptimusDeformerInstance.h b/Engine/Plugins/Animation/DeformerGraph/Source/OptimusCore/Public/OptimusDeformerInstance.h
--- a/Engine/Plugins/Animation/DeformerGraph/Source/OptimusCore/Public/OptimusDeformerInstance.h
+++ b/Engine/Plugins/Animation/DeformerGraph/Source/OptimusCore/Public/OptimusDeformerInstance.h
@@ -6,6 +6,8 @@
 #include "ComputeFramework/ComputeGraphInstance.h"
 #include "Engine/EngineTypes.h"
 #include "OptimusValue.h"
+#include "UObject/WeakInterfacePtr.h"
+#include "IOptimusDeformerGeometryReadbackProvider.h"
 
 #include "OptimusDeformerInstance.generated.h"
 
@@ -309,12 +311,6 @@ public:
 	
 	/** Used to see which buffers have valid data produced by dispatched instances and are safe to access for the current instance */	
 	EMeshDeformerOutputBuffer OutputBuffersFromPreviousInstances = EMeshDeformerOutputBuffer::None;
-	
-#if WITH_EDITORONLY_DATA
-	/** Readback requests for the current frame */
-	TArray<TFunction<void(const FMeshDescription&)>> OnGeometryReadbackCompletedCallbacks;
-#endif // WITH_EDITORONLY_DATA
-	
 protected:
 	friend class UOptimusDeformerDynamicInstanceManager;
 
@@ -356,5 +352,8 @@ private:
 	FOptimusPersistentBufferPoolPtr BufferPool;
 	FSceneInterface* Scene = nullptr;
 
+	// Data provider responsible for reading back the final deformed geometry after the deformer runs, only assigned when WITH_EDITORONLY_DATA
+	TWeakInterfacePtr<IOptimusDeformerGeometryReadbackProvider> WeakGeometryReadbackProvider = nullptr;
+
 	bool bCanBeActive = true;
 };
