[Deformer Graph] Added support for comment nodes

diff --git a/Engine/Plugins/Animation/DeformerGraph/Source/OptimusCore/Private/Actions/OptimusNodeActions.cpp b/Engine/Plugins/Animation/DeformerGraph/Source/OptimusCore/Private/Actions/OptimusNodeActions.cpp
--- a/Engine/Plugins/Animation/DeformerGraph/Source/OptimusCore/Private/Actions/OptimusNodeActions.cpp
+++ b/Engine/Plugins/Animation/DeformerGraph/Source/OptimusCore/Private/Actions/OptimusNodeActions.cpp
@@ -9,6 +9,7 @@
 #include "OptimusNodePin.h"
 #include "IOptimusNodeAdderPinProvider.h"
 #include "OptimusDataTypeRegistry.h"
+#include "OptimusNode_Comment.h"
 
 #include UE_INLINE_GENERATED_CPP_BY_NAME(OptimusNodeActions)
 
@@ -87,6 +88,41 @@ bool FOptimusNodeAction_MoveNode::Undo(IOptimusPathResolver* InRoot)
 }
 
 
+
+FOptimusCommentNodeAction_ResizeNode::FOptimusCommentNodeAction_ResizeNode(
+	UOptimusNode_Comment* InNode,
+	const FVector2D& InNewSize)
+{
+	NodePath = InNode->GetNodePath();
+	OldSize = InNode->GetSize();	
+	NewSize = InNewSize;
+	
+	SetTitlef(TEXT("Resize Comment Node"));
+}
+
+bool FOptimusCommentNodeAction_ResizeNode::Do(IOptimusPathResolver* InRoot)
+{
+	UOptimusNode_Comment* Node = Cast<UOptimusNode_Comment>(InRoot->ResolveNodePath(NodePath));
+	if (!Node)
+	{
+		return false;
+	}
+
+	return Node->SetSizeDirect(NewSize);
+}
+
+bool FOptimusCommentNodeAction_ResizeNode::Undo(IOptimusPathResolver* InRoot)
+{
+	UOptimusNode_Comment* Node = Cast<UOptimusNode_Comment>(InRoot->ResolveNodePath(NodePath));
+	if (!Node)
+	{
+		return false;
+	}
+
+	return Node->SetSizeDirect(OldSize);
+}
+
+
 FOptimusNodeAction_SetPinValue::FOptimusNodeAction_SetPinValue(
 	UOptimusNodePin* InPin, 
 	const FString& InNewValue
diff --git a/Engine/Plugins/Animation/DeformerGraph/Source/OptimusCore/Private/Actions/OptimusNodeActions.h b/Engine/Plugins/Animation/DeformerGraph/Source/OptimusCore/Private/Actions/OptimusNodeActions.h
--- a/Engine/Plugins/Animation/DeformerGraph/Source/OptimusCore/Private/Actions/OptimusNodeActions.h
+++ b/Engine/Plugins/Animation/DeformerGraph/Source/OptimusCore/Private/Actions/OptimusNodeActions.h
@@ -10,6 +10,7 @@
 class UOptimusNode;
 class UOptimusNodePin;
 class IOptimusNodeAdderPinProvider;
+class UOptimusNode_Comment;
 
 USTRUCT()
 struct FOptimusNodeAction_RenameNode : 
@@ -70,6 +71,32 @@ private:
 	FVector2D OldPosition;
 };
 
+USTRUCT()
+struct FOptimusCommentNodeAction_ResizeNode :
+	public FOptimusAction
+{
+	GENERATED_BODY()
+
+public:
+	FOptimusCommentNodeAction_ResizeNode() = default;
+
+	FOptimusCommentNodeAction_ResizeNode(
+		UOptimusNode_Comment* InNode,
+		const FVector2D& InNewSize
+	);
+
+protected:
+	bool Do(IOptimusPathResolver* InRoot) override;
+	bool Undo(IOptimusPathResolver* InRoot) override;
+
+private:
+	// The path of the node to be resized
+	FString NodePath;
+
+	FVector2D OldSize;
+
+	FVector2D NewSize;
+};
 
 USTRUCT()
 struct FOptimusNodeAction_SetPinValue :
diff --git a/Engine/Plugins/Animation/DeformerGraph/Source/OptimusCore/Private/OptimusNodeGraph.cpp b/Engine/Plugins/Animation/DeformerGraph/Source/OptimusCore/Private/OptimusNodeGraph.cpp
--- a/Engine/Plugins/Animation/DeformerGraph/Source/OptimusCore/Private/OptimusNodeGraph.cpp
+++ b/Engine/Plugins/Animation/DeformerGraph/Source/OptimusCore/Private/OptimusNodeGraph.cpp
@@ -34,6 +34,7 @@
 #include "OptimusFunctionNodeGraph.h"
 #include "OptimusFunctionNodeGraphHeader.h"
 #include "OptimusNodeSubGraph.h"
+#include "OptimusNode_Comment.h"
 #include "Nodes/OptimusNode_CustomComputeKernel.h"
 #include "Nodes/OptimusNode_FunctionReference.h"
 #include "Nodes/OptimusNode_GraphTerminal.h"
@@ -590,6 +591,29 @@ TArray<UOptimusNode*> UOptimusNodeGraph::AddLoopTerminalNodes(const FVector2D& I
 	return AddedNodes;
 }
 
+UOptimusNode* UOptimusNodeGraph::AddCommentNode(const FVector2D& InPosition, const FVector2D& InSize,const FLinearColor& InColor)
+{
+	return AddCommentNode(InPosition, InSize, InColor, false);
+	
+}
+
+UOptimusNode* UOptimusNodeGraph::AddCommentNode(const FVector2D& InPosition, const FVector2D& InSize, const FLinearColor& InColor,
+	bool bCreatedFromUI)
+{
+	return AddNodeInternal(UOptimusNode_Comment::StaticClass(), InPosition,
+		[InColor, InSize, bCreatedFromUI](UOptimusNode *InNode)
+		{
+			UOptimusNode_Comment* CommentNode = CastChecked<UOptimusNode_Comment>(InNode);
+			CommentNode->SetSizeDirect(InSize);
+			CommentNode->CommentColor = InColor;
+			CommentNode->FontSize = 18;
+			CommentNode->Comment = TEXT("Comment");
+#if WITH_EDITOR
+			CommentNode->bCreatedFromUI = bCreatedFromUI;
+#endif
+		});	
+}
+
 UOptimusNode* UOptimusNodeGraph::AddFunctionReferenceNode(TSoftObjectPtr<UOptimusFunctionNodeGraph> InFunctionGraph, const FVector2D& InPosition)
 {
 	return AddNodeInternal(UOptimusNode_FunctionReference::StaticClass(), InPosition,
diff --git a/Engine/Plugins/Animation/DeformerGraph/Source/OptimusCore/Private/OptimusNode_Comment.cpp b/Engine/Plugins/Animation/DeformerGraph/Source/OptimusCore/Private/OptimusNode_Comment.cpp
new file mode 100644
--- /dev/null
+++ b/Engine/Plugins/Animation/DeformerGraph/Source/OptimusCore/Private/OptimusNode_Comment.cpp
@@ -0,0 +1,52 @@
+// Copyright Epic Games, Inc. All Rights Reserved.
+
+#include "OptimusNode_Comment.h"
+
+#include "OptimusActionStack.h"
+#include "OptimusNodeGraph.h"
+#include "Actions/OptimusNodeActions.h"
+
+#include UE_INLINE_GENERATED_CPP_BY_NAME(OptimusNode_Comment)
+
+#define LOCTEXT_NAMESPACE "OptimusNode_Comment"
+
+#if WITH_EDITOR
+void UOptimusNode_Comment::PostEditChangeProperty(FPropertyChangedEvent& PropertyChangedEvent)
+{
+	Super::PostEditChangeProperty(PropertyChangedEvent);
+
+	OnPropertyChangedDelegate.ExecuteIfBound(this);
+}
+#endif
+
+bool UOptimusNode_Comment::SetSize(const FVector2D& InSize)
+{
+	return GetActionStack()->RunAction<FOptimusCommentNodeAction_ResizeNode>(this, InSize);
+}
+
+bool UOptimusNode_Comment::SetSizeDirect(const FVector2D& InNewSize)
+{
+	if (Size == InNewSize || InNewSize.ContainsNaN())
+	{
+		return false;
+	}
+	
+	Size = InNewSize;
+	
+	OnPropertyChangedDelegate.ExecuteIfBound(this);
+	
+	return true;	
+}
+
+void UOptimusNode_Comment::SetComment(const FString& InNewComment)
+{
+	// Caller should avoid create a transaction if not changing anything
+	if (ensure(Comment != InNewComment))
+	{
+		Modify();
+		Comment = InNewComment;
+		OnPropertyChangedDelegate.ExecuteIfBound(this);
+	}
+}
+
+#undef LOCTEXT_NAMESPACE
diff --git a/Engine/Plugins/Animation/DeformerGraph/Source/OptimusCore/Public/OptimusNode.h b/Engine/Plugins/Animation/DeformerGraph/Source/OptimusCore/Public/OptimusNode.h
--- a/Engine/Plugins/Animation/DeformerGraph/Source/OptimusCore/Public/OptimusNode.h
+++ b/Engine/Plugins/Animation/DeformerGraph/Source/OptimusCore/Public/OptimusNode.h
@@ -209,6 +209,7 @@ public:
 	virtual void PostLoadNodeSpecificData();
 #if WITH_EDITOR
 	void PostEditChangeProperty(FPropertyChangedEvent& PropertyChangedEvent) override;
+	bool IsCreatedFromUI() const {return bCreatedFromUI;}
 #endif
 	
 protected:
@@ -454,7 +455,12 @@ private:
 
 	// A sentinel to indicate we're doing node construction.
 	bool bConstructingNode = false;
-
+	
+#if WITH_EDITOR
+	// Optionally one can mark this node as created from UI during node creation for the Editor
+	// to do special things when the node is spawned
+	bool bCreatedFromUI = false;
+#endif
 	/// Cached pin lookups
 	mutable TMap<TArray<FName>, UOptimusNodePin*> CachedPinLookup;
 };
diff --git a/Engine/Plugins/Animation/DeformerGraph/Source/OptimusCore/Public/OptimusNodeGraph.h b/Engine/Plugins/Animation/DeformerGraph/Source/OptimusCore/Public/OptimusNodeGraph.h
--- a/Engine/Plugins/Animation/DeformerGraph/Source/OptimusCore/Public/OptimusNodeGraph.h
+++ b/Engine/Plugins/Animation/DeformerGraph/Source/OptimusCore/Public/OptimusNodeGraph.h
@@ -153,6 +153,20 @@ public:
 		const FVector2D& InPosition
 	);
 
+	UFUNCTION(BlueprintCallable, Category = OptimusNodeGraph)
+	UOptimusNode* AddCommentNode(
+		const FVector2D& InPosition,
+		const FVector2D& InSize,
+		const FLinearColor& InColor
+	);
+	
+	UOptimusNode* AddCommentNode(
+    		const FVector2D& InPosition,
+    		const FVector2D& InSize,
+    		const FLinearColor& InColor,
+    		bool bCreatedFromUI
+    );
+
 	UFUNCTION(BlueprintCallable, Category = OptimusNodeGraph)
 	UOptimusNode* AddFunctionReferenceNode(
 		TSoftObjectPtr<UOptimusFunctionNodeGraph> InFunctionGraph,
diff --git a/Engine/Plugins/Animation/DeformerGraph/Source/OptimusCore/Public/OptimusNode_Comment.h b/Engine/Plugins/Animation/DeformerGraph/Source/OptimusCore/Public/OptimusNode_Comment.h
new file mode 100644
--- /dev/null
+++ b/Engine/Plugins/Animation/DeformerGraph/Source/OptimusCore/Public/OptimusNode_Comment.h
@@ -0,0 +1,64 @@
+// Copyright Epic Games, Inc. All Rights Reserved.
+
+#pragma once
+
+#include "OptimusNode.h"
+
+#include "OptimusNode_Comment.generated.h"
+
+class UOptimusNode_Comment;
+
+DECLARE_DELEGATE_OneParam(FOnOptimusCommentNodePropertyChanged, UOptimusNode_Comment*);
+
+UCLASS(Hidden)
+class OPTIMUSCORE_API  UOptimusNode_Comment : 
+	public UOptimusNode
+{
+	GENERATED_BODY()
+
+
+public:
+	
+#if WITH_EDITOR
+	void PostEditChangeProperty(FPropertyChangedEvent& PropertyChangedEvent) override;
+#endif // WITH_EDITOR
+
+	FName GetNodeCategory() const override {return NAME_None;}
+
+	bool SetSize(const FVector2D& InSize);
+	const FVector2D& GetSize() {return Size;}
+	void SetComment(const FString& InNewComment);
+	const FString& GetComment() {return Comment;}
+	
+	FOnOptimusCommentNodePropertyChanged& GetOnPropertyChanged() {return OnPropertyChangedDelegate; };
+
+	/** Color to style comment with */
+	UPROPERTY(EditAnywhere, Category=Comment)
+	FLinearColor CommentColor = FLinearColor::White;
+
+	/** Size of the text in the comment box */
+	UPROPERTY(EditAnywhere, Category=Comment, meta=(ClampMin=1, ClampMax=1000))
+	int32 FontSize = 18;
+
+	/** Comment to show */
+	UPROPERTY(EditAnywhere, Category=Comment)
+	FString Comment = TEXT("Comment");
+
+	/** Whether to show a zoom-invariant comment bubble when zoomed out (making the comment readable at any distance). */
+	UPROPERTY(EditAnywhere, Category=Comment, meta=(DisplayName="Show Bubble When Zoomed"))
+	bool bBubbleVisible = false;
+
+	/** Whether to use Comment Color to color the background of the comment bubble shown when zoomed out. */
+	UPROPERTY(EditAnywhere, Category=Comment, meta=(DisplayName="Color Bubble", EditCondition=bBubbleVisible))
+	bool bColorBubble = false;
+	
+private:
+	friend class UOptimusNodeGraph;
+	friend struct FOptimusCommentNodeAction_ResizeNode;
+	bool SetSizeDirect(const FVector2D& InNewSize);	
+
+	UPROPERTY(NonTransactional)
+	FVector2D Size = FVector2D(400, 100);	
+	
+	FOnOptimusCommentNodePropertyChanged OnPropertyChangedDelegate;
+};
diff --git a/Engine/Plugins/Animation/DeformerGraph/Source/OptimusEditor/Private/OptimusEditor.cpp b/Engine/Plugins/Animation/DeformerGraph/Source/OptimusEditor/Private/OptimusEditor.cpp
--- a/Engine/Plugins/Animation/DeformerGraph/Source/OptimusEditor/Private/OptimusEditor.cpp
+++ b/Engine/Plugins/Animation/DeformerGraph/Source/OptimusEditor/Private/OptimusEditor.cpp
@@ -17,6 +17,7 @@
 #include "OptimusCoreNotify.h"
 #include "OptimusDeformer.h"
 #include "OptimusNode.h"
+#include "OptimusNode_Comment.h"
 #include "OptimusNodeGraph.h"
 #include "OptimusFunctionNodeGraph.h"
 
@@ -53,7 +54,11 @@
 #include "OptimusEditorStyle.h"
 #include "Toolkits/ToolkitManager.h"
 #include "EditorViewportClient.h"
-
+#include "GraphEditorActions.h"
+#include "OptimusEditorGraphNode_Comment.h"
+#include "OptimusEditorGraphSchemaActions.h"
+#include "OptimusEditorHelpers.h"
+#include "ScopedTransaction.h"
 
 #define LOCTEXT_NAMESPACE "OptimusEditor"
 
@@ -114,6 +119,9 @@ void FOptimusEditor::Construct(
 		DeformerObject, UOptimusEditorGraph::StaticClass(), NAME_None, 
 		RF_Transient);
 	EditorGraph->Schema = UOptimusEditorGraphSchema::StaticClass();
+	
+	// Used to query the size of slate node widgets
+	EditorGraph->WeakEditor = SharedThis(this);
 
 	constexpr bool bCreateDefaultStandaloneMenu = true;
 	constexpr bool bCreateDefaultToolbar = true;
@@ -448,12 +456,11 @@ bool FOptimusEditor::CanSelectAllNodes() const
 void FOptimusEditor::DeleteSelectedNodes()
 {
 	TArray<UOptimusNode*> NodesToDelete;
-	for (UObject* Object : GraphEditorWidget->GetSelectedNodes())
+	for (UOptimusNode* ModelNode : GetSelectedModelNodes())
 	{
-		UOptimusEditorGraphNode* GraphNode = Cast<UOptimusEditorGraphNode>(Object);
-		if (GraphNode && GraphNode->CanUserDeleteNode())
+		if (ModelNode->CanUserDeleteNode())
 		{
-			NodesToDelete.Add(GraphNode->ModelNode);
+			NodesToDelete.Add(ModelNode);
 		}
 	}
 
@@ -791,6 +798,14 @@ bool FOptimusEditor::CanConvertToSubGraph() const
 	return !ModelNodes.IsEmpty();
 }
 
+void FOptimusEditor::OnCreateComment()
+{
+	// In case there aren't any nodes selected, create a comment node at current mouse position
+	const FVector2D CurrentMouseGraphPosition = GraphEditorWidget->GetPasteLocation();
+	FOptimusGraphSchemaAction_NewCommentNode NewCommentNodeAction;
+	NewCommentNodeAction.PerformAction(EditorGraph, nullptr, CurrentMouseGraphPosition, true);
+}
+
 void FOptimusEditor::SubscribeToGraphNotifies(UOptimusNodeGraph* Graph, const FOnGraphNotified& OnGraphNotifiedDelegate)
 {
 	if (IsValid(Graph))
@@ -823,14 +838,14 @@ void FOptimusEditor::UnsubscribeToGraphNotifies(FDelegateHandle DelegateHandle)
 void FOptimusEditor::OnSelectedNodesChanged(const TSet<UObject*>& NewSelection)
 {
 	TArray<TWeakObjectPtr<UObject>> SelectedObjects;
-	TSet<UOptimusEditorGraphNode*> SelectedNodes;
+	TSet<UEdGraphNode*> SelectedNodes;
 
 	for (UObject* Object : NewSelection)
 	{
 		UObject* ObjectToShow;
-		if (UOptimusEditorGraphNode* GraphNode = Cast<UOptimusEditorGraphNode>(Object))
+		if (UEdGraphNode* GraphNode = Cast<UEdGraphNode>(Object))
 		{
-			ObjectToShow = GraphNode->ModelNode;
+			ObjectToShow = OptimusEditor::FindModelNodeFromGraphNode(GraphNode);
 			SelectedNodes.Add(GraphNode);
 		}
 		else
@@ -879,34 +894,32 @@ void FOptimusEditor::OnSelectedNodesChanged(const TSet<UObject*>& NewSelection)
 
 void FOptimusEditor::OnNodeDoubleClicked(class UEdGraphNode* Node)
 {
-	if (UOptimusEditorGraphNode* GraphNode = Cast<UOptimusEditorGraphNode>(Node))
+	UOptimusNode* ModelNode = OptimusEditor::FindModelNodeFromGraphNode(Node);
+	if (IOptimusNodeGraphProvider* GraphProvider = Cast<IOptimusNodeGraphProvider>(ModelNode))
 	{
-		if (IOptimusNodeGraphProvider* GraphProvider = Cast<IOptimusNodeGraphProvider>(GraphNode->ModelNode))
-		{
-			UOptimusNodeGraph* GraphToShow = GraphProvider->GetNodeGraphToShow();
-			UOptimusNodeGraph* CurrentGraph = GraphNode->ModelNode->GetOwningGraph();
+		UOptimusNodeGraph* GraphToShow = GraphProvider->GetNodeGraphToShow();
+		UOptimusNodeGraph* CurrentGraph = ModelNode->GetOwningGraph();
 
-			if (GraphToShow)
+		if (GraphToShow)
+		{
+			if (CurrentGraph->GetCollectionRoot() == GraphToShow->GetCollectionRoot())
 			{
-				if (CurrentGraph->GetCollectionRoot() == GraphToShow->GetCollectionRoot())
-				{
-					SetEditGraph(GraphToShow);
-				}
-				else if (UOptimusDeformer* Deformer = Cast<UOptimusDeformer>(GraphToShow->GetCollectionRoot()))
-				{
-					GEditor->GetEditorSubsystem<UAssetEditorSubsystem>()->OpenEditorForAsset(Deformer);
+				SetEditGraph(GraphToShow);
+			}
+			else if (UOptimusDeformer* Deformer = Cast<UOptimusDeformer>(GraphToShow->GetCollectionRoot()))
+			{
+				GEditor->GetEditorSubsystem<UAssetEditorSubsystem>()->OpenEditorForAsset(Deformer);
 
-					TSharedPtr< IToolkit > FoundAssetEditor = FToolkitManager::Get().FindEditorForAsset(Deformer);
+				TSharedPtr< IToolkit > FoundAssetEditor = FToolkitManager::Get().FindEditorForAsset(Deformer);
+				
+				if (FoundAssetEditor.IsValid())
+				{
+					FoundAssetEditor->BringToolkitToFront();
 					
-					if (FoundAssetEditor.IsValid())
-					{
-						FoundAssetEditor->BringToolkitToFront();
-						
-						TSharedPtr<FOptimusEditor> OptimusEditorForExternalGraph = StaticCastSharedPtr<FOptimusEditor>(FoundAssetEditor);
+					TSharedPtr<FOptimusEditor> OptimusEditorForExternalGraph = StaticCastSharedPtr<FOptimusEditor>(FoundAssetEditor);
 
-						OptimusEditorForExternalGraph->SetEditGraph(GraphToShow);
-					}	
-				}
+					OptimusEditorForExternalGraph->SetEditGraph(GraphToShow);
+				}	
 			}
 		}
 	}
@@ -915,12 +928,24 @@ void FOptimusEditor::OnNodeDoubleClicked(class UEdGraphNode* Node)
 
 void FOptimusEditor::OnNodeTitleCommitted(const FText& NewText, ETextCommit::Type CommitInfo, UEdGraphNode* NodeBeingChanged)
 {
-
+	if (UOptimusNode_Comment* CommentNode = Cast<UOptimusNode_Comment>(OptimusEditor::FindModelNodeFromGraphNode(NodeBeingChanged)))
+	{
+		FString NewComment = NewText.ToString();
+		if (CommentNode->GetComment() != NewComment)
+		{
+			FScopedTransaction ScopedTransaction(LOCTEXT("OptimusEditorChangeCommentNodeComment_Transaction", "Change Comment Text"));
+			CommentNode->SetComment(NewText.ToString());
+		}
+	}
 }
 
 
 bool FOptimusEditor::OnVerifyNodeTextCommit(const FText& NewText, UEdGraphNode* NodeBeingChanged, FText& OutErrorMessage)
 {
+	if (UOptimusEditorGraphNode_Comment* CommentGraphNode = Cast<UOptimusEditorGraphNode_Comment>(NodeBeingChanged))
+	{
+		return true;
+	}
 	return false;
 }
 
@@ -1063,7 +1088,7 @@ void FOptimusEditor::HandleMessageTokenClicked(const TSharedRef<IMessageToken>&
 			}
 
 			// Highlight the selection in the graph.
-			UOptimusEditorGraphNode* GraphNode = EditorGraph->FindGraphNodeFromModelNode(Node);
+			UEdGraphNode* GraphNode = EditorGraph->FindGraphNodeFromModelNode(Node);
 			EditorGraph->SetSelectedNodes({GraphNode});
 			
 			InspectObject(Node);
@@ -1171,13 +1196,10 @@ TSharedRef<SGraphEditor> FOptimusEditor::CreateGraphEditorWidget()
 			FCanExecuteAction::CreateSP(this, &FOptimusEditor::CanConvertToSubGraph)
 		);
 		
-#if 0
-
 		// Graph Editor Commands
 		GraphEditorCommands->MapAction(FGraphEditorCommands::Get().CreateComment,
 			FExecuteAction::CreateSP(this, &FOptimusEditor::OnCreateComment)
 		);
-#endif
 
 #if 0
 
@@ -1291,19 +1313,17 @@ bool FOptimusEditor::IsPropertyReadOnly(const FPropertyAndParent& InPropertyAndP
 
 TArray<UOptimusNode*> FOptimusEditor::GetSelectedModelNodes() const
 {
-	TArray<UOptimusNode*> Nodes;
+	TArray<UOptimusNode*> Nodes = EditorGraph->GetSelectedModelNodes();
 
-	for (UObject* Object: GraphEditorWidget->GetSelectedNodes())
-	{
-		UOptimusEditorGraphNode* GraphNode = Cast<UOptimusEditorGraphNode>(Object);
-		if (GraphNode)
-		{
-			Nodes.Add(GraphNode->ModelNode);
-		}
-	}
 	return Nodes;
 }
 
+bool FOptimusEditor::GetBoundsForSelectedNodes(FSlateRect& Rect)
+{
+	constexpr float BoundPadding = 50.0f;
+	return GraphEditorWidget->GetBoundsForSelectedNodes(Rect, BoundPadding);
+}
+
 
 void FOptimusEditor::StoreCurrentViewLocation()
 {
@@ -1417,7 +1437,7 @@ void FOptimusEditor::OnFinishedChangingProperties(const FPropertyChangedEvent& P
 
 					if (ModelPin)
 					{
-						if (UOptimusEditorGraphNode* GraphNode = EditorGraph->FindGraphNodeFromModelNode(ModelNode))
+						if (UOptimusEditorGraphNode* GraphNode = EditorGraph->FindOptimusGraphNodeFromModelNode(ModelNode))
 						{
 							GraphNode->SynchronizeGraphPinValueWithModelPin(ModelPin);
 						}
diff --git a/Engine/Plugins/Animation/DeformerGraph/Source/OptimusEditor/Private/OptimusEditor.h b/Engine/Plugins/Animation/DeformerGraph/Source/OptimusEditor/Private/OptimusEditor.h
--- a/Engine/Plugins/Animation/DeformerGraph/Source/OptimusEditor/Private/OptimusEditor.h
+++ b/Engine/Plugins/Animation/DeformerGraph/Source/OptimusEditor/Private/OptimusEditor.h
@@ -33,6 +33,7 @@ enum class EOptimusGlobalNotifyType;
 struct FGraphAppearanceInfo;
 struct FOptimusCompilerDiagnostic;
 enum class EOptimusGraphNotifyType;
+class FSlateRect;
 
 
 class FOptimusEditor :
@@ -168,6 +169,8 @@ private:
 	
 	void ConvertToSubGraph();
 	bool CanConvertToSubGraph() const;
+
+	void OnCreateComment();	
 	
 	// ----------------------------------------------------------------------------------------
 	// Graph event listeners
@@ -212,6 +215,8 @@ public:
 	// KILL ME
 	TSharedPtr<SGraphEditor> GetGraphEditorWidget() const { return GraphEditorWidget; }
 
+	bool GetBoundsForSelectedNodes(FSlateRect& Rect);
+	
 private:
 	void CreateMessageLog();
 	void HandleMessageTokenClicked(const TSharedRef<class IMessageToken>& InMessageToken);
diff --git a/Engine/Plugins/Animation/DeformerGraph/Source/OptimusEditor/Private/OptimusEditorGraph.cpp b/Engine/Plugins/Animation/DeformerGraph/Source/OptimusEditor/Private/OptimusEditorGraph.cpp
--- a/Engine/Plugins/Animation/DeformerGraph/Source/OptimusEditor/Private/OptimusEditorGraph.cpp
+++ b/Engine/Plugins/Animation/DeformerGraph/Source/OptimusEditor/Private/OptimusEditorGraph.cpp
@@ -5,6 +5,7 @@
 #include "OptimusEditorGraphNode.h"
 
 #include "OptimusNode.h"
+#include "OptimusNode_Comment.h"
 #include "OptimusNodeGraph.h"
 #include "OptimusNodeLink.h"
 #include "OptimusNodePin.h"
@@ -12,7 +13,13 @@
 #include "EdGraph/EdGraphPin.h"
 #include "Styling/AppStyle.h"
 #include "GraphEditAction.h"
+#include "OptimusActionStack.h"
+#include "OptimusEditorHelpers.h"
 #include "OptimusEditorStyle.h"
+#include "Editor.h"
+#include "OptimusEditor.h"
+#include "OptimusEditorGraphNode_Comment.h"
+#include "Layout/SlateRect.h"
 
 
 UOptimusEditorGraph::UOptimusEditorGraph()
@@ -26,13 +33,11 @@ void UOptimusEditorGraph::InitFromNodeGraph(UOptimusNodeGraph* InNodeGraph)
 	NodeGraph = InNodeGraph;
 
 	// Create all the nodes.
-	TMap<UOptimusNode*, UOptimusEditorGraphNode*> NodeMap;
 	for (UOptimusNode* ModelNode : InNodeGraph->GetAllNodes())
 	{
 		if (ensure(ModelNode != nullptr))
 		{
-			UOptimusEditorGraphNode* GraphNode = AddGraphNodeFromModelNode(ModelNode);
-			NodeMap.Add(ModelNode, GraphNode);
+			UEdGraphNode* GraphNode = AddGraphNodeFromModelNode(ModelNode);
 		}
 	}
 
@@ -44,8 +49,8 @@ void UOptimusEditorGraph::InitFromNodeGraph(UOptimusNodeGraph* InNodeGraph)
 			continue;
 		}
 
-		UOptimusEditorGraphNode* OutputGraphNode = NodeMap.FindRef(Link->GetNodeOutputPin()->GetOwningNode());
-		UOptimusEditorGraphNode* InputGraphNode = NodeMap.FindRef(Link->GetNodeInputPin()->GetOwningNode());
+		UEdGraphNode* OutputGraphNode = FindGraphNodeFromModelNode(Link->GetNodeOutputPin()->GetOwningNode());
+		UEdGraphNode* InputGraphNode = FindGraphNodeFromModelNode(Link->GetNodeInputPin()->GetOwningNode());
 
 		if (OutputGraphNode == nullptr || InputGraphNode == nullptr)
 		{
@@ -87,6 +92,87 @@ void UOptimusEditorGraph::Reset()
 }
 
 
+const TArray<UOptimusNode*> UOptimusEditorGraph::GetSelectedModelNodes() const
+{
+	TArray<UOptimusNode*> SelectedModelNodes;
+
+	for (UEdGraphNode* GraphNode: GetSelectedNodes())
+	{
+		if (UOptimusNode* ModelNode = OptimusEditor::FindModelNodeFromGraphNode(GraphNode))
+		{
+			SelectedModelNodes.Add(ModelNode);
+		}
+	}
+	
+	return SelectedModelNodes;
+}
+
+void UOptimusEditorGraph::HandleGraphNodeMoved()
+{
+#if WITH_EDITOR
+	// Cancel the current transaction created by SNodePanel::OnMouseMove so that the
+	// only transaction recorded is the one we place on the action stack.
+	if (GEditor)
+	{
+		GEditor->CancelTransaction(0);
+	}
+#endif
+
+	
+	if (SelectedNodes.Num() == 0)
+	{
+		return;
+	}
+	
+	TSet<UEdGraphNode*> MovedNodes;
+	for (UEdGraphNode* SelectedNode : SelectedNodes)
+	{
+		MovedNodes.Add(SelectedNode);
+		if (UOptimusEditorGraphNode_Comment* CommentNode = Cast<UOptimusEditorGraphNode_Comment>(SelectedNode))
+		{
+			for (FCommentNodeSet::TConstIterator NodeIt(CommentNode->GetNodesUnderComment()); NodeIt; ++NodeIt)
+			{
+				if (UEdGraphNode* NodeUnderComment = Cast<UEdGraphNode>(*NodeIt))
+				{
+					MovedNodes.Add(NodeUnderComment);
+				}
+			}
+		}
+	}
+	
+	FString ActionTitle;
+	if (MovedNodes.Num() == 1)
+	{
+		ActionTitle = TEXT("Move Node");
+	}
+	else
+	{
+		ActionTitle = FString::Printf(TEXT("Move %d Nodes"), MovedNodes.Num());
+	}
+	
+	FOptimusActionScope Scope(*GetModelGraph()->GetActionStack(), ActionTitle);
+	for (UEdGraphNode* MovedNode : MovedNodes)
+	{
+		FVector2D Position(MovedNode->NodePosX, MovedNode->NodePosY);
+		// It is possible that the model node is deleted during drag
+		UOptimusNode* ModelNode = OptimusEditor::FindModelNodeFromGraphNode(MovedNode);
+		if (ModelNode)
+		{
+			ModelNode->SetGraphPosition(Position);
+		}
+	}
+}
+
+bool UOptimusEditorGraph::GetBoundsForSelectedNodes(FSlateRect& Rect)
+{
+	if (!WeakEditor.IsValid())
+	{
+		return false;
+	}
+
+	return WeakEditor.Pin()->GetBoundsForSelectedNodes(Rect);
+}
+
 void UOptimusEditorGraph::RefreshVisualNode(UOptimusEditorGraphNode* InGraphNode)
 {
 	// Ensure that SOptimusEditorGraphNode captures the latest.
@@ -184,13 +270,11 @@ void UOptimusEditorGraph::HandleNodeGraphModified(EOptimusGraphNotifyType InNoti
 		{
 			const UOptimusNode* ModelNode = Cast<UOptimusNode>(InSubject);
 
-			UOptimusEditorGraphNode* GraphNode = FindGraphNodeFromModelNode(ModelNode);
+			UEdGraphNode* GraphNode = FindGraphNodeFromModelNode(ModelNode);
 
 			if (ensure(GraphNode))
 			{
-				RemoveNode(GraphNode, true);
-
-				GraphNode->ModelNode = nullptr;
+				RemoveGraphNode(GraphNode);
 			}
 		    break;
 		}
@@ -199,8 +283,8 @@ void UOptimusEditorGraph::HandleNodeGraphModified(EOptimusGraphNotifyType InNoti
 		case EOptimusGraphNotifyType::LinkRemoved:
 		{
 			const UOptimusNodeLink *ModelNodeLink = Cast<UOptimusNodeLink>(InSubject);
-			const UOptimusEditorGraphNode* OutputGraphNode = FindGraphNodeFromModelNode(ModelNodeLink->GetNodeOutputPin()->GetOwningNode());
-			const UOptimusEditorGraphNode* InputGraphNode = FindGraphNodeFromModelNode(ModelNodeLink->GetNodeInputPin()->GetOwningNode());
+			const UOptimusEditorGraphNode* OutputGraphNode = FindOptimusGraphNodeFromModelNode(ModelNodeLink->GetNodeOutputPin()->GetOwningNode());
+			const UOptimusEditorGraphNode* InputGraphNode = FindOptimusGraphNodeFromModelNode(ModelNodeLink->GetNodeInputPin()->GetOwningNode());
 
 			if (ensure(OutputGraphNode) && ensure(InputGraphNode))
 			{
@@ -225,7 +309,7 @@ void UOptimusEditorGraph::HandleNodeGraphModified(EOptimusGraphNotifyType InNoti
 		case EOptimusGraphNotifyType::NodeDisplayNameChanged:
 		{
 			const UOptimusNode* ModelNode = Cast<UOptimusNode>(InSubject);
-			UOptimusEditorGraphNode* GraphNode = FindGraphNodeFromModelNode(ModelNode);
+			UOptimusEditorGraphNode* GraphNode = FindOptimusGraphNodeFromModelNode(ModelNode);
 
 			GraphNode->SyncGraphNodeNameWithModelNodeName();
 		}
@@ -234,20 +318,26 @@ void UOptimusEditorGraph::HandleNodeGraphModified(EOptimusGraphNotifyType InNoti
 		case EOptimusGraphNotifyType::NodePositionChanged:
 		{
 			const UOptimusNode* ModelNode = Cast<UOptimusNode>(InSubject);
-			UOptimusEditorGraphNode* GraphNode = FindGraphNodeFromModelNode(ModelNode);
+			UEdGraphNode* GraphNode = FindGraphNodeFromModelNode(ModelNode);
 
 			if (ensure(GraphNode))
 			{
 				GraphNode->NodePosX = FMath::RoundToInt(ModelNode->GetGraphPosition().X);
 				GraphNode->NodePosY = FMath::RoundToInt(ModelNode->GetGraphPosition().Y);
 			}
+
+			if (UOptimusEditorGraphNode_Comment* CommentGraphNode = Cast<UOptimusEditorGraphNode_Comment>(GraphNode))
+			{
+				// Notify slate graph node widget
+				CommentGraphNode->OnPositionChanged();
+			}
 		    break;
 		}
 
 		case EOptimusGraphNotifyType::NodeDiagnosticLevelChanged:
 		{
 			const UOptimusNode* ModelNode = Cast<UOptimusNode>(InSubject);
-			UOptimusEditorGraphNode* GraphNode = FindGraphNodeFromModelNode(ModelNode);
+			UOptimusEditorGraphNode* GraphNode = FindOptimusGraphNodeFromModelNode(ModelNode);
 
 			GraphNode->SyncDiagnosticStateWithModelNode();
 
@@ -259,7 +349,7 @@ void UOptimusEditorGraph::HandleNodeGraphModified(EOptimusGraphNotifyType InNoti
 		    const UOptimusNodePin* ModelPin = Cast<UOptimusNodePin>(InSubject);
 		    if (ensure(ModelPin))
 		    {
-			    UOptimusEditorGraphNode* GraphNode = FindGraphNodeFromModelNode(ModelPin->GetOwningNode());
+			    UOptimusEditorGraphNode* GraphNode = FindOptimusGraphNodeFromModelNode(ModelPin->GetOwningNode());
 
 			    if (ensure(GraphNode))
 			    {
@@ -274,7 +364,7 @@ void UOptimusEditorGraph::HandleNodeGraphModified(EOptimusGraphNotifyType InNoti
 			const UOptimusNodePin* ModelPin = Cast<UOptimusNodePin>(InSubject);
 			if (ensure(ModelPin))
 			{
-				UOptimusEditorGraphNode* GraphNode = FindGraphNodeFromModelNode(ModelPin->GetOwningNode());
+				UOptimusEditorGraphNode* GraphNode = FindOptimusGraphNodeFromModelNode(ModelPin->GetOwningNode());
 
 				if (ensure(GraphNode))
 				{
@@ -289,7 +379,7 @@ void UOptimusEditorGraph::HandleNodeGraphModified(EOptimusGraphNotifyType InNoti
 			const UOptimusNodePin* ModelPin = Cast<UOptimusNodePin>(InSubject);
 			if (ensure(ModelPin))
 			{
-				UOptimusEditorGraphNode* GraphNode = FindGraphNodeFromModelNode(ModelPin->GetOwningNode());
+				UOptimusEditorGraphNode* GraphNode = FindOptimusGraphNodeFromModelNode(ModelPin->GetOwningNode());
 
 				if (ensure(GraphNode))
 				{
@@ -304,7 +394,7 @@ void UOptimusEditorGraph::HandleNodeGraphModified(EOptimusGraphNotifyType InNoti
 		    const UOptimusNodePin* ModelPin = Cast<UOptimusNodePin>(InSubject);
 		    if (ensure(ModelPin))
 		    {
-			    UOptimusEditorGraphNode* GraphNode = FindGraphNodeFromModelNode(ModelPin->GetOwningNode());
+			    UOptimusEditorGraphNode* GraphNode = FindOptimusGraphNodeFromModelNode(ModelPin->GetOwningNode());
 
 			    if (ensure(GraphNode))
 			    {
@@ -322,7 +412,7 @@ void UOptimusEditorGraph::HandleNodeGraphModified(EOptimusGraphNotifyType InNoti
 			const UOptimusNodePin *ModelPin = Cast<UOptimusNodePin>(InSubject);
 			if (ensure(ModelPin))
 			{
-			    UOptimusEditorGraphNode* GraphNode = FindGraphNodeFromModelNode(ModelPin->GetOwningNode());
+			    UOptimusEditorGraphNode* GraphNode = FindOptimusGraphNodeFromModelNode(ModelPin->GetOwningNode());
 
 				if (ensure(GraphNode))
 				{
@@ -341,7 +431,7 @@ void UOptimusEditorGraph::HandleNodeGraphModified(EOptimusGraphNotifyType InNoti
 		    const UOptimusNodePin* ModelPin = Cast<UOptimusNodePin>(InSubject);
 		    if (ensure(ModelPin))
 		    {
-			    UOptimusEditorGraphNode* GraphNode = FindGraphNodeFromModelNode(ModelPin->GetOwningNode());
+			    UOptimusEditorGraphNode* GraphNode = FindOptimusGraphNodeFromModelNode(ModelPin->GetOwningNode());
 
 			    if (ensure(GraphNode))
 			    {
@@ -356,7 +446,7 @@ void UOptimusEditorGraph::HandleNodeGraphModified(EOptimusGraphNotifyType InNoti
 			const UOptimusNodePin* ModelPin = Cast<UOptimusNodePin>(InSubject);
 			if (ensure(ModelPin))
 			{
-				UOptimusEditorGraphNode* GraphNode = FindGraphNodeFromModelNode(ModelPin->GetOwningNode());
+				UOptimusEditorGraphNode* GraphNode = FindOptimusGraphNodeFromModelNode(ModelPin->GetOwningNode());
 
 				if (ensure(GraphNode))
 				{
@@ -369,36 +459,68 @@ void UOptimusEditorGraph::HandleNodeGraphModified(EOptimusGraphNotifyType InNoti
 }
 
 
-UOptimusEditorGraphNode* UOptimusEditorGraph::AddGraphNodeFromModelNode(UOptimusNode* InModelNode)
+UEdGraphNode* UOptimusEditorGraph::AddGraphNodeFromModelNode(UOptimusNode* InModelNode)
 {
-	FGraphNodeCreator<UOptimusEditorGraphNode> NodeCreator(*this);
+	UEdGraphNode* GraphNode = nullptr;
+	
+	const bool bIsCreatedFromUI = InModelNode->IsCreatedFromUI();
+	
+	if (UOptimusNode_Comment* CommentNode = Cast<UOptimusNode_Comment>(InModelNode))
+	{
+		FGraphNodeCreator<UOptimusEditorGraphNode_Comment> NodeCreator(*this);
 
-	UOptimusEditorGraphNode* GraphNode = NodeCreator.CreateNode(false);
-	GraphNode->Construct(InModelNode);
-	NodeCreator.Finalize();
+		UOptimusEditorGraphNode_Comment* CommentGraphNode = bIsCreatedFromUI ?
+			NodeCreator.CreateUserInvokedNode(false) :
+			NodeCreator.CreateNode(false);
+		
+		CommentGraphNode->Construct(CommentNode);
+		NodeCreator.Finalize();
+
+		GraphNode = CommentGraphNode;
+	}
+	else
+	{
+		FGraphNodeCreator<UOptimusEditorGraphNode> NodeCreator(*this);
+
+		UOptimusEditorGraphNode* OptimusGraphNode = bIsCreatedFromUI ?
+			NodeCreator.CreateUserInvokedNode(false) :
+			NodeCreator.CreateNode(false);
+		
+		OptimusGraphNode->Construct(InModelNode);
+		NodeCreator.Finalize();
+
+		GraphNode = OptimusGraphNode;
+	}
 
+		
+	GraphNodeToModelNodeMap.Add(GraphNode, InModelNode);
+	ModelNodeToGraphNodeMap.Add(InModelNode, GraphNode);
+	
 	return GraphNode;
 }
 
-
-UOptimusEditorGraphNode* UOptimusEditorGraph::FindGraphNodeFromModelNode(const UOptimusNode* ModelNode)
+bool UOptimusEditorGraph::RemoveGraphNode(UEdGraphNode* NodeToRemove, bool bBreakAllLinks)
 {
-	if (!ModelNode)
-	{
-		return nullptr;
-	}
+	UOptimusNode* ModelNode = GraphNodeToModelNodeMap[NodeToRemove].Get(); 
 
-	// FIXME: Store this info in a map.
-	for (UEdGraphNode* Node : Nodes)
-	{
-		UOptimusEditorGraphNode *GraphNode = Cast<UOptimusEditorGraphNode>(Node);
-		ensure(GraphNode);
+	ModelNodeToGraphNodeMap.Remove(ModelNode);
+	GraphNodeToModelNodeMap.Remove(NodeToRemove);
+	
+	return RemoveNode(NodeToRemove, bBreakAllLinks);
+}
 
-		if (GraphNode && GraphNode->ModelNode == ModelNode)
-		{
-			return GraphNode;
-		}
-	}
 
-	return nullptr;
+UOptimusEditorGraphNode* UOptimusEditorGraph::FindOptimusGraphNodeFromModelNode(const UOptimusNode* InModelNode)
+{
+	return Cast<UOptimusEditorGraphNode>(FindGraphNodeFromModelNode(InModelNode));
+}
+
+UEdGraphNode* UOptimusEditorGraph::FindGraphNodeFromModelNode(const UOptimusNode* InModelNode)
+{
+	return ModelNodeToGraphNodeMap.FindRef(InModelNode).Get();
+}
+
+UOptimusNode* UOptimusEditorGraph::FindModelNodeFromGraphNode(const UEdGraphNode* InGraphNode)
+{
+	return GraphNodeToModelNodeMap.FindRef(InGraphNode).Get();
 }
diff --git a/Engine/Plugins/Animation/DeformerGraph/Source/OptimusEditor/Private/OptimusEditorGraph.h b/Engine/Plugins/Animation/DeformerGraph/Source/OptimusEditor/Private/OptimusEditorGraph.h
--- a/Engine/Plugins/Animation/DeformerGraph/Source/OptimusEditor/Private/OptimusEditorGraph.h
+++ b/Engine/Plugins/Animation/DeformerGraph/Source/OptimusEditor/Private/OptimusEditorGraph.h
@@ -13,6 +13,8 @@ struct FSlateBrush;
 class UOptimusNode;
 class UOptimusDeformer;
 class UOptimusEditorGraphNode;
+class FSlateRect;
+class FOptimusEditor;
 
 UCLASS()
 class UOptimusEditorGraph : public UEdGraph
@@ -28,10 +30,19 @@ public:
 
 	UOptimusNodeGraph* GetModelGraph() const { return NodeGraph; }
 
-	UOptimusEditorGraphNode* FindGraphNodeFromModelNode(const UOptimusNode* Node);
+	UEdGraphNode* FindGraphNodeFromModelNode(const UOptimusNode* InModelNode);
+	UOptimusNode* FindModelNodeFromGraphNode(const UEdGraphNode* InGraphNode);
+	
+	UOptimusEditorGraphNode* FindOptimusGraphNodeFromModelNode(const UOptimusNode* InModelNode);
 
-	const TSet<UOptimusEditorGraphNode*> &GetSelectedNodes() const { return SelectedNodes; }
+	const TSet<UEdGraphNode*> &GetSelectedNodes() const { return SelectedNodes; }
+	
+	const TArray<UOptimusNode*> GetSelectedModelNodes() const;
 
+	void HandleGraphNodeMoved();
+
+	bool GetBoundsForSelectedNodes(FSlateRect& Rect);
+	
 	// Do a visual refresh of the node.
 	void RefreshVisualNode(UOptimusEditorGraphNode *InGraphNode);
 
@@ -41,7 +52,7 @@ public:
 protected:
 	friend class FOptimusEditor;
 
-	void SetSelectedNodes(const TSet<UOptimusEditorGraphNode*>& InSelectedNodes)
+	void SetSelectedNodes(const TSet<UEdGraphNode*>& InSelectedNodes)
 	{
 		SelectedNodes = InSelectedNodes;
 	}
@@ -57,9 +68,16 @@ private:
 		UObject *InSubject
 		);
 
-	UOptimusEditorGraphNode* AddGraphNodeFromModelNode(UOptimusNode* InModelNode);
+	UEdGraphNode* AddGraphNodeFromModelNode(UOptimusNode* InModelNode);
+
+	bool RemoveGraphNode(UEdGraphNode* NodeToRemove, bool bBreakAllLinks = true);
 
 	TObjectPtr<UOptimusNodeGraph> NodeGraph;
 
-	TSet<UOptimusEditorGraphNode*> SelectedNodes;
+	TSet<UEdGraphNode*> SelectedNodes;
+
+	TMap<TWeakObjectPtr<UEdGraphNode>, TWeakObjectPtr<UOptimusNode>> GraphNodeToModelNodeMap;
+	TMap<TWeakObjectPtr<UOptimusNode>, TWeakObjectPtr<UEdGraphNode>> ModelNodeToGraphNodeMap;
+
+	TWeakPtr<FOptimusEditor> WeakEditor;
 };
diff --git a/Engine/Plugins/Animation/DeformerGraph/Source/OptimusEditor/Private/OptimusEditorGraphNodeFactory.cpp b/Engine/Plugins/Animation/DeformerGraph/Source/OptimusEditor/Private/OptimusEditorGraphNodeFactory.cpp
--- a/Engine/Plugins/Animation/DeformerGraph/Source/OptimusEditor/Private/OptimusEditorGraphNodeFactory.cpp
+++ b/Engine/Plugins/Animation/DeformerGraph/Source/OptimusEditor/Private/OptimusEditorGraphNodeFactory.cpp
@@ -3,7 +3,9 @@
 #include "OptimusEditorGraphNodeFactory.h"
 
 #include "OptimusEditorGraphNode.h"
+#include "OptimusEditorGraphNode_Comment.h"
 #include "Widgets/SOptimusEditorGraphNode.h"
+#include "Widgets/SOptimusEditorGraphNode_Comment.h"
 
 TSharedPtr<SGraphNode> FOptimusEditorGraphNodeFactory::CreateNode(UEdGraphNode* InNode) const
 {
@@ -17,6 +19,17 @@ TSharedPtr<SGraphNode> FOptimusEditorGraphNodeFactory::CreateNode(UEdGraphNode*
 		// FIXME: Copy visual node's dimensions back to graph node.
 		return VisualNode;
 	}
+	
+	if (UOptimusEditorGraphNode_Comment* CommentNode = Cast<UOptimusEditorGraphNode_Comment>(InNode))
+	{
+		TSharedRef<SGraphNode> VisualNode = 
+			SNew(SOptimusEditorGraphNode_Comment, CommentNode);
+
+		VisualNode->SlatePrepass();
+		
+		// FIXME: Copy visual node's dimensions back to graph node.
+		return VisualNode;
+	}
 
 	return nullptr;
 }
diff --git a/Engine/Plugins/Animation/DeformerGraph/Source/OptimusEditor/Private/OptimusEditorGraphNode_Comment.cpp b/Engine/Plugins/Animation/DeformerGraph/Source/OptimusEditor/Private/OptimusEditorGraphNode_Comment.cpp
new file mode 100644
--- /dev/null
+++ b/Engine/Plugins/Animation/DeformerGraph/Source/OptimusEditor/Private/OptimusEditorGraphNode_Comment.cpp
@@ -0,0 +1,64 @@
+// Copyright Epic Games, Inc. All Rights Reserved.
+
+#include "OptimusEditorGraphNode_Comment.h"
+#include "OptimusNode_Comment.h"
+
+#define LOCTEXT_NAMESPACE "OptimusEditorGraphNode_Comment"
+
+
+void UOptimusEditorGraphNode_Comment::Construct(UOptimusNode_Comment* InCommentNode)
+{
+	// Our graph nodes are not transactional. We handle the transacting ourselves.
+	ClearFlags(RF_Transactional);
+	NodePosX = UE::LWC::FloatToIntCastChecked<int32>(InCommentNode->GetGraphPosition().X);
+	NodePosY = UE::LWC::FloatToIntCastChecked<int32>(InCommentNode->GetGraphPosition().Y);
+
+	// Comment node specific properties
+	NodeWidth = UE::LWC::FloatToIntCastChecked<int32>(InCommentNode->GetSize().X);
+	NodeHeight = UE::LWC::FloatToIntCastChecked<int32>(InCommentNode->GetSize().Y);
+	CommentColor = InCommentNode->CommentColor;
+	NodeComment = InCommentNode->Comment;
+	FontSize = InCommentNode->FontSize;
+	bCommentBubbleVisible_InDetailsPanel = InCommentNode->bBubbleVisible;
+	bCommentBubbleVisible = bCommentBubbleVisible_InDetailsPanel;
+	bCommentBubblePinned = bCommentBubbleVisible_InDetailsPanel;
+	bColorCommentBubble = InCommentNode->bColorBubble;
+
+	InCommentNode->GetOnPropertyChanged().BindUObject(this, &UOptimusEditorGraphNode_Comment::OnModelNodePropertyChanged);
+}
+
+void UOptimusEditorGraphNode_Comment::OnModelNodePropertyChanged(UOptimusNode_Comment* InCommentNode)
+{
+	ResizeNode(InCommentNode->GetSize());
+	
+	CommentColor = InCommentNode->CommentColor;
+	if (NodeComment != InCommentNode->Comment)
+	{
+		OnRenameNode(InCommentNode->Comment);
+	}
+	FontSize = InCommentNode->FontSize;
+	bCommentBubbleVisible_InDetailsPanel = InCommentNode->bBubbleVisible;
+	bCommentBubbleVisible = bCommentBubbleVisible_InDetailsPanel;
+	bCommentBubblePinned = bCommentBubbleVisible_InDetailsPanel;
+	bColorCommentBubble = InCommentNode->bColorBubble;
+}
+
+void UOptimusEditorGraphNode_Comment::OnPositionChanged()
+{
+	OnPositionChangedDelegate.ExecuteIfBound();
+}
+
+void UOptimusEditorGraphNode_Comment::ResizeNode(const FVector2D& NewSize)
+{
+	Super::ResizeNode(NewSize);
+
+	// Notify graph node widget about the new size, since it is not done automatically on undo 
+	OnSizeChangedDelegate.ExecuteIfBound();
+}
+
+void UOptimusEditorGraphNode_Comment::PostPlacedNewNode()
+{
+	// Do nothing here to avoid overriding work done during construct()
+}
+
+#undef LOCTEXT_NAMESPACE
diff --git a/Engine/Plugins/Animation/DeformerGraph/Source/OptimusEditor/Private/OptimusEditorGraphNode_Comment.h b/Engine/Plugins/Animation/DeformerGraph/Source/OptimusEditor/Private/OptimusEditorGraphNode_Comment.h
new file mode 100644
--- /dev/null
+++ b/Engine/Plugins/Animation/DeformerGraph/Source/OptimusEditor/Private/OptimusEditorGraphNode_Comment.h
@@ -0,0 +1,33 @@
+// Copyright Epic Games, Inc. All Rights Reserved.
+
+#pragma once
+
+#include "EdGraphNode_Comment.h"
+
+#include "OptimusEditorGraphNode_Comment.generated.h"
+
+class UOptimusNode_Comment;
+
+DECLARE_DELEGATE(FOptimusCommentNodeSizeChanged);
+DECLARE_DELEGATE(FOptimusCommentNodePositionChanged);
+
+UCLASS()
+class UOptimusEditorGraphNode_Comment : public UEdGraphNode_Comment
+{
+	GENERATED_BODY()
+
+public:
+	void Construct(UOptimusNode_Comment* InCommentNode);
+	void OnModelNodePropertyChanged(UOptimusNode_Comment* InCommentNode);
+
+	void OnPositionChanged();
+
+	void ResizeNode(const FVector2D& NewSize) override;
+	void PostPlacedNewNode() override;
+	
+	FOptimusCommentNodeSizeChanged& GetOnSizeChanged() {return OnSizeChangedDelegate; }
+	FOptimusCommentNodePositionChanged& GetOnPositionChanged() {return OnPositionChangedDelegate; }
+private:
+	FOptimusCommentNodeSizeChanged OnSizeChangedDelegate;
+	FOptimusCommentNodePositionChanged OnPositionChangedDelegate;
+};
diff --git a/Engine/Plugins/Animation/DeformerGraph/Source/OptimusEditor/Private/OptimusEditorGraphSchema.cpp b/Engine/Plugins/Animation/DeformerGraph/Source/OptimusEditor/Private/OptimusEditorGraphSchema.cpp
--- a/Engine/Plugins/Animation/DeformerGraph/Source/OptimusEditor/Private/OptimusEditorGraphSchema.cpp
+++ b/Engine/Plugins/Animation/DeformerGraph/Source/OptimusEditor/Private/OptimusEditorGraphSchema.cpp
@@ -246,6 +246,17 @@ void UOptimusEditorGraphSchema::GetGraphActions(
 		IoActionBuilder.AddAction(Action);
 	}
 	
+	// Comment nodes
+	{
+		TSharedPtr< FOptimusGraphSchemaAction_NewCommentNode> Action(
+			new FOptimusGraphSchemaAction_NewCommentNode(
+				{},
+				FText::FromName(TEXT("Comment")),
+				/* Tooltip */{}, 0, /* Keywords */{}
+		));
+		
+		IoActionBuilder.AddAction(Action);
+	}	
 }
 
 
@@ -444,16 +455,9 @@ void UOptimusEditorGraphSchema::GetGraphContextActions(
 
 bool UOptimusEditorGraphSchema::SafeDeleteNodeFromGraph(UEdGraph* InGraph, UEdGraphNode* InNode) const
 {
-	UOptimusEditorGraphNode* GraphNode = Cast<UOptimusEditorGraphNode>(InNode);
+	UOptimusEditorGraph* Graph = Cast<UOptimusEditorGraph>(InNode->GetGraph());
 
-	if (GraphNode)
-	{
-		UOptimusEditorGraph* Graph = Cast<UOptimusEditorGraph>(GraphNode->GetGraph());
-
-		return Graph->GetModelGraph()->RemoveNode(GraphNode->ModelNode);
-	}
-
-	return false;
+	return Graph->GetModelGraph()->RemoveNode(Graph->FindModelNodeFromGraphNode(InNode));
 }
 
 void UOptimusEditorGraphSchema::GetGraphDisplayInformation(const UEdGraph& Graph, FGraphDisplayInfo& DisplayInfo) const
diff --git a/Engine/Plugins/Animation/DeformerGraph/Source/OptimusEditor/Private/OptimusEditorGraphSchemaActions.cpp b/Engine/Plugins/Animation/DeformerGraph/Source/OptimusEditor/Private/OptimusEditorGraphSchemaActions.cpp
--- a/Engine/Plugins/Animation/DeformerGraph/Source/OptimusEditor/Private/OptimusEditorGraphSchemaActions.cpp
+++ b/Engine/Plugins/Animation/DeformerGraph/Source/OptimusEditor/Private/OptimusEditorGraphSchemaActions.cpp
@@ -5,6 +5,7 @@
 #include "OptimusComponentSource.h"
 #include "OptimusEditorGraph.h"
 #include "OptimusEditorGraphNode.h"
+#include "OptimusEditorHelpers.h"
 #include "OptimusFunctionNodeGraph.h"
 
 #include "OptimusNode.h"
@@ -30,7 +31,7 @@ UEdGraphNode* FOptimusGraphSchemaAction_NewNode::PerformAction(
 
  		// FIXME: Automatic connection from the given pin.
 
-		UOptimusEditorGraphNode* GraphNode = Graph->FindGraphNodeFromModelNode(ModelNode);
+		UEdGraphNode* GraphNode = Graph->FindGraphNodeFromModelNode(ModelNode);
 		if (GraphNode && bInSelectNewNode)
 		{
 			Graph->SelectNodeSet({GraphNode});
@@ -57,7 +58,7 @@ UEdGraphNode* FOptimusGraphSchemaAction_NewConstantValueNode::PerformAction(
 
 		// FIXME: Automatic connection from the given pin.
 
-		UOptimusEditorGraphNode* GraphNode = Graph->FindGraphNodeFromModelNode(ModelNode);
+		UEdGraphNode* GraphNode = Graph->FindGraphNodeFromModelNode(ModelNode);
 		if (GraphNode && bInSelectNewNode)
 		{
 			Graph->SelectNodeSet({GraphNode});
@@ -84,7 +85,7 @@ UEdGraphNode* FOptimusGraphSchemaAction_NewDataInterfaceNode::PerformAction(
 
 		// FIXME: Automatic connection from the given pin.
 
-		UOptimusEditorGraphNode* GraphNode = Graph->FindGraphNodeFromModelNode(ModelNode);
+		UEdGraphNode* GraphNode = Graph->FindGraphNodeFromModelNode(ModelNode);
 		if (GraphNode && bInSelectNewNode)
 		{
 			Graph->SelectNodeSet({GraphNode});
@@ -106,7 +107,7 @@ UEdGraphNode* FOptimusGraphSchemaAction_NewLoopTerminalNodes::PerformAction(UEdG
 
 		if (ensure(Nodes.Num() == 2))
 		{
-			UOptimusEditorGraphNode* GraphNode = Graph->FindGraphNodeFromModelNode(Nodes[0]);
+			UEdGraphNode* GraphNode = Graph->FindGraphNodeFromModelNode(Nodes[0]);
 			if (GraphNode && bInSelectNewNode)
 			{
 				Graph->SelectNodeSet({GraphNode});
@@ -118,6 +119,26 @@ UEdGraphNode* FOptimusGraphSchemaAction_NewLoopTerminalNodes::PerformAction(UEdG
 	return nullptr;
 }
 
+UEdGraphNode* FOptimusGraphSchemaAction_NewCommentNode::PerformAction(UEdGraph* InParentGraph, UEdGraphPin* FromPin, const FVector2D Location,
+	bool bInSelectNewNode)
+{
+	UOptimusEditorGraph* Graph = Cast<UOptimusEditorGraph>(InParentGraph);
+	
+	if (ensure(Graph != nullptr))
+	{
+		UOptimusNode* ModelNode = OptimusEditor::CreateCommentNode(Graph, Location);
+
+		UEdGraphNode* GraphNode = Graph->FindGraphNodeFromModelNode(ModelNode);
+		if (GraphNode && bInSelectNewNode)
+		{
+			Graph->SelectNodeSet({GraphNode});
+		}
+		return GraphNode;	
+	}
+
+	return nullptr;
+}
+
 UEdGraphNode* FOptimusGraphSchemaAction_NewFunctionReferenceNode::PerformAction(UEdGraph* InParentGraph, UEdGraphPin* InFromPin, const FVector2D InLocation, bool bInSelectNewNode)
 {
 	UOptimusEditorGraph* Graph = Cast<UOptimusEditorGraph>(InParentGraph);
@@ -128,7 +149,7 @@ UEdGraphNode* FOptimusGraphSchemaAction_NewFunctionReferenceNode::PerformAction(
 
 		// FIXME: Automatic connection from the given pin.
 
-		UOptimusEditorGraphNode* GraphNode = Graph->FindGraphNodeFromModelNode(ModelNode);
+		UEdGraphNode* GraphNode = Graph->FindGraphNodeFromModelNode(ModelNode);
 		if (GraphNode && bInSelectNewNode)
 		{
 			Graph->SelectNodeSet({GraphNode});
diff --git a/Engine/Plugins/Animation/DeformerGraph/Source/OptimusEditor/Private/OptimusEditorGraphSchemaActions.h b/Engine/Plugins/Animation/DeformerGraph/Source/OptimusEditor/Private/OptimusEditorGraphSchemaActions.h
--- a/Engine/Plugins/Animation/DeformerGraph/Source/OptimusEditor/Private/OptimusEditorGraphSchemaActions.h
+++ b/Engine/Plugins/Animation/DeformerGraph/Source/OptimusEditor/Private/OptimusEditorGraphSchemaActions.h
@@ -104,6 +104,21 @@ struct FOptimusGraphSchemaAction_NewLoopTerminalNodes : public FEdGraphSchemaAct
 	UEdGraphNode* PerformAction(class UEdGraph* InParentGraph, UEdGraphPin* FromPin, const FVector2D Location, bool bInSelectNewNode = true) override;
 };
 
+USTRUCT()
+struct FOptimusGraphSchemaAction_NewCommentNode : public FEdGraphSchemaAction
+{
+	GENERATED_BODY()
+
+	// Inherit the base class's constructors
+	using FEdGraphSchemaAction::FEdGraphSchemaAction;
+	
+	static FName StaticGetTypeId() { static FName Type("FOptimusGraphSchemaAction_CommentNode"); return Type; }
+	FName GetTypeId() const override { return StaticGetTypeId(); }
+
+	// FEdGraphSchemaAction overrides
+	UEdGraphNode* PerformAction(class UEdGraph* InParentGraph, UEdGraphPin* FromPin, const FVector2D Location, bool bInSelectNewNode = true) override;
+};
+
 USTRUCT()
 struct FOptimusGraphSchemaAction_NewFunctionReferenceNode: 
 	public FEdGraphSchemaAction
diff --git a/Engine/Plugins/Animation/DeformerGraph/Source/OptimusEditor/Private/OptimusEditorHelpers.cpp b/Engine/Plugins/Animation/DeformerGraph/Source/OptimusEditor/Private/OptimusEditorHelpers.cpp
--- a/Engine/Plugins/Animation/DeformerGraph/Source/OptimusEditor/Private/OptimusEditorHelpers.cpp
+++ b/Engine/Plugins/Animation/DeformerGraph/Source/OptimusEditor/Private/OptimusEditorHelpers.cpp
@@ -1,22 +1,20 @@
 // Copyright Epic Games, Inc. All Rights Reserved.
 
 #include "OptimusEditorHelpers.h"
+
+#include "GraphEditorSettings.h"
+#include "OptimusEditorGraph.h"
 #include "EdGraph/EdGraphPin.h"
-#include "OptimusEditorGraphNode.h"
+#include "EdGraph/EdGraphNode.h"
 
 #include "OptimusNode.h"
 #include "OptimusNodePin.h"
+#include "OptimusNodeGraph.h"
+#include "Layout/SlateRect.h"
 
 UOptimusNode* OptimusEditor::GetModelNodeFromGraphPin(const UEdGraphPin* InGraphPin)
 {
-	UOptimusEditorGraphNode* GraphNode = Cast<UOptimusEditorGraphNode>(InGraphPin->GetOwningNode());
-
-	if (ensure(GraphNode != nullptr) && ensure(GraphNode->ModelNode != nullptr))
-	{
-		return GraphNode->ModelNode;
-	}
-
-	return nullptr;
+	return FindModelNodeFromGraphNode(InGraphPin->GetOwningNode());
 }
 
 UOptimusNodePin* OptimusEditor::GetModelPinFromGraphPin(const UEdGraphPin* InGraphPin)
@@ -31,6 +29,22 @@ UOptimusNodePin* OptimusEditor::GetModelPinFromGraphPin(const UEdGraphPin* InGra
 	return nullptr;
 }
 
+UOptimusNode* OptimusEditor::FindModelNodeFromGraphNode(const UEdGraphNode* InGraphNode)
+{
+	if (!InGraphNode)
+	{
+		return nullptr;
+	}
+	
+	UOptimusEditorGraph* EditorGraph = Cast<UOptimusEditorGraph>(InGraphNode->GetGraph());
+	if (!ensure(EditorGraph))
+	{
+		return nullptr;
+	}
+
+	return EditorGraph->FindModelNodeFromGraphNode(InGraphNode);
+}
+
 FName OptimusEditor::GetAdderPinName(EEdGraphPinDirection InDirection)
 {
 	return InDirection == EGPD_Input ? TEXT("_AdderPinInput") : TEXT("_AdderPinOutput");
@@ -65,4 +79,18 @@ bool OptimusEditor::IsAdderPinType(const FEdGraphPinType& InPinType)
 	return false;
 }
 
+UOptimusNode* OptimusEditor::CreateCommentNode(UOptimusEditorGraph* InEditorGraph, const FVector2D& InPosition)
+{
+	const FVector2D DefaultSize(400, 100);
+	FSlateRect Rect;
+	if (!InEditorGraph->GetBoundsForSelectedNodes(Rect))
+	{
+		Rect = FSlateRect(InPosition.X, InPosition.Y, InPosition.X+DefaultSize.X, InPosition.Y+DefaultSize.Y);
+	}
+
+	const UGraphEditorSettings* GraphEditorSettings = GetDefault<UGraphEditorSettings>();
+	constexpr bool bCreatedFromUI = true;
+	return InEditorGraph->GetModelGraph()->AddCommentNode(Rect.GetTopLeft(), Rect.GetSize(), GraphEditorSettings->DefaultCommentNodeTitleColor, bCreatedFromUI);	
+}
+
 
diff --git a/Engine/Plugins/Animation/DeformerGraph/Source/OptimusEditor/Private/OptimusEditorHelpers.h b/Engine/Plugins/Animation/DeformerGraph/Source/OptimusEditor/Private/OptimusEditorHelpers.h
--- a/Engine/Plugins/Animation/DeformerGraph/Source/OptimusEditor/Private/OptimusEditorHelpers.h
+++ b/Engine/Plugins/Animation/DeformerGraph/Source/OptimusEditor/Private/OptimusEditorHelpers.h
@@ -2,12 +2,14 @@
 
 #pragma once
 #include "CoreMinimal.h"
-#include "EdGraph/EdGraphNode.h"
 
 class UEdGraphPin;
+class UEdGraphNode;
 class UOptimusNode;
 class UOptimusNodePin;
 struct FEdGraphPinType;
+class UOptimusEditorGraph;
+enum EEdGraphPinDirection : int;
 
 namespace OptimusEditor
 {
@@ -15,6 +17,8 @@ namespace OptimusEditor
 
 	UOptimusNodePin* GetModelPinFromGraphPin(const UEdGraphPin* InGraphPin);
 
+	UOptimusNode* FindModelNodeFromGraphNode(const UEdGraphNode* InGraphNode);
+
 	FName GetAdderPinName(EEdGraphPinDirection InDirection);
 		
 	FText GetAdderPinFriendlyName(EEdGraphPinDirection InDirection);
@@ -24,4 +28,6 @@ namespace OptimusEditor
 	bool IsAdderPin(const UEdGraphPin* InGraphPin);
 
 	bool IsAdderPinType(const FEdGraphPinType& InPinType);
+
+	UOptimusNode* CreateCommentNode(UOptimusEditorGraph* InGraph, const FVector2D& InPosition);
 }
diff --git a/Engine/Plugins/Animation/DeformerGraph/Source/OptimusEditor/Private/Widgets/SOptimusEditorGraphNode.cpp b/Engine/Plugins/Animation/DeformerGraph/Source/OptimusEditor/Private/Widgets/SOptimusEditorGraphNode.cpp
--- a/Engine/Plugins/Animation/DeformerGraph/Source/OptimusEditor/Private/Widgets/SOptimusEditorGraphNode.cpp
+++ b/Engine/Plugins/Animation/DeformerGraph/Source/OptimusEditor/Private/Widgets/SOptimusEditorGraphNode.cpp
@@ -8,12 +8,10 @@
 #include "OptimusEditorStyle.h"
 
 #include "IOptimusNodeAdderPinProvider.h"
-#include "OptimusActionStack.h"
 #include "OptimusNode.h"
 #include "OptimusNodeGraph.h"
 #include "OptimusNodePin.h"
 
-#include "Editor.h"
 #include "GraphEditorSettings.h"
 #include "IDocumentation.h"
 #include "OptimusComponentSource.h"
@@ -684,42 +682,7 @@ void SOptimusEditorGraphNode::EndUserInteraction() const
 	UOptimusEditorGraph* Graph = Cast<UOptimusEditorGraph>(GraphNode->GetGraph());
 	if (ensure(Graph))
 	{
-#if WITH_EDITOR
-		// Cancel the current transaction created by SNodePanel::OnMouseMove so that the
-		// only transaction recorded is the one we place on the action stack.
-		if (GEditor)
-		{
-			GEditor->CancelTransaction(0);
-		}
-#endif
-
-		const TSet<UOptimusEditorGraphNode*> &SelectedNodes = Graph->GetSelectedNodes();
-
-		if (SelectedNodes.Num() == 0)
-		{
-			return;
-		}
-
-		FString ActionTitle;
-		if (SelectedNodes.Num() == 1)
-		{
-			ActionTitle = TEXT("Move Node");
-		}
-		else
-		{
-			ActionTitle = FString::Printf(TEXT("Move %d Nodes"), SelectedNodes.Num());
-		}
-
-		FOptimusActionScope Scope(*Graph->GetModelGraph()->GetActionStack(), ActionTitle);
-		for (UOptimusEditorGraphNode* SelectedNode : SelectedNodes)
-		{
-			FVector2D Position(SelectedNode->NodePosX, SelectedNode->NodePosY);
-			// It is possible that the model node is deleted during drag
-			if (SelectedNode->ModelNode)
-			{
-				SelectedNode->ModelNode->SetGraphPosition(Position);
-			}
-		}
+		Graph->HandleGraphNodeMoved();
 	}
 }
 
diff --git a/Engine/Plugins/Animation/DeformerGraph/Source/OptimusEditor/Private/Widgets/SOptimusEditorGraphNode_Comment.cpp b/Engine/Plugins/Animation/DeformerGraph/Source/OptimusEditor/Private/Widgets/SOptimusEditorGraphNode_Comment.cpp
new file mode 100644
--- /dev/null
+++ b/Engine/Plugins/Animation/DeformerGraph/Source/OptimusEditor/Private/Widgets/SOptimusEditorGraphNode_Comment.cpp
@@ -0,0 +1,66 @@
+// Copyright Epic Games, Inc. All Rights Reserved.
+
+#include "SOptimusEditorGraphNode_Comment.h"
+#include "EdGraphNode_Comment.h"
+#include "OptimusEditorGraph.h"
+#include "OptimusEditorHelpers.h"
+#include "OptimusEditorGraphNode_Comment.h"
+#include "OptimusNode_Comment.h"
+
+#define LOCTEXT_NAMESPACE "SOptimusEditorGraphNode_Comment"
+
+void SOptimusEditorGraphNode_Comment::Construct(const FArguments& InArgs, UOptimusEditorGraphNode_Comment* InGraphNode)
+{
+	InGraphNode->GetOnSizeChanged().BindSP(this, &SOptimusEditorGraphNode_Comment::OnSizeChanged);
+	InGraphNode->GetOnPositionChanged().BindSP(this, &SOptimusEditorGraphNode_Comment::OnPositionChanged);
+	
+	SGraphNodeComment::Construct({}, InGraphNode);
+}
+
+FReply SOptimusEditorGraphNode_Comment::OnMouseButtonUp(const FGeometry& MyGeometry, const FPointerEvent& MouseEvent)
+{
+	FReply Reply = SGraphNodeComment::OnMouseButtonUp(MyGeometry, MouseEvent);
+	
+	UOptimusNode_Comment* CommentNode = Cast<UOptimusNode_Comment>(OptimusEditor::FindModelNodeFromGraphNode(GraphNode));
+	if (ensure(CommentNode))
+	{
+		CommentNode->SetSize(UserSize);
+	}
+	
+	return Reply;
+}
+
+void SOptimusEditorGraphNode_Comment::EndUserInteraction() const
+{
+	
+	UOptimusEditorGraph* Graph = Cast<UOptimusEditorGraph>(GraphNode->GetGraph());
+	if (ensure(Graph))
+	{
+		Graph->HandleGraphNodeMoved();
+	}
+	
+	SGraphNodeComment::EndUserInteraction();
+}
+
+void SOptimusEditorGraphNode_Comment::OnSizeChanged()
+{
+	UserSize = FVector2D(GraphNode->NodeWidth, GraphNode->NodeHeight);
+}
+
+void SOptimusEditorGraphNode_Comment::OnPositionChanged()
+{
+	// Goal here is to use HandleSelection to update the nodes under the comment node after a undo
+	
+	// SGraphNodeComment::bIsSelected is private, so here goes the workaround to avoid changing the selection state
+	bool bIsGraphNodeSelected = false;
+	UOptimusEditorGraph* Graph = Cast<UOptimusEditorGraph>(GraphNode->GetGraph());
+	if (ensure(Graph))
+	{
+		 bIsGraphNodeSelected = Graph->GetSelectedNodes().Contains(GraphNode);
+	}
+
+	constexpr bool bUpdateNodesUnderComment = true;
+	HandleSelection(bIsGraphNodeSelected, bUpdateNodesUnderComment);
+}
+
+#undef LOCTEXT_NAMESPACE
diff --git a/Engine/Plugins/Animation/DeformerGraph/Source/OptimusEditor/Private/Widgets/SOptimusEditorGraphNode_Comment.h b/Engine/Plugins/Animation/DeformerGraph/Source/OptimusEditor/Private/Widgets/SOptimusEditorGraphNode_Comment.h
new file mode 100644
--- /dev/null
+++ b/Engine/Plugins/Animation/DeformerGraph/Source/OptimusEditor/Private/Widgets/SOptimusEditorGraphNode_Comment.h
@@ -0,0 +1,27 @@
+// Copyright Epic Games, Inc. All Rights Reserved.
+
+#pragma once
+
+#include "SGraphNodeComment.h"
+
+class UOptimusEditorGraphNode_Comment;
+class UEdGraphNode_Comment;
+
+
+
+class SOptimusEditorGraphNode_Comment : 
+	public SGraphNodeComment
+{
+public:
+	SLATE_BEGIN_ARGS(SOptimusEditorGraphNode_Comment) {}
+	SLATE_END_ARGS()
+
+	void Construct(const FArguments& InArgs, UOptimusEditorGraphNode_Comment* InGraphNode);
+
+	FReply OnMouseButtonUp(const FGeometry& MyGeometry, const FPointerEvent& MouseEvent) override;
+	void EndUserInteraction() const override;
+private:
+	void OnSizeChanged();
+	void OnPositionChanged();
+
+};
