DeformerGraph Fix Scene data interface not being valid in Editor preview.
We want it to be valid but to return 0 for Time.
This fixes some deformers not showing in Editor preview.

diff --git a/Engine/Plugins/Animation/DeformerGraph/Source/OptimusCore/Private/DataInterfaces/OptimusDataInterfaceScene.cpp b/Engine/Plugins/Animation/DeformerGraph/Source/OptimusCore/Private/DataInterfaces/OptimusDataInterfaceScene.cpp
--- a/Engine/Plugins/Animation/DeformerGraph/Source/OptimusCore/Private/DataInterfaces/OptimusDataInterfaceScene.cpp
+++ b/Engine/Plugins/Animation/DeformerGraph/Source/OptimusCore/Private/DataInterfaces/OptimusDataInterfaceScene.cpp
@@ -105,9 +105,9 @@ FOptimusSceneDataProviderProxy::FOptimusSceneDataProviderProxy(USceneComponent*
 	bool bUseSceneTime = bIsValid;
 #if WITH_EDITOR
 	// Don't tick time in Editor unless in PIE.
-	if (GIsEditor && bIsValid)
+	if (GIsEditor && bUseSceneTime)
 	{
-		bIsValid &= (bUseSceneTime && SceneComponent->GetWorld() != nullptr && SceneComponent->GetWorld()->WorldType != EWorldType::Editor);
+		bUseSceneTime = SceneComponent->GetWorld() != nullptr && SceneComponent->GetWorld()->WorldType != EWorldType::Editor;
 	}
 #endif
 	GameTime = bUseSceneTime ? SceneComponent->GetWorld()->TimeSeconds : 0;
