[Optimus Readback] added validity checks to attribute containers that the processor writes to.

diff --git a/Engine/Plugins/Animation/DeformerGraph/Source/OptimusCore/Private/OptimusGeometryReadbackProcessor.cpp b/Engine/Plugins/Animation/DeformerGraph/Source/OptimusCore/Private/OptimusGeometryReadbackProcessor.cpp
--- a/Engine/Plugins/Animation/DeformerGraph/Source/OptimusCore/Private/OptimusGeometryReadbackProcessor.cpp
+++ b/Engine/Plugins/Animation/DeformerGraph/Source/OptimusCore/Private/OptimusGeometryReadbackProcessor.cpp
@@ -139,7 +139,7 @@ void FOptimusGeometryReadbackProcessor::ProcessCompletedGeometryReadback_RenderT
 							}
 							
 							// Positions
-							if (Positions.Num() > 0)
+							if (Positions.Num() > 0 && ensure(PositionAttribute.IsValid()))
 							{
 								const float* Ptr = reinterpret_cast<const float*>(Positions.GetData());
 							
@@ -160,11 +160,20 @@ void FOptimusGeometryReadbackProcessor::ProcessCompletedGeometryReadback_RenderT
 							
 								for (const FVertexInstanceID VertexInstanceID : VertexInstances)
 								{
-									NormalAttribute[VertexInstanceID] = FVector3f(Normal);
+									if (ensure(NormalAttribute.IsValid()))
+									{
+										NormalAttribute[VertexInstanceID] = FVector3f(Normal);
+									}
 			
-									TangentAttribute[VertexInstanceID] = Tangent;	
-
-									BinormalSignAttribute[VertexInstanceID] = Normal.W;
+									if (ensure(TangentAttribute.IsValid()))
+									{
+										TangentAttribute[VertexInstanceID] = Tangent;
+									}
+									
+									if (BinormalSignAttribute.IsValid())
+									{
+										BinormalSignAttribute[VertexInstanceID] = Normal.W;
+									}
 								}
 							}
 
