[GeometryCacheLevelSequenceBaker] Fixed crash when SKM has MinLOD = some auto generated LOD level.
-also make sure we only record the component transform after it is computed by the constraint system

diff --git a/Engine/Plugins/Animation/DeformerGraph/Source/OptimusCore/Private/OptimusGeometryReadbackProcessor.cpp b/Engine/Plugins/Animation/DeformerGraph/Source/OptimusCore/Private/OptimusGeometryReadbackProcessor.cpp
--- a/Engine/Plugins/Animation/DeformerGraph/Source/OptimusCore/Private/OptimusGeometryReadbackProcessor.cpp
+++ b/Engine/Plugins/Animation/DeformerGraph/Source/OptimusCore/Private/OptimusGeometryReadbackProcessor.cpp
@@ -47,159 +47,35 @@ void FOptimusGeometryReadbackProcessor::ProcessCompletedGeometryReadback_RenderT
 		
 		auto ReadbackGeometryToMeshDescription = [Readback]()
 		{
+			const int32 LODIndex = Readback->LodIndex;
+			const TArray<uint8>& Positions = Readback->Position.ReadbackData; 
+			const TArray<uint8>& NormalsTangents = Readback->Tangent.ReadbackData; 
+			const TArray<uint8>& Colors = Readback->Color.ReadbackData; 
+
+			const int32 NumVertPositions = Positions.Num() / SizeOfVertPosition;
+			const int32 NumVertNormalsTangents= NormalsTangents.Num() / SizeOfVertTangents;
+			const int32 NumVertColors = Colors.Num() / SizeOfVertColor;
+
+			int32 NumVertices = 0;
+			// Make sure sizes of all buffers make sense
+			if (ensure(NumVertPositions == NumVertNormalsTangents || NumVertPositions == 0 || NumVertNormalsTangents == 0) && 
+				ensure(NumVertNormalsTangents == NumVertColors || NumVertNormalsTangents == 0 || NumVertColors == 0) &&
+				ensure(NumVertPositions == NumVertColors || NumVertPositions == 0 || NumVertColors == 0))
 			{
-				const int32 LODIndex = Readback->LodIndex;
-				const TArray<uint8>& Positions = Readback->Position.ReadbackData; 
-				const TArray<uint8>& NormalsTangents = Readback->Tangent.ReadbackData; 
-				const TArray<uint8>& Colors = Readback->Color.ReadbackData; 
+				NumVertices =
+					NumVertPositions != 0 ? NumVertPositions :
+					NumVertNormalsTangents != 0 ? NumVertNormalsTangents :
+					NumVertColors != 0 ? NumVertColors : 0;
+			}
 
-				if (const USkeletalMesh* SkeletalMesh = Readback->SkeletalMesh.Get())
-				{
-					FMeshDescription MeshDescription;
-					SkeletalMesh->CloneMeshDescription(LODIndex, MeshDescription);
-					
-					// Avoid conditional build during parallel for loop below
-					MeshDescription.BuildVertexIndexers();
-
-					constexpr int32 SizeOfVertPosition = FSkeletalMeshDeformerHelpers::PosBufferElementMultiplier * FSkeletalMeshDeformerHelpers::PosBufferBytesPerElement; // float3, 3 * float, 12 Bytes per vert
-					constexpr int32 SizeOfVertTangents = FSkeletalMeshDeformerHelpers::TangentBufferElementMultiplier * FSkeletalMeshDeformerHelpers::TangentBufferBytesPerElement; // 2 * half4, which is 4 * 2 bytes, 16 bytes per vert
-					constexpr int32 SizeOfVertColor = FSkeletalMeshDeformerHelpers::ColorBufferBytesPerElement; // 1 * 4 * 8bit, 4 bytes per vert
-
-					const int32 NumVertPositions = Positions.Num() / SizeOfVertPosition;
-					const int32 NumVertNormalsTangents= NormalsTangents.Num() / SizeOfVertTangents;
-					const int32 NumVertColors = Colors.Num() / SizeOfVertColor;
-
-					// Make sure sizes of all buffers make sense
-					if (ensure(NumVertPositions == NumVertNormalsTangents || NumVertPositions == 0 || NumVertNormalsTangents == 0) && 
-						ensure(NumVertNormalsTangents == NumVertColors || NumVertNormalsTangents == 0 || NumVertColors == 0) &&
-						ensure(NumVertPositions == NumVertColors || NumVertPositions == 0 || NumVertColors == 0))
-					{
-						const int32 NumVertices =
-							NumVertPositions != 0 ? NumVertPositions :
-							NumVertNormalsTangents != 0 ? NumVertNormalsTangents :
-							NumVertColors != 0 ? NumVertColors : 0;
-						TMeshAttributesRef<FVertexID, FVector3f> PositionAttribute =
-							MeshDescription.VertexAttributes().GetAttributesRef<FVector3f>(MeshAttribute::Vertex::Position);
-						TMeshAttributesRef<FVertexInstanceID, FVector3f> NormalAttribute =
-							MeshDescription.VertexInstanceAttributes().GetAttributesRef<FVector3f>(MeshAttribute::VertexInstance::Normal);
-						TMeshAttributesRef<FVertexInstanceID, FVector3f> TangentAttribute = 
-							MeshDescription.VertexInstanceAttributes().GetAttributesRef<FVector3f>(MeshAttribute::VertexInstance::Tangent);
-						TMeshAttributesRef<FVertexInstanceID, float> BinormalSignAttribute = 
-							MeshDescription.VertexInstanceAttributes().GetAttributesRef<float>(MeshAttribute::VertexInstance::BinormalSign);
-						TMeshAttributesRef<FVertexID, FVector4f> ColorAttribute =
-							MeshDescription.VertexAttributes().GetAttributesRef<FVector4f>(MeshAttribute::VertexInstance::Color);
-						
-						FSkeletalMeshLODModel& LODModel = SkeletalMesh->GetImportedModel()->LODModels[LODIndex];
-
-	#if USE_MESH_TO_IMPORT_VERTEX_MAP
-						const TArray<int32>& RenderToImportedMap = LODModel.MeshToImportVertexMap;
-						uint32 MaxImportVertex = LODModel.MaxImportVertex;
-	#else
-						const TArray<uint32>& RenderToImportedMap = LODModel.GetRawPointIndices();
-						uint32 MaxImportVertex = 0;
-						for (uint32 Raw : RenderToImportedMap)
-						{
-							MaxImportVertex = FMath::Max(MaxImportVertex, Raw);
-						}
-	#endif
-						
-						TArray<std::atomic<int32>> ImportedVertWriteCount;
-						ImportedVertWriteCount.SetNum(MaxImportVertex+1);
-						for(std::atomic<int32>& Counter : ImportedVertWriteCount)
-						{
-							Counter = 0;
-						}
-						
-						// Parallel for each render vert
-						ParallelFor(NumVertices,
-							[
-							&Readback,
-							&MeshDescription,
-							&PositionAttribute, 
-							&NormalAttribute, 
-							&TangentAttribute, 
-							&BinormalSignAttribute,
-							&ColorAttribute, 
-							NumVertices,
-							&RenderToImportedMap, 
-							&Positions, 
-							&NormalsTangents,
-							&Colors,
-							&ImportedVertWriteCount
-							](int32 RenderVertIndex)
-						{
-							// int32 ImportedVertexIndex = RenderToImportedMap[RenderVertIndex];
-							uint32 ImportedVertexIndex = RenderToImportedMap[RenderVertIndex];
-							const FVertexID ImportedVertexID(ImportedVertexIndex);
-
-							// Only write to each imported vert once
-							if (++ImportedVertWriteCount[ImportedVertexIndex] > 1)
-							{
-								return;
-							}
-							
-							// Positions
-							if (Positions.Num() > 0 && ensure(PositionAttribute.IsValid()))
-							{
-								const float* Ptr = reinterpret_cast<const float*>(Positions.GetData());
-							
-								const float PositionX = Ptr[RenderVertIndex*3 + 0];
-								const float PositionY = Ptr[RenderVertIndex*3 + 1];
-								const float PositionZ = Ptr[RenderVertIndex*3 + 2];
-
-								PositionAttribute[ImportedVertexID] = FVector3f(PositionX,PositionY,PositionZ);
-							}
-
-							// Normals Tangents
-							const TArrayView<const FVertexInstanceID> VertexInstances = MeshDescription.GetVertexVertexInstanceIDs(ImportedVertexID);
-							if (NormalsTangents.Num() > 0)
-							{
-								const FPackedRGBA16N* Ptr = reinterpret_cast<const FPackedRGBA16N*>(NormalsTangents.GetData());
-								const FVector3f Tangent = Ptr[RenderVertIndex*2 + 0].ToFVector3f();
-								const FVector4f Normal = Ptr[RenderVertIndex*2 + 1].ToFVector4f();
-							
-								for (const FVertexInstanceID VertexInstanceID : VertexInstances)
-								{
-									if (ensure(NormalAttribute.IsValid()))
-									{
-										NormalAttribute[VertexInstanceID] = FVector3f(Normal);
-									}
+			// Readback has no valid data,
+			if (!ensure(NumVertices != 0))
+			{
+				return;
+			}
 			
-									if (ensure(TangentAttribute.IsValid()))
-									{
-										TangentAttribute[VertexInstanceID] = Tangent;
-									}
-									
-									if (BinormalSignAttribute.IsValid())
-									{
-										BinormalSignAttribute[VertexInstanceID] = Normal.W;
-									}
-								}
-							}
-
-							// Colors
-							if (Colors.Num() > 0 && ColorAttribute.IsValid())
-							{
-								check(sizeof(FColor) == 4);
-							
-								const FColor* Ptr = reinterpret_cast<const FColor*>(Colors.GetData());
-
-								for (const FVertexInstanceID VertexInstanceID : VertexInstances)
-								{
-									ColorAttribute[VertexInstanceID] = FVector4f(Ptr[RenderVertIndex]);
-								}
-							}
-						});
-
-						for (TUniquePtr<FMeshDeformerGeometryReadbackRequest>& Request : Readback->GeometryReadbackRequests)
-						{
-							Request->MeshDescriptionCallback_AnyThread(MeshDescription);
-							Request->bHandled = true;
-						}
-					}
-				}
-				
-			}	
+			ProcessReadbackRequestingMeshDescription(Readback,NumVertices);
+			ProcessReadbackRequestingVertexDataArray(Readback,NumVertices);
 		};
 
 		// Readback data should be processed sequentially on a background thread, so using last task as prerequisite
@@ -214,4 +90,252 @@ void FOptimusGeometryReadbackProcessor::ProcessCompletedGeometryReadback_RenderT
 	}	
 }
 
+void FOptimusGeometryReadbackProcessor::ProcessReadbackRequestingMeshDescription(const TSharedPtr<FGeometryReadback>& Readback, int32 NumVertices)
+{
+	bool bNeedMeshDescription = false;
+	for (TUniquePtr<FMeshDeformerGeometryReadbackRequest>& Request : Readback->GeometryReadbackRequests)
+	{
+		if (Request->MeshDescriptionCallback_AnyThread.IsSet())
+		{
+			bNeedMeshDescription = true;	
+		}
+	}
+
+	if (!bNeedMeshDescription)
+	{
+		return;
+	}
+
+	const int32 LODIndex = Readback->LodIndex;
+	const TArray<uint8>& Positions = Readback->Position.ReadbackData; 
+	const TArray<uint8>& NormalsTangents = Readback->Tangent.ReadbackData; 
+	const TArray<uint8>& Colors = Readback->Color.ReadbackData; 
+
+	const USkeletalMesh* SkeletalMesh = Readback->SkeletalMesh.Get();
+
+	if (!SkeletalMesh)
+	{
+		return;
+	}
+	
+	FMeshDescription MeshDescription;
+	SkeletalMesh->CloneMeshDescription(LODIndex, MeshDescription);
+	
+	// Avoid conditional build during parallel for loop below
+	MeshDescription.BuildVertexIndexers();
+
+	TMeshAttributesRef<FVertexID, FVector3f> PositionAttribute =
+		MeshDescription.VertexAttributes().GetAttributesRef<FVector3f>(MeshAttribute::Vertex::Position);
+	TMeshAttributesRef<FVertexInstanceID, FVector3f> NormalAttribute =
+		MeshDescription.VertexInstanceAttributes().GetAttributesRef<FVector3f>(MeshAttribute::VertexInstance::Normal);
+	TMeshAttributesRef<FVertexInstanceID, FVector3f> TangentAttribute = 
+		MeshDescription.VertexInstanceAttributes().GetAttributesRef<FVector3f>(MeshAttribute::VertexInstance::Tangent);
+	TMeshAttributesRef<FVertexInstanceID, float> BinormalSignAttribute = 
+		MeshDescription.VertexInstanceAttributes().GetAttributesRef<float>(MeshAttribute::VertexInstance::BinormalSign);
+	TMeshAttributesRef<FVertexID, FVector4f> ColorAttribute =
+		MeshDescription.VertexAttributes().GetAttributesRef<FVector4f>(MeshAttribute::VertexInstance::Color);
+	
+	FSkeletalMeshLODModel& LODModel = SkeletalMesh->GetImportedModel()->LODModels[LODIndex];
+
+#if USE_MESH_TO_IMPORT_VERTEX_MAP
+	const TArray<int32>& RenderToImportedMap = LODModel.MeshToImportVertexMap;
+	uint32 MaxImportVertex = LODModel.MaxImportVertex;
+#else
+	const TArray<uint32>& RenderToImportedMap = LODModel.GetRawPointIndices();
+	uint32 MaxImportVertex = 0;
+	for (uint32 Raw : RenderToImportedMap)
+	{
+		MaxImportVertex = FMath::Max(MaxImportVertex, Raw);
+	}
+#endif
+	
+	TArray<std::atomic<int32>> ImportedVertWriteCount;
+	ImportedVertWriteCount.SetNum(MaxImportVertex+1);
+	for(std::atomic<int32>& Counter : ImportedVertWriteCount)
+	{
+		Counter = 0;
+	}
+	
+	// Parallel for each render vert
+	ParallelFor(NumVertices,
+		[
+		&Readback,
+		&MeshDescription,
+		&PositionAttribute, 
+		&NormalAttribute, 
+		&TangentAttribute, 
+		&BinormalSignAttribute,
+		&ColorAttribute,
+		NumVertices,
+		&RenderToImportedMap, 
+		&Positions, 
+		&NormalsTangents,
+		&Colors,
+		&ImportedVertWriteCount
+		](int32 RenderVertIndex)
+	{
+		// int32 ImportedVertexIndex = RenderToImportedMap[RenderVertIndex];
+		uint32 ImportedVertexIndex = RenderToImportedMap[RenderVertIndex];
+		const FVertexID ImportedVertexID(ImportedVertexIndex);
+
+		// Only write to each imported vert once
+		if (++ImportedVertWriteCount[ImportedVertexIndex] > 1)
+		{
+			return;
+		}
+		
+		// Positions
+		if (Positions.Num() > 0 && PositionAttribute.IsValid() && PositionAttribute.GetNumElements() == Positions.Num())
+		{
+			const float* Ptr = reinterpret_cast<const float*>(Positions.GetData());
+		
+			const float PositionX = Ptr[RenderVertIndex*3 + 0];
+			const float PositionY = Ptr[RenderVertIndex*3 + 1];
+			const float PositionZ = Ptr[RenderVertIndex*3 + 2];
+
+			PositionAttribute[ImportedVertexID] = FVector3f(PositionX,PositionY,PositionZ);
+		}
+
+		// Normals Tangents
+		const TArrayView<const FVertexInstanceID> VertexInstances = MeshDescription.GetVertexVertexInstanceIDs(ImportedVertexID);
+		if (NormalsTangents.Num() > 0)
+		{
+			const FPackedRGBA16N* Ptr = reinterpret_cast<const FPackedRGBA16N*>(NormalsTangents.GetData());
+			const FVector3f Tangent = Ptr[RenderVertIndex*2 + 0].ToFVector3f();
+			const FVector4f Normal = Ptr[RenderVertIndex*2 + 1].ToFVector4f();
+		
+			for (const FVertexInstanceID VertexInstanceID : VertexInstances)
+			{
+				if (ensure(NormalAttribute.IsValid()))
+				{
+					NormalAttribute[VertexInstanceID] = FVector3f(Normal);
+				}
+
+				if (ensure(TangentAttribute.IsValid()))
+				{
+					TangentAttribute[VertexInstanceID] = Tangent;
+				}
+				
+				if (BinormalSignAttribute.IsValid())
+				{
+					BinormalSignAttribute[VertexInstanceID] = Normal.W;
+				}
+			}
+		}
+
+		// Colors
+		if (Colors.Num() > 0 && ColorAttribute.IsValid())
+		{
+			check(sizeof(FColor) == 4);
+		
+			const FColor* Ptr = reinterpret_cast<const FColor*>(Colors.GetData());
+
+			for (const FVertexInstanceID VertexInstanceID : VertexInstances)
+			{
+				ColorAttribute[VertexInstanceID] = FVector4f(Ptr[RenderVertIndex]);
+			}
+		}
+	});
+
+	for (TUniquePtr<FMeshDeformerGeometryReadbackRequest>& Request : Readback->GeometryReadbackRequests)
+	{
+		if (Request->MeshDescriptionCallback_AnyThread.IsSet())
+		{
+			Request->MeshDescriptionCallback_AnyThread(MeshDescription);
+			Request->bMeshDescriptionHandled = true;
+		}
+	}
+}
+
+void FOptimusGeometryReadbackProcessor::ProcessReadbackRequestingVertexDataArray(const TSharedPtr<FGeometryReadback>& Readback, int32 NumVertices)
+{
+	bool bNeedVertDataArrays= false;
+	for (TUniquePtr<FMeshDeformerGeometryReadbackRequest>& Request : Readback->GeometryReadbackRequests)
+	{
+		if (Request->VertexDataArraysCallback_AnyThread.IsSet())
+		{
+			bNeedVertDataArrays = true;
+		}
+	}
+
+	if (!bNeedVertDataArrays)
+	{
+		return;
+	}
+	
+	const TArray<uint8>& Positions = Readback->Position.ReadbackData; 
+	const TArray<uint8>& NormalsTangents = Readback->Tangent.ReadbackData; 
+	const TArray<uint8>& Colors = Readback->Color.ReadbackData; 
+
+	const USkeletalMesh* SkeletalMesh = Readback->SkeletalMesh.Get();
+
+	if (!SkeletalMesh)
+	{
+		return;
+	}
+
+	FMeshDeformerGeometryReadbackVertexDataArrays VertexDataArrays;
+
+	VertexDataArrays.LODIndex = Readback->LodIndex;
+	VertexDataArrays.Positions.SetNum(NumVertices);
+	VertexDataArrays.Normals.SetNum(NumVertices);
+	VertexDataArrays.Tangents.SetNum(NumVertices);
+	VertexDataArrays.Colors.SetNum(NumVertices);
+	
+	// Parallel for each render vert
+	ParallelFor(NumVertices,
+		[
+		&Readback,
+		NumVertices,
+		&Positions, 
+		&NormalsTangents,
+		&Colors,
+		&VertexDataArrays
+		](int32 RenderVertIndex)
+	{
+		
+		// Positions
+		if (Positions.Num() > 0)
+		{
+			const float* Ptr = reinterpret_cast<const float*>(Positions.GetData());
+		
+			const float PositionX = Ptr[RenderVertIndex*3 + 0];
+			const float PositionY = Ptr[RenderVertIndex*3 + 1];
+			const float PositionZ = Ptr[RenderVertIndex*3 + 2];
+
+			VertexDataArrays.Positions[RenderVertIndex] = FVector3f(PositionX,PositionY,PositionZ);
+		}
+
+		// Normals Tangents
+		if (NormalsTangents.Num() > 0)
+		{
+			const FPackedRGBA16N* Ptr = reinterpret_cast<const FPackedRGBA16N*>(NormalsTangents.GetData());
+			const FVector3f Tangent = Ptr[RenderVertIndex*2 + 0].ToFVector3f();
+			const FVector4f Normal = Ptr[RenderVertIndex*2 + 1].ToFVector4f();
+		
+			VertexDataArrays.Normals[RenderVertIndex] = Normal;	
+			VertexDataArrays.Tangents[RenderVertIndex] = Tangent;	
+		}
+
+		// Colors
+		if (Colors.Num() > 0)
+		{
+			check(sizeof(FColor) == 4);
+		
+			const FColor* Ptr = reinterpret_cast<const FColor*>(Colors.GetData());
+			VertexDataArrays.Colors[RenderVertIndex] = FVector4f(Ptr[RenderVertIndex]);
+		}
+	});
+
+	for (TUniquePtr<FMeshDeformerGeometryReadbackRequest>& Request : Readback->GeometryReadbackRequests)
+	{
+		if (Request->VertexDataArraysCallback_AnyThread.IsSet())
+		{
+			Request->VertexDataArraysCallback_AnyThread(VertexDataArrays);
+			Request->bVertexDataArraysHandled= true;
+		}
+	}
+
+}
+
 #endif // WITH_EDITORONLY_DATA
diff --git a/Engine/Plugins/Animation/DeformerGraph/Source/OptimusCore/Private/OptimusGeometryReadbackProcessor.h b/Engine/Plugins/Animation/DeformerGraph/Source/OptimusCore/Private/OptimusGeometryReadbackProcessor.h
--- a/Engine/Plugins/Animation/DeformerGraph/Source/OptimusCore/Private/OptimusGeometryReadbackProcessor.h
+++ b/Engine/Plugins/Animation/DeformerGraph/Source/OptimusCore/Private/OptimusGeometryReadbackProcessor.h
@@ -1,6 +1,7 @@
 // Copyright Epic Games, Inc. All Rights Reserved.
 
 #pragma once
+#include "SkeletalMeshDeformerHelpers.h"
 #include "ComputeFramework/ComputeDataProvider.h"
 #include "Engine/SkeletalMesh.h"
 #include "Containers/MpscQueue.h"
@@ -36,9 +37,18 @@ public:
 	void ProcessCompletedGeometryReadback_RenderThread();
 
 private:
+
+	static void ProcessReadbackRequestingMeshDescription(const TSharedPtr<FGeometryReadback>& Readback, int32 NumVertices);
+	static void ProcessReadbackRequestingVertexDataArray(const TSharedPtr<FGeometryReadback>& Readback, int32 NumVertices);
 	
 	// All data provider shares the same processor singleton
 	TMpscQueue<TSharedPtr<FGeometryReadback>> GeometryReadbacks;
-	FGraphEventRef LastReadbackProcessingTask_RenderThread;	
+	FGraphEventRef LastReadbackProcessingTask_RenderThread;
+	
+	static constexpr int32 SizeOfVertPosition = FSkeletalMeshDeformerHelpers::PosBufferElementMultiplier * FSkeletalMeshDeformerHelpers::PosBufferBytesPerElement; // float3, 3 * float, 12 Bytes per vert
+	static constexpr int32 SizeOfVertTangents = FSkeletalMeshDeformerHelpers::TangentBufferElementMultiplier * FSkeletalMeshDeformerHelpers::TangentBufferBytesPerElement; // 2 * half4, which is 4 * 2 bytes, 16 bytes per vert
+	static constexpr int32 SizeOfVertColor = FSkeletalMeshDeformerHelpers::ColorBufferBytesPerElement; // 1 * 4 * 8bit, 4 bytes per vert
+
+
 };
 #endif
diff --git a/Engine/Plugins/Experimental/GeometryCacheLevelSequenceBaker/Source/GeometryCacheLevelSequenceBaker/Private/GeometryCacheLevelSequenceBaker.cpp b/Engine/Plugins/Experimental/GeometryCacheLevelSequenceBaker/Source/GeometryCacheLevelSequenceBaker/Private/GeometryCacheLevelSequenceBaker.cpp
--- a/Engine/Plugins/Experimental/GeometryCacheLevelSequenceBaker/Source/GeometryCacheLevelSequenceBaker/Private/GeometryCacheLevelSequenceBaker.cpp
+++ b/Engine/Plugins/Experimental/GeometryCacheLevelSequenceBaker/Source/GeometryCacheLevelSequenceBaker/Private/GeometryCacheLevelSequenceBaker.cpp
@@ -38,10 +38,6 @@
 
 #define LOCTEXT_NAMESPACE "GeometryCacheLevelSequenceBaker"
 
-// At the moment there are meshes with MeshToImportVertexMap not matching their mesh description for some reason
-// Once the root cause is determined we can remove this
-#define USE_MESH_TO_IMPORT_VERTEX_MAP 0
-
 const int32 FGeometryCacheLevelSequenceBaker::LODIndexToBake = 0;
 const float FGeometryCacheLevelSequenceBaker::TotalAmountOfWork = 1.0f;
 const float FGeometryCacheLevelSequenceBaker::AmountOfWorkGatherStage = TotalAmountOfWork / 2.0f;
@@ -220,6 +216,7 @@ void FGeometryCacheLevelSequenceBaker::Bake(TSharedRef<ISequencer> InSequencer)
 	BakeTask->FixedDeltaTimeScope = MakeUnique<FEngineFixedDeltaTimeScope>(DeltaTime);
 	BakeTask->SequencerStateScope = MakeUnique<FSequencerSettingScope>(Sequencer);
 	BakeTask->ConsoleVariableOverrideScope = MakeUnique<FConsoleVariableOverrideScope>();
+	BakeTask->WorldTickEndDelegateScope = MakeUnique<FWorldTickEndDelegateScope>();
 }
 
 TArray<FGuid> FGeometryCacheLevelSequenceBaker::GetBindingsToBake(TSharedRef<ISequencer> InSequencer)
@@ -380,7 +377,10 @@ void FGeometryCacheLevelSequenceBaker::Tick(float DeltaTime)
 
 		if (!CurrentBakeTask->IsSequencerPlaying())
 		{
-			check(CurrentBakeTask->CurrentSampleIndex == CurrentBakeTask->NumSamples);
+			if (CurrentBakeTask->CurrentSampleIndex < CurrentBakeTask->NumSamples)
+			{
+				return;
+			}
 			
 			CurrentBakeTask->CurrentSampleIndex = 0;
 			CurrentBakeTask->GatherStageComponentSettingScopes.Reset();
@@ -402,7 +402,11 @@ void FGeometryCacheLevelSequenceBaker::Tick(float DeltaTime)
 		
 		if (!CurrentBakeTask->IsSequencerPlaying())
 		{
-			check(CurrentBakeTask->CurrentSampleIndex == CurrentBakeTask->NumSamples);
+			if (CurrentBakeTask->CurrentSampleIndex < CurrentBakeTask->NumSamples)
+			{
+				return;
+			}
+
 			CurrentBakeTask->Stage = EStage::WriteToAsset;
 		}
 	}
@@ -423,7 +427,8 @@ void FGeometryCacheLevelSequenceBaker::Tick(float DeltaTime)
 	}
 }
 
-void FGeometryCacheLevelSequenceBaker::OnSequencerEvaluated()
+
+void FGeometryCacheLevelSequenceBaker::OnWorldTickEnd(UWorld*, ELevelTick, float)
 {
 	if (CurrentBakeTask->Stage == EStage::Gather)
 	{
@@ -472,15 +477,6 @@ void FGeometryCacheLevelSequenceBaker::SetupComponentBakeTasks()
 			ComponentTask.ComponentInfo = ComponentInfo.Value;
 			ComponentTask.NumSamplesPending = CurrentBakeTask->NumSamples;
 			ComponentTask.GeometrySamples.Init({}, ComponentTask.NumSamplesPending);
-		
-			const FSkeletalMeshLODModel& LODModel = ComponentTask.ComponentInfo.SkeletalMeshAsset->GetImportedModel()->LODModels[LODIndexToBake];
-			for (FFrameData& Sample : ComponentTask.GeometrySamples)
-			{
-				Sample.Positions.Init({}, LODModel.NumVertices);
-				Sample.Normals.Init({}, LODModel.NumVertices);
-				Sample.TangentsX.Init({}, LODModel.NumVertices);
-			}
-			
 			ComponentTask.VisibilitySamples.Init({}, ComponentTask.NumSamplesPending);
 		}
 	}
@@ -634,12 +630,13 @@ void FGeometryCacheLevelSequenceBaker::RequestReadback()
 
 			FTransform Transform = Component->GetComponentTransform();
 			TUniquePtr<FMeshDeformerGeometryReadbackRequest> Request = MakeUnique<FMeshDeformerGeometryReadbackRequest>();
-			Request->MeshDescriptionCallback_AnyThread =
+
+			Request->VertexDataArraysCallback_AnyThread =
 				[
 					Transform,
 					SampleIndex = CurrentBakeTask->CurrentSampleIndex,
 					TaskIndex
-				](const FMeshDescription& MeshDescription)
+				](const FMeshDeformerGeometryReadbackVertexDataArrays& VertexDataArrays)
 				{
 					// Making sure current bake task is alive when this callback is invoked on a worker thread
 					FScopeLock BakeTaskLifeTimeLock(&Get().CurrentBakeTaskLifeTimeCriticalSection);
@@ -649,61 +646,44 @@ void FGeometryCacheLevelSequenceBaker::RequestReadback()
 						return;
 					}
 
+					int32 NumVertices = VertexDataArrays.Positions.Num();
+					
 					FComponentTask& ComponentTask = BakeTask->ComponentTasks[TaskIndex];
-
-					bool bMeshAvailable = false;
-					if (!MeshDescription.IsEmpty())
+					FFrameData& FrameData = ComponentTask.GeometrySamples[SampleIndex];
+					
+					bool bMeshAvailable = true;
+					if (VertexDataArrays.LODIndex == INDEX_NONE || NumVertices == 0)
+					{
+						bMeshAvailable = false;
+					}
+					else if (ComponentTask.ActualLODIndexBaked != INDEX_NONE && ComponentTask.ActualLODIndexBaked != VertexDataArrays.LODIndex) // LOD changed during bake, not supported
+					{
+						bMeshAvailable = false;
+					}
+					
+					if (bMeshAvailable)
 					{
-						bMeshAvailable = true;
+						 if (ComponentTask.ActualLODIndexBaked == INDEX_NONE)
+						{
+							ComponentTask.ActualLODIndexBaked = VertexDataArrays.LODIndex;
+						}
 						
-						const USkeletalMesh* MeshAsset = ComponentTask.ComponentInfo.SkeletalMeshAsset;
-						const FSkeletalMeshLODModel& LODModel = MeshAsset->GetImportedModel()->LODModels[LODIndexToBake];
-#if USE_MESH_TO_IMPORT_VERTEX_MAP
-						const TArray<int32>& RenderToImportedMap = LODModel.MeshToImportVertexMap;
-#else
-						const TArray<uint32>& RenderToImportedMap = LODModel.GetRawPointIndices();
-#endif
+						FrameData.Positions.Init({}, NumVertices);
+						FrameData.Normals.Init({}, NumVertices);
+						FrameData.TangentsX.Init({}, NumVertices);
 						
-						TMeshAttributesRef<FVertexID, const FVector3f> PositionAttribute =
-							MeshDescription.VertexAttributes().GetAttributesRef<FVector3f>(MeshAttribute::Vertex::Position);
-						TMeshAttributesRef<FVertexInstanceID, const FVector3f> NormalAttribute =
-							MeshDescription.VertexInstanceAttributes().GetAttributesRef<FVector3f>(MeshAttribute::VertexInstance::Normal);
-						TMeshAttributesRef<FVertexInstanceID, const FVector3f> TangentAttribute =
-							MeshDescription.VertexInstanceAttributes().GetAttributesRef<FVector3f>(MeshAttribute::VertexInstance::Tangent);
-
-						ParallelFor(LODModel.NumVertices,
-						[
-							&RenderToImportedMap,
-							&MeshDescription,
-							&PositionAttribute,
-							&NormalAttribute,
-							&TangentAttribute,
-							&FrameData = ComponentTask.GeometrySamples[SampleIndex],
-							&Transform
-						](int32 Index)
+						for (int32 Index = 0; Index < NumVertices; Index++)
 						{
-							FVertexID ID = RenderToImportedMap[Index];
-
-							FrameData.Positions[Index] = FVector3f(Transform.TransformPosition(FVector(PositionAttribute[ID])));
-
-							const TArrayView<const FVertexInstanceID> VertexInstances = MeshDescription.
-								GetVertexVertexInstanceIDs(ID);
-
-							// Some vertices for some reason have no vertex instances, dangling verts?
-							if (VertexInstances.Num() == 0)
+							FrameData.Positions[Index] = FVector3f(Transform.TransformPosition(FVector(VertexDataArrays.Positions[Index])));
+							if (VertexDataArrays.Normals.IsValidIndex(Index))
 							{
-								return;
+								FrameData.Normals[Index] = FVector3f(Transform.TransformVector(FVector(VertexDataArrays.Normals[Index])));
 							}
-							
-							// Normal
-							FVector3f Normal = NormalAttribute[VertexInstances[0]];
-
-							FrameData.Normals[Index] = FVector3f(Transform.TransformVector(FVector(Normal)));
-
-							// TangentX
-							FVector3f TangentX = TangentAttribute[VertexInstances[0]];
-							FrameData.TangentsX[Index] = FVector3f(Transform.TransformVector(FVector(TangentX)));
-						});	
+							if (VertexDataArrays.Tangents.IsValidIndex(Index))
+							{
+								FrameData.TangentsX[Index] = FVector3f(Transform.TransformVector(FVector(VertexDataArrays.Tangents[Index])));
+							}
+						}
 					}
 
 					OnReadbackResultConfirmed(ComponentTask, SampleIndex, bMeshAvailable);	
@@ -752,6 +732,8 @@ void FGeometryCacheLevelSequenceBaker::WriteToAsset()
 			continue;
 		}
 
+		int32 NumVertices = Task.GeometrySamples[FirstVisibleFrameIndex].Positions.Num();
+		
 		// Copy neighbor frames when visibility is changing such that frame interpolation works well
 		for (const FVisibilitySample& VisibilitySample : Task.VisibilitySamples)
 		{	
@@ -780,7 +762,22 @@ void FGeometryCacheLevelSequenceBaker::WriteToAsset()
 			Task.GeometrySamples[0] = Task.GeometrySamples[Task.VisibilitySamples[FirstVisibleFrameIndex].FrameIndex];
 		}
 
-		int32 TrackIndex = AddTrackWriterFromSkinnedAssetAndMaterials(Writer, *(Task.ComponentInfo.SkeletalMeshAsset), Task.ComponentInfo.Materials);
+		for (int32 Frame = 0; Frame < Task.GeometrySamples.Num(); ++Frame)
+		{
+			FFrameData& FrameData = Task.GeometrySamples[Frame];
+			
+			if (FrameData.Positions.Num() != Task.GeometrySamples[0].Positions.Num())
+			{
+				if (ensure(FrameData.Positions.Num() == 0))
+				{
+					FrameData.Positions.Init({}, NumVertices);
+					FrameData.Normals.Init({}, NumVertices);
+					FrameData.TangentsX.Init({}, NumVertices);	
+				}
+			}
+		}
+
+		int32 TrackIndex = AddTrackWriterFromSkinnedAssetAndMaterials(Writer, *(Task.ComponentInfo.SkeletalMeshAsset), Task.ActualLODIndexBaked, Task.ComponentInfo.Materials);
 
 		FTrackWriter& TrackWriter = Writer.GetTrackWriter(TrackIndex);
 
@@ -850,7 +847,10 @@ void FGeometryCacheLevelSequenceBaker::OnReadbackResultConfirmed(FComponentTask&
 	{
 		FFunctionGraphTask::CreateAndDispatchWhenReady([]()
 		{
-			Get().CurrentBakeTask->UpdateBakeProgress();
+			if (Get().CurrentBakeTask)
+			{
+				Get().CurrentBakeTask->UpdateBakeProgress();
+			}
 		}, TStatId(), NULL, ENamedThreads::GameThread);
 	}
 
@@ -882,8 +882,6 @@ FGeometryCacheLevelSequenceBaker::FSequencerSettingScope::FSequencerSettingScope
 	PreviousLoopMode = InSequencer->GetSequencerSettings()->GetLoopMode();
 	InSequencer->GetSequencerSettings()->SetLoopMode(ESequencerLoopMode::SLM_NoLoop);
 
-	OnGlobalTimeChangedDelegateHandle = InSequencer->OnGlobalTimeChanged().AddRaw(&Get(), &FGeometryCacheLevelSequenceBaker::OnSequencerEvaluated);
-
 	UMovieSceneSequence* MovieSceneSequence = InSequencer->GetFocusedMovieSceneSequence();
 	UMovieScene* MovieScene = MovieSceneSequence->GetMovieScene();
 
@@ -902,7 +900,6 @@ FGeometryCacheLevelSequenceBaker::FSequencerSettingScope::~FSequencerSettingScop
 		MovieScene->SetClockSource(PreviousClockSource);
 		Sequencer.Pin()->ResetTimeController();
 		
-		Sequencer.Pin()->OnGlobalTimeChanged().Remove(OnGlobalTimeChangedDelegateHandle);
 		Sequencer.Pin()->GetSequencerSettings()->SetLoopMode(PreviousLoopMode);
 		Sequencer.Pin()->SetGlobalTime(PreviousGlobalTime.Time);
 	}
@@ -938,6 +935,16 @@ FGeometryCacheLevelSequenceBaker::FConsoleVariableOverrideScope::~FConsoleVariab
 	}
 }
 
+FGeometryCacheLevelSequenceBaker::FWorldTickEndDelegateScope::FWorldTickEndDelegateScope()
+{
+	OnWorldTickEndDelegate = FWorldDelegates::OnWorldTickEnd.AddRaw(&Get(), &FGeometryCacheLevelSequenceBaker::OnWorldTickEnd);
+}
+
+FGeometryCacheLevelSequenceBaker::FWorldTickEndDelegateScope::~FWorldTickEndDelegateScope()
+{
+	FWorldDelegates::OnWorldTickEnd.Remove(OnWorldTickEndDelegate);
+}
+
 bool FGeometryCacheLevelSequenceBaker::FBakeTask::IsSequencerPlaying()
 {
 	return Sequencer.Pin()->GetPlaybackStatus() == EMovieScenePlayerStatus::Playing;
diff --git a/Engine/Plugins/Experimental/GeometryCacheLevelSequenceBaker/Source/GeometryCacheLevelSequenceBaker/Private/GeometryCacheLevelSequenceBaker.h b/Engine/Plugins/Experimental/GeometryCacheLevelSequenceBaker/Source/GeometryCacheLevelSequenceBaker/Private/GeometryCacheLevelSequenceBaker.h
--- a/Engine/Plugins/Experimental/GeometryCacheLevelSequenceBaker/Source/GeometryCacheLevelSequenceBaker/Private/GeometryCacheLevelSequenceBaker.h
+++ b/Engine/Plugins/Experimental/GeometryCacheLevelSequenceBaker/Source/GeometryCacheLevelSequenceBaker/Private/GeometryCacheLevelSequenceBaker.h
@@ -38,7 +38,7 @@ public:
 	static bool GetGeometryCacheAssetPathFromUser(FString& OutPackageName, FString& OutAssetName);
 
 	void Tick(float DeltaTime) override;
-	void OnSequencerEvaluated();
+	void OnWorldTickEnd(UWorld*, ELevelTick, float);
 	bool IsTickable() const override;
 	TStatId GetStatId() const override {RETURN_QUICK_DECLARE_CYCLE_STAT(FBlueprintActionDatabase, STATGROUP_Tickables);}
 	
@@ -85,7 +85,8 @@ private:
 		FGuid Binding;
 
 		FComponentInfo ComponentInfo;
-		
+
+		int32 ActualLODIndexBaked = INDEX_NONE;
 		TArray<FFrameData> GeometrySamples;
 		std::atomic<uint32> NumSamplesPending = 0;
 		TArray<FVisibilitySample> VisibilitySamples;
@@ -108,7 +109,6 @@ private:
 		TWeakPtr<ISequencer> Sequencer;
 		FQualifiedFrameTime PreviousGlobalTime;
 		ESequencerLoopMode PreviousLoopMode;
-		FDelegateHandle OnGlobalTimeChangedDelegateHandle;
 		EUpdateClockSource PreviousClockSource;
 		
 		FSequencerSettingScope(TSharedPtr<ISequencer> InSequencer);
@@ -122,6 +122,13 @@ private:
 		FConsoleVariableOverrideScope();
 		~FConsoleVariableOverrideScope();
 	};
+	
+	struct FWorldTickEndDelegateScope
+	{
+		FDelegateHandle OnWorldTickEndDelegate;
+		FWorldTickEndDelegateScope();
+		~FWorldTickEndDelegateScope();
+	};
 
 	enum class EStage
 	{
@@ -157,6 +164,7 @@ private:
 		TUniquePtr<FEngineFixedDeltaTimeScope> FixedDeltaTimeScope;
 		TUniquePtr<FSequencerSettingScope> SequencerStateScope;
 		TUniquePtr<FConsoleVariableOverrideScope> ConsoleVariableOverrideScope;
+		TUniquePtr<FWorldTickEndDelegateScope> WorldTickEndDelegateScope;
 
 		FDelegateHandle OnEndFrameDelegateHandle;
 
diff --git a/Engine/Plugins/Runtime/GeometryCache/Source/GeometryCache/Private/GeometryCacheConstantTopologyWriter.cpp b/Engine/Plugins/Runtime/GeometryCache/Source/GeometryCache/Private/GeometryCacheConstantTopologyWriter.cpp
--- a/Engine/Plugins/Runtime/GeometryCache/Source/GeometryCache/Private/GeometryCacheConstantTopologyWriter.cpp
+++ b/Engine/Plugins/Runtime/GeometryCache/Source/GeometryCache/Private/GeometryCacheConstantTopologyWriter.cpp
@@ -184,9 +184,8 @@ namespace UE::GeometryCacheHelpers
 			return BatchesInfo;
 		}
 
-		TArray<uint32> GetImportedVertexNumbers(const USkinnedAsset& Asset)
+		TArray<uint32> GetImportedVertexNumbers(const USkinnedAsset& Asset, int32 LODIndex)
 		{
-			constexpr int32 LODIndex = 0;
 			const FSkeletalMeshModel* const ImportedModel = Asset.GetImportedModel();
 			if (!ImportedModel || !ImportedModel->LODModels.IsValidIndex(LODIndex))
 			{
@@ -448,13 +447,13 @@ namespace UE::GeometryCacheHelpers
 
 	int32 AddTrackWriterFromSkinnedAsset(FGeometryCacheConstantTopologyWriter& Writer, const USkinnedAsset& Asset)
 	{
-		return AddTrackWriterFromSkinnedAssetAndMaterials(Writer, Asset, {});
+		constexpr int32 LODIndex = 0;
+		return AddTrackWriterFromSkinnedAssetAndMaterials(Writer, Asset, LODIndex, {});
 	}
 
-	int32 AddTrackWriterFromSkinnedAssetAndMaterials(FGeometryCacheConstantTopologyWriter& Writer, const USkinnedAsset& Asset, const TArray<TObjectPtr<UMaterialInterface>>& Materials)
+	int32 AddTrackWriterFromSkinnedAssetAndMaterials(FGeometryCacheConstantTopologyWriter& Writer, const USkinnedAsset& Asset, int32 LODIndex, const TArray<TObjectPtr<UMaterialInterface>>& Materials)
 	{
 		const FSkeletalMeshRenderData* RenderData = Asset.GetResourceForRendering();
-		constexpr int32 LODIndex = 0;
 		if (!RenderData || !RenderData->LODRenderData.IsValidIndex(LODIndex))
 		{
 			return INDEX_NONE;
@@ -472,7 +471,7 @@ namespace UE::GeometryCacheHelpers
 		TrackWriter.Indices = Indices;
 		TrackWriter.UVs = UVs;
 		TrackWriter.Colors = Colors;
-		TrackWriter.ImportedVertexNumbers = GetImportedVertexNumbers(Asset);
+		TrackWriter.ImportedVertexNumbers = GetImportedVertexNumbers(Asset, LODIndex);
 		const int32 MaterialOffset = Writer.GetNumMaterials();
 		TrackWriter.BatchesInfo = GetBatchesInfo(LODData, MaterialOffset);
 
diff --git a/Engine/Plugins/Runtime/GeometryCache/Source/GeometryCache/Public/GeometryCacheConstantTopologyWriter.h b/Engine/Plugins/Runtime/GeometryCache/Source/GeometryCache/Public/GeometryCacheConstantTopologyWriter.h
--- a/Engine/Plugins/Runtime/GeometryCache/Source/GeometryCache/Public/GeometryCacheConstantTopologyWriter.h
+++ b/Engine/Plugins/Runtime/GeometryCache/Source/GeometryCache/Public/GeometryCacheConstantTopologyWriter.h
@@ -139,7 +139,7 @@ namespace UE::GeometryCacheHelpers
 	 * @brief This will create a track writer and fill in the track writer's data (indices, UVs, .etc) from the skinned asset, with the option to specify a set of materials 
 	 * which typically come from a specific skeletal mesh component using this skinned asset.
 	 */
-	GEOMETRYCACHE_API int32 AddTrackWriterFromSkinnedAssetAndMaterials(FGeometryCacheConstantTopologyWriter& Writer, const USkinnedAsset& Asset, const TArray<TObjectPtr<UMaterialInterface>>& Materials);
+	GEOMETRYCACHE_API int32 AddTrackWriterFromSkinnedAssetAndMaterials(FGeometryCacheConstantTopologyWriter& Writer, const USkinnedAsset& Asset, int32 LODIndex, const TArray<TObjectPtr<UMaterialInterface>>& Materials);
 
 	/**
 	 * @brief This will create multiple track writers and fill in the track writer's data (indices, UVs, materials .etc) from the template geometry cache. 
diff --git a/Engine/Source/Runtime/Engine/Private/Animation/MeshDeformerGeometryReadback.cpp b/Engine/Source/Runtime/Engine/Private/Animation/MeshDeformerGeometryReadback.cpp
--- a/Engine/Source/Runtime/Engine/Private/Animation/MeshDeformerGeometryReadback.cpp
+++ b/Engine/Source/Runtime/Engine/Private/Animation/MeshDeformerGeometryReadback.cpp
@@ -5,12 +5,21 @@
 
 FMeshDeformerGeometryReadbackRequest::~FMeshDeformerGeometryReadbackRequest()
 {
-	if (!bHandled)
+	// Signal to the requester that the request was not successfully fulfilled
+	
+	if (MeshDescriptionCallback_AnyThread.IsSet())
 	{
-		if (MeshDescriptionCallback_AnyThread.IsSet())
+		if (!bMeshDescriptionHandled)
 		{
-			// Signal to the requester that the request was not successfully fullfilled
 			MeshDescriptionCallback_AnyThread(FMeshDescription());
 		}
 	}
+	
+	if (VertexDataArraysCallback_AnyThread.IsSet())
+	{
+		if (!bVertexDataArraysHandled)
+		{
+			VertexDataArraysCallback_AnyThread(FMeshDeformerGeometryReadbackVertexDataArrays());
+		}
+	}
 }
\ No newline at end of file
diff --git a/Engine/Source/Runtime/Engine/Public/Animation/MeshDeformerGeometryReadback.h b/Engine/Source/Runtime/Engine/Public/Animation/MeshDeformerGeometryReadback.h
--- a/Engine/Source/Runtime/Engine/Public/Animation/MeshDeformerGeometryReadback.h
+++ b/Engine/Source/Runtime/Engine/Public/Animation/MeshDeformerGeometryReadback.h
@@ -6,10 +6,22 @@
 
 struct FMeshDescription;
 
+struct ENGINE_API FMeshDeformerGeometryReadbackVertexDataArrays
+{
+	int32 LODIndex = INDEX_NONE;
+	TArray<FVector3f> Positions;
+	TArray<FVector3f> Normals;
+	TArray<FVector3f> Tangents;
+	TArray<FVector4f> Colors;
+};
+
 struct ENGINE_API FMeshDeformerGeometryReadbackRequest
 {
 	~FMeshDeformerGeometryReadbackRequest();
 
 	TFunction<void(const FMeshDescription&)> MeshDescriptionCallback_AnyThread;
-	bool bHandled = false;
+	bool bMeshDescriptionHandled = false;
+	
+	TFunction<void(const FMeshDeformerGeometryReadbackVertexDataArrays&)> VertexDataArraysCallback_AnyThread;
+	bool bVertexDataArraysHandled = false;
 };
\ No newline at end of file
