[DeformerGraph] switching to a new convention, the goal is to completely hide kernel invocations from the user. User should not be aware of whether the underlying buffer that they are programming against is a single buffer of multiple buffers:
1. kernel with simple execution domain can support both single invocation (unified dispatch) or multiple invocation (non unified dispatch), but kernel with non-simple, custom/expression domain can only support unified dispatch
(Simple: Vertex, Triangle, Not Simple: Vertex * 2, Vertex + 2 .....)
4. data interface that cannot support unified dispatch can only be connected to a kernel that uses simple execution domain
2. If a resource has the domain Vertex + 2, when allocated for a mesh with 2 sections(invocations) it no longer has the size of (Section0.NumVertex + 2) + (Section1.NumVertex + 2), instead it now has the size of (Section0.NumVertex + Section1.NumVertex + 2)
3. When a kernel is running multiple invocations, the user-facing Index value will no longer use the range 0 - ThreadCountPerInvocation, instead it will be something like 0 - 9 for the first invocation, 10 - 29 for the 2nd invocation, 30 - 49 for the third.... As a result, I removed the use of InputStreamStart/BaseVertexIndex from some of the DI shader source since they are no longer needed.

diff --git a/Engine/Plugins/Animation/DeformerGraph/Shaders/Private/DataInterfaceAdvancedSkeleton.ush b/Engine/Plugins/Animation/DeformerGraph/Shaders/Private/DataInterfaceAdvancedSkeleton.ush
--- a/Engine/Plugins/Animation/DeformerGraph/Shaders/Private/DataInterfaceAdvancedSkeleton.ush
+++ b/Engine/Plugins/Animation/DeformerGraph/Shaders/Private/DataInterfaceAdvancedSkeleton.ush
@@ -3,9 +3,7 @@
 #include "/Plugin/Optimus/Private/DataInterfaceSkeletonCommon.ush"
 
 uint {DataInterfaceName}_NumVertices;
-uint {DataInterfaceName}_InputStreamStart;
 uint {DataInterfaceName}_NumBoneInfluences;
-uint {DataInterfaceName}_InputWeightStart;
 uint {DataInterfaceName}_InputWeightStride;
 uint {DataInterfaceName}_InputWeightIndexSize;
 Buffer<float4> {DataInterfaceName}_BoneMatrices;
@@ -22,7 +20,6 @@ uint ReadNumVertices_{DataInterfaceName}()
 uint ReadNumBones_{DataInterfaceName}(uint VertexIndex)
 {
 	return ReadNumBones(
-		{DataInterfaceName}_InputStreamStart,
 		{DataInterfaceName}_InputWeightLookupStream,
 		{DataInterfaceName}_NumBoneInfluences,
 		VertexIndex
@@ -32,11 +29,9 @@ uint ReadNumBones_{DataInterfaceName}(uint VertexIndex)
 float3x4 ReadBoneMatrix_{DataInterfaceName}(uint VertexIndex, uint BoneIndex)
 {
 	return ReadBoneMatrix(
-		{DataInterfaceName}_InputStreamStart,
 		{DataInterfaceName}_InputWeightLookupStream,
 		{DataInterfaceName}_InputWeightIndexSize,
 		{DataInterfaceName}_InputWeightStream,
-		{DataInterfaceName}_InputWeightStart,
 		{DataInterfaceName}_InputWeightStride,
 		{DataInterfaceName}_BoneMatrices,
 		VertexIndex,
@@ -47,11 +42,9 @@ float3x4 ReadBoneMatrix_{DataInterfaceName}(uint VertexIndex, uint BoneIndex)
 float ReadBoneWeight_{DataInterfaceName}(uint VertexIndex, uint BoneIndex)
 {
 	return ReadBoneWeight(
-		{DataInterfaceName}_InputStreamStart,
 		{DataInterfaceName}_InputWeightLookupStream,
 		{DataInterfaceName}_InputWeightIndexSize,
 		{DataInterfaceName}_InputWeightStream,
-		{DataInterfaceName}_InputWeightStart,
 		{DataInterfaceName}_InputWeightStride,
 		{DataInterfaceName}_NumBoneInfluences,
 		VertexIndex,
@@ -62,11 +55,9 @@ float ReadBoneWeight_{DataInterfaceName}(uint VertexIndex, uint BoneIndex)
 float3x4 ReadWeightedBoneMatrix_{DataInterfaceName}(uint VertexIndex)
 {
 	return ReadWeightedBoneMatrix(
-		{DataInterfaceName}_InputStreamStart,
 		{DataInterfaceName}_InputWeightLookupStream,
 		{DataInterfaceName}_InputWeightIndexSize,
 		{DataInterfaceName}_InputWeightStream,
-		{DataInterfaceName}_InputWeightStart,
 		{DataInterfaceName}_InputWeightStride,
 		{DataInterfaceName}_NumBoneInfluences,
 		{DataInterfaceName}_BoneMatrices,
@@ -82,7 +73,7 @@ float3x4 ReadLayeredBoneMatrix_{DataInterfaceName}(uint VertexIndex, uint BoneIn
 	return float3x4(1, 0, 0, 0, 0, 1, 0, 0, 0, 0, 1, 0);
 
 #elif GPUSKIN_UNLIMITED_BONE_INFLUENCE
-	const uint BufferIndex 		= {DataInterfaceName}_InputStreamStart + VertexIndex;
+	const uint BufferIndex 		= VertexIndex;
 	uint BlendOffsetCount = {DataInterfaceName}_InputWeightLookupStream[BufferIndex];
 	int StreamOffset = BlendOffsetCount >> 8;
 	int BoneIndexSize = {DataInterfaceName}_InputWeightIndexSize & 0xff;
@@ -97,7 +88,7 @@ float3x4 ReadLayeredBoneMatrix_{DataInterfaceName}(uint VertexIndex, uint BoneIn
 
 #else // !GPUSKIN_UNLIMITED_BONE_INFLUENCE
 	
-	uint StreamOffset = {DataInterfaceName}_InputWeightStart + (VertexIndex * ({DataInterfaceName}_InputWeightStride / 4));
+	uint StreamOffset = VertexIndex * ({DataInterfaceName}_InputWeightStride / 4);
 	int BlendIndex = GetBlendIndices({DataInterfaceName}_InputWeightStream, StreamOffset, BoneIndex / 4)[BoneIndex & 0x3];
 	return {DataInterfaceName}_LayeredBoneMatrices[BlendIndex];
 
diff --git a/Engine/Plugins/Animation/DeformerGraph/Shaders/Private/DataInterfaceAdvancedSkeletonAnimAttribute.ush b/Engine/Plugins/Animation/DeformerGraph/Shaders/Private/DataInterfaceAdvancedSkeletonAnimAttribute.ush
--- a/Engine/Plugins/Animation/DeformerGraph/Shaders/Private/DataInterfaceAdvancedSkeletonAnimAttribute.ush
+++ b/Engine/Plugins/Animation/DeformerGraph/Shaders/Private/DataInterfaceAdvancedSkeletonAnimAttribute.ush
@@ -10,7 +10,7 @@ StructuredBuffer<{TypeName}> {DataInterfaceName}_{AttributeName};
 	return 0;
 
 #elif GPUSKIN_UNLIMITED_BONE_INFLUENCE
-	const uint BufferIndex 		= {DataInterfaceName}_InputStreamStart + VertexIndex;
+	const uint BufferIndex 		= VertexIndex;
 	uint BlendOffsetCount = {DataInterfaceName}_InputWeightLookupStream[BufferIndex];
 	int StreamOffset = BlendOffsetCount >> 8;
 	int BoneIndexSize = {DataInterfaceName}_InputWeightIndexSize & 0xff;
@@ -25,7 +25,7 @@ StructuredBuffer<{TypeName}> {DataInterfaceName}_{AttributeName};
 
 #else // !GPUSKIN_UNLIMITED_BONE_INFLUENCE
 	
-	uint StreamOffset = {DataInterfaceName}_InputWeightStart + (VertexIndex * ({DataInterfaceName}_InputWeightStride / 4));
+	uint StreamOffset = (VertexIndex * ({DataInterfaceName}_InputWeightStride / 4));
 	int BlendIndex = GetBlendIndices({DataInterfaceName}_InputWeightStream, StreamOffset, BoneIndex / 4)[BoneIndex & 0x3];
 	return {DataInterfaceName}_{AttributeName}[BlendIndex];
 
diff --git a/Engine/Plugins/Animation/DeformerGraph/Shaders/Private/DataInterfaceCloth.ush b/Engine/Plugins/Animation/DeformerGraph/Shaders/Private/DataInterfaceCloth.ush
--- a/Engine/Plugins/Animation/DeformerGraph/Shaders/Private/DataInterfaceCloth.ush
+++ b/Engine/Plugins/Animation/DeformerGraph/Shaders/Private/DataInterfaceCloth.ush
@@ -23,6 +23,7 @@ struct FClothResult
 #endif // DATA_INTERFACE_CLOTH_ONCE
 
 uint {DataInterfaceName}_NumVertices;
+uint {DataInterfaceName}_BaseVertexIndex;
 uint {DataInterfaceName}_InputStreamStart;
 uint {DataInterfaceName}_NumInfluencesPerVertex;
 float {DataInterfaceName}_ClothBlendWeight;
@@ -47,7 +48,7 @@ float4x4 ReadClothToLocal_{DataInterfaceName}()
 
 FVertexTriangleInfluence GetClothInfluence_{DataInterfaceName}(uint VertexIndex, uint InfluenceIndex)
 {
-	const uint VertexOffset = {DataInterfaceName}_NumInfluencesPerVertex * VertexIndex + {DataInterfaceName}_InputStreamStart;
+	const uint VertexOffset = {DataInterfaceName}_NumInfluencesPerVertex * (VertexIndex - {DataInterfaceName}_BaseVertexIndex)  + {DataInterfaceName}_InputStreamStart;
 
 	const uint NUM_FLOAT4S_PER_VERTEX_INFLUENCE = 4; // Stride of GPUSkinApexCloth
 	const uint Offset = (VertexOffset + InfluenceIndex) * NUM_FLOAT4S_PER_VERTEX_INFLUENCE;
diff --git a/Engine/Plugins/Animation/DeformerGraph/Shaders/Private/DataInterfaceConnectivity.ush b/Engine/Plugins/Animation/DeformerGraph/Shaders/Private/DataInterfaceConnectivity.ush
--- a/Engine/Plugins/Animation/DeformerGraph/Shaders/Private/DataInterfaceConnectivity.ush
+++ b/Engine/Plugins/Animation/DeformerGraph/Shaders/Private/DataInterfaceConnectivity.ush
@@ -1,7 +1,6 @@
 // Copyright Epic Games, Inc. All Rights Reserved.
 
 uint {DataInterfaceName}_NumVertices;
-uint {DataInterfaceName}_InputStreamStart;
 uint {DataInterfaceName}_MaxConnectedVertexCount;
 StructuredBuffer<uint> {DataInterfaceName}_ConnectivityBuffer;
 
@@ -17,8 +16,7 @@ uint ReadNumConnectedVertices_{DataInterfaceName}(uint VertexIndex)
 
 uint ReadConnectedVertex_{DataInterfaceName}(uint VertexIndex, uint ConnectedVertexSubIndex)
 {
-	uint BaseIndex = ({DataInterfaceName}_InputStreamStart + VertexIndex) * {DataInterfaceName}_MaxConnectedVertexCount;
+	uint BaseIndex = (VertexIndex) * {DataInterfaceName}_MaxConnectedVertexCount;
 	uint ConnectedVertexIndex = {DataInterfaceName}_ConnectivityBuffer[BaseIndex + ConnectedVertexSubIndex];
-	uint BaseOffset = ConnectedVertexIndex == 0xFFFFFFFF ? 0 : {DataInterfaceName}_InputStreamStart;
-	return ConnectedVertexIndex - BaseOffset;
+	return ConnectedVertexIndex;
 }
diff --git a/Engine/Plugins/Animation/DeformerGraph/Shaders/Private/DataInterfaceCopyKernel.ush b/Engine/Plugins/Animation/DeformerGraph/Shaders/Private/DataInterfaceCopyKernel.ush
new file mode 100644
--- /dev/null
+++ b/Engine/Plugins/Animation/DeformerGraph/Shaders/Private/DataInterfaceCopyKernel.ush
@@ -0,0 +1,20 @@
+// Copyright Epic Games, Inc. All Rights Reserved.
+
+int3 {DataInterfaceName}_NumThreads;
+int {DataInterfaceName}_NumThreadsPerInvocation;
+int {DataInterfaceName}_ThreadIndexOffset;
+
+int3 ReadNumThreads_{DataInterfaceName}()
+{
+	return {DataInterfaceName}_NumThreads;
+}
+
+int ReadNumThreadsPerInvocation_{DataInterfaceName}()
+{
+	return {DataInterfaceName}_NumThreadsPerInvocation;
+}
+
+int ReadThreadIndexOffset_{DataInterfaceName}()
+{
+	return {DataInterfaceName}_ThreadIndexOffset;
+}
\ No newline at end of file
diff --git a/Engine/Plugins/Animation/DeformerGraph/Shaders/Private/DataInterfaceCustomComputeKernel.ush b/Engine/Plugins/Animation/DeformerGraph/Shaders/Private/DataInterfaceCustomComputeKernel.ush
new file mode 100644
--- /dev/null
+++ b/Engine/Plugins/Animation/DeformerGraph/Shaders/Private/DataInterfaceCustomComputeKernel.ush
@@ -0,0 +1,20 @@
+// Copyright Epic Games, Inc. All Rights Reserved.
+
+int3 {DataInterfaceName}_NumThreads;
+int {DataInterfaceName}_NumThreadsPerInvocation;
+int {DataInterfaceName}_ThreadIndexOffset;
+
+int3 ReadNumThreads_{DataInterfaceName}()
+{
+	return {DataInterfaceName}_NumThreads;
+}
+
+int ReadNumThreadsPerInvocation_{DataInterfaceName}()
+{
+	return {DataInterfaceName}_NumThreadsPerInvocation;
+}
+
+int ReadThreadIndexOffset_{DataInterfaceName}()
+{
+	return {DataInterfaceName}_ThreadIndexOffset;
+}
\ No newline at end of file
diff --git a/Engine/Plugins/Animation/DeformerGraph/Shaders/Private/DataInterfaceDuplicateVertices.ush b/Engine/Plugins/Animation/DeformerGraph/Shaders/Private/DataInterfaceDuplicateVertices.ush
--- a/Engine/Plugins/Animation/DeformerGraph/Shaders/Private/DataInterfaceDuplicateVertices.ush
+++ b/Engine/Plugins/Animation/DeformerGraph/Shaders/Private/DataInterfaceDuplicateVertices.ush
@@ -2,7 +2,7 @@
 
 uint {DataInterfaceName}_NumVertices;
 uint {DataInterfaceName}_NumDuplicateVertices;
-uint {DataInterfaceName}_InputStreamStart;
+uint {DataInterfaceName}_BaseVertexIndex;
 Buffer<uint> {DataInterfaceName}_DuplicateVertexStartAndLength;
 Buffer<uint> {DataInterfaceName}_DuplicateVertices;
 
@@ -19,8 +19,7 @@ uint ReadNumDuplicateVertices_{DataInterfaceName}()
 int2 ReadDuplicateVerticesStartAndLength_{DataInterfaceName}(uint VertexIndex)
 {
 #if ENABLE_DUPLICATED_VERTICES
-	// DuplicateVertex buffer is provided per section and thus we don't need to offset index with InputStreamStart
-	uint BufferIndex = VertexIndex * 2;
+	uint BufferIndex = (VertexIndex - {DataInterfaceName}_BaseVertexIndex) * 2;
 	return int2({DataInterfaceName}_DuplicateVertexStartAndLength[BufferIndex  + 1], {DataInterfaceName}_DuplicateVertexStartAndLength[BufferIndex]);
 #else	
 	return 0;
@@ -30,7 +29,7 @@ int2 ReadDuplicateVerticesStartAndLength_{DataInterfaceName}(uint VertexIndex)
 uint ReadDuplicateVertex_{DataInterfaceName}(uint Index)
 {
 #if ENABLE_DUPLICATED_VERTICES
-	return {DataInterfaceName}_DuplicateVertices[Index];
+	return {DataInterfaceName}_DuplicateVertices[Index] + {DataInterfaceName}_BaseVertexIndex;
 #else	
 	return 0;
 #endif
diff --git a/Engine/Plugins/Animation/DeformerGraph/Shaders/Private/DataInterfaceHalfEdge.ush b/Engine/Plugins/Animation/DeformerGraph/Shaders/Private/DataInterfaceHalfEdge.ush
--- a/Engine/Plugins/Animation/DeformerGraph/Shaders/Private/DataInterfaceHalfEdge.ush
+++ b/Engine/Plugins/Animation/DeformerGraph/Shaders/Private/DataInterfaceHalfEdge.ush
@@ -2,8 +2,6 @@
 
 uint {DataInterfaceName}_NumVertices;
 uint {DataInterfaceName}_NumTriangles;
-uint {DataInterfaceName}_IndexBufferStart;
-uint {DataInterfaceName}_InputStreamStart;
 StructuredBuffer<int> {DataInterfaceName}_VertexToEdgeBuffer;
 StructuredBuffer<int> {DataInterfaceName}_EdgeToTwinEdgeBuffer;
 
@@ -24,15 +22,15 @@ uint ReadNumTriangles_{DataInterfaceName}()
 int ReadEdge_{DataInterfaceName}(uint VertexIndex)
 {
 	return {DataInterfaceName}_bUseBufferFromRenderData == 0  ?
-		{DataInterfaceName}_VertexToEdgeBuffer[VertexIndex + {DataInterfaceName}_InputStreamStart] - {DataInterfaceName}_IndexBufferStart :
-		{DataInterfaceName}_RenderDataVertexToEdgeBuffer[VertexIndex + {DataInterfaceName}_InputStreamStart] - {DataInterfaceName}_IndexBufferStart;
+		{DataInterfaceName}_VertexToEdgeBuffer[VertexIndex] :
+		{DataInterfaceName}_RenderDataVertexToEdgeBuffer[VertexIndex];
 }
 
 int ReadTwinEdge_{DataInterfaceName}(uint EdgeIndex)
 {
 	int TwinEdge = {DataInterfaceName}_bUseBufferFromRenderData == 0 ?
-		{DataInterfaceName}_EdgeToTwinEdgeBuffer[EdgeIndex + {DataInterfaceName}_IndexBufferStart] :
-		{DataInterfaceName}_RenderDataEdgeToTwinEdgeBuffer[EdgeIndex + {DataInterfaceName}_IndexBufferStart];
+		{DataInterfaceName}_EdgeToTwinEdgeBuffer[EdgeIndex] :
+		{DataInterfaceName}_RenderDataEdgeToTwinEdgeBuffer[EdgeIndex];
 	
-	return TwinEdge == -1 ? -1 : TwinEdge - {DataInterfaceName}_IndexBufferStart;
+	return TwinEdge == -1 ? -1 : TwinEdge;
 }
diff --git a/Engine/Plugins/Animation/DeformerGraph/Shaders/Private/DataInterfaceMorphTarget.ush b/Engine/Plugins/Animation/DeformerGraph/Shaders/Private/DataInterfaceMorphTarget.ush
--- a/Engine/Plugins/Animation/DeformerGraph/Shaders/Private/DataInterfaceMorphTarget.ush
+++ b/Engine/Plugins/Animation/DeformerGraph/Shaders/Private/DataInterfaceMorphTarget.ush
@@ -1,7 +1,6 @@
 // Copyright Epic Games, Inc. All Rights Reserved.
 
 uint {DataInterfaceName}_NumVertices;
-uint {DataInterfaceName}_InputStreamStart;
 Buffer<float> {DataInterfaceName}_MorphBuffer;
 
 uint ReadNumVertices_{DataInterfaceName}()
@@ -12,7 +11,7 @@ uint ReadNumVertices_{DataInterfaceName}()
 float3 ReadDeltaPosition_{DataInterfaceName}(uint VertexIndex)
 {
 #if ENABLE_DEFORMER_MORPHTARGET
-	uint BufferIndex = ({DataInterfaceName}_InputStreamStart + VertexIndex) * 6;
+	uint BufferIndex = (VertexIndex) * 6;
 	return float3({DataInterfaceName}_MorphBuffer[BufferIndex], {DataInterfaceName}_MorphBuffer[BufferIndex + 1], {DataInterfaceName}_MorphBuffer[BufferIndex + 2]);
 #else
 	return float3(0, 0, 0);
@@ -22,7 +21,7 @@ float3 ReadDeltaPosition_{DataInterfaceName}(uint VertexIndex)
 float3 ReadDeltaNormal_{DataInterfaceName}(uint VertexIndex)
 {
 #if ENABLE_DEFORMER_MORPHTARGET
-	uint BufferIndex = ({DataInterfaceName}_InputStreamStart + VertexIndex) * 6;
+	uint BufferIndex = (VertexIndex) * 6;
 	return float3({DataInterfaceName}_MorphBuffer[BufferIndex + 3], {DataInterfaceName}_MorphBuffer[BufferIndex + 4], {DataInterfaceName}_MorphBuffer[BufferIndex + 5]);
 #else
 	return float3(0, 0, 0);
diff --git a/Engine/Plugins/Animation/DeformerGraph/Shaders/Private/DataInterfaceRawBuffer.ush b/Engine/Plugins/Animation/DeformerGraph/Shaders/Private/DataInterfaceRawBuffer.ush
--- a/Engine/Plugins/Animation/DeformerGraph/Shaders/Private/DataInterfaceRawBuffer.ush
+++ b/Engine/Plugins/Animation/DeformerGraph/Shaders/Private/DataInterfaceRawBuffer.ush
@@ -13,7 +13,6 @@
 // BUFFER_SPLIT_READ_WRITE is set if we can use an SRV to read the buffer. It is unset if we might need to read and write the buffer in the same kernel.
 #define BUFFER_SPLIT_READ_WRITE {SplitReadWrite}
 
-uint {DataInterfaceName}_StartOffset;
 uint {DataInterfaceName}_BufferSize;
 #if BUFFER_SPLIT_READ_WRITE
 StructuredBuffer<BUFFER_TYPE> {DataInterfaceName}_BufferSRV;
@@ -27,7 +26,7 @@ uint ReadNumValues_{DataInterfaceName}()
 
 PUBLIC_TYPE ReadValue_{DataInterfaceName}(uint Index)
 {
-	uint BufferIndex = {DataInterfaceName}_StartOffset + Index;
+	uint BufferIndex = Index;
 	PUBLIC_TYPE ReturnValue;
 #if BUFFER_SPLIT_READ_WRITE
 #if BUFFER_TYPE_RAW
@@ -47,7 +46,7 @@ PUBLIC_TYPE ReadValue_{DataInterfaceName}(uint Index)
 
 PUBLIC_TYPE ReadValueUAV_{DataInterfaceName}(uint Index)
 {
-	uint BufferIndex = {DataInterfaceName}_StartOffset + Index;
+	uint BufferIndex = Index;
 	PUBLIC_TYPE ReturnValue;
 
 #if BUFFER_TYPE_RAW
@@ -61,7 +60,7 @@ PUBLIC_TYPE ReadValueUAV_{DataInterfaceName}(uint Index)
 
 void WriteValue_{DataInterfaceName}(uint Index, PUBLIC_TYPE Value)
 {
-	uint BufferIndex = {DataInterfaceName}_StartOffset + Index;
+	uint BufferIndex = Index;
 #if BUFFER_TYPE_RAW
 	WriteRawBuffer({DataInterfaceName}_BufferUAV, BufferIndex, Value);
 #else
@@ -72,7 +71,7 @@ void WriteValue_{DataInterfaceName}(uint Index, PUBLIC_TYPE Value)
 PUBLIC_TYPE WriteAtomicAdd_{DataInterfaceName}(uint Index, PUBLIC_TYPE Value)
 {
 #if BUFFER_TYPE_SUPPORTS_ATOMIC
-	uint BufferIndex = {DataInterfaceName}_StartOffset + Index;
+	uint BufferIndex = Index;
 	PUBLIC_TYPE OldValue;
 	InterlockedAdd({DataInterfaceName}_BufferUAV[BufferIndex], Value, OldValue);
 	return OldValue;
@@ -83,7 +82,7 @@ PUBLIC_TYPE WriteAtomicAdd_{DataInterfaceName}(uint Index, PUBLIC_TYPE Value)
 PUBLIC_TYPE WriteAtomicMin_{DataInterfaceName}(uint Index, PUBLIC_TYPE Value)
 {
 #if BUFFER_TYPE_SUPPORTS_ATOMIC
-	uint BufferIndex = {DataInterfaceName}_StartOffset + Index;
+	uint BufferIndex = Index;
 	PUBLIC_TYPE OldValue;
 	InterlockedMin({DataInterfaceName}_BufferUAV[BufferIndex], Value, OldValue);
 	return OldValue;
@@ -94,7 +93,7 @@ PUBLIC_TYPE WriteAtomicMin_{DataInterfaceName}(uint Index, PUBLIC_TYPE Value)
 PUBLIC_TYPE WriteAtomicMax_{DataInterfaceName}(uint Index, PUBLIC_TYPE Value)
 {
 #if BUFFER_TYPE_SUPPORTS_ATOMIC
-	uint BufferIndex = {DataInterfaceName}_StartOffset + Index;
+	uint BufferIndex = Index;
 	PUBLIC_TYPE OldValue;
 	InterlockedMax({DataInterfaceName}_BufferUAV[BufferIndex], Value, OldValue);
 	return OldValue;
diff --git a/Engine/Plugins/Animation/DeformerGraph/Shaders/Private/DataInterfaceSkeleton.ush b/Engine/Plugins/Animation/DeformerGraph/Shaders/Private/DataInterfaceSkeleton.ush
--- a/Engine/Plugins/Animation/DeformerGraph/Shaders/Private/DataInterfaceSkeleton.ush
+++ b/Engine/Plugins/Animation/DeformerGraph/Shaders/Private/DataInterfaceSkeleton.ush
@@ -3,9 +3,7 @@
 #include "/Plugin/Optimus/Private/DataInterfaceSkeletonCommon.ush"
 
 uint {DataInterfaceName}_NumVertices;
-uint {DataInterfaceName}_InputStreamStart;
 uint {DataInterfaceName}_NumBoneInfluences;
-uint {DataInterfaceName}_InputWeightStart;
 uint {DataInterfaceName}_InputWeightStride;
 uint {DataInterfaceName}_InputWeightIndexSize;
 Buffer<float4> {DataInterfaceName}_BoneMatrices;
@@ -20,7 +18,6 @@ uint ReadNumVertices_{DataInterfaceName}()
 uint ReadNumBones_{DataInterfaceName}(uint VertexIndex)
 {
 	return ReadNumBones(
-		{DataInterfaceName}_InputStreamStart,
 		{DataInterfaceName}_InputWeightLookupStream,
 		{DataInterfaceName}_NumBoneInfluences,
 		VertexIndex
@@ -30,11 +27,9 @@ uint ReadNumBones_{DataInterfaceName}(uint VertexIndex)
 float3x4 ReadBoneMatrix_{DataInterfaceName}(uint VertexIndex, uint BoneIndex)
 {
 	return ReadBoneMatrix(
-		{DataInterfaceName}_InputStreamStart,
 		{DataInterfaceName}_InputWeightLookupStream,
 		{DataInterfaceName}_InputWeightIndexSize,
 		{DataInterfaceName}_InputWeightStream,
-		{DataInterfaceName}_InputWeightStart,
 		{DataInterfaceName}_InputWeightStride,
 		{DataInterfaceName}_BoneMatrices,
 		VertexIndex,
@@ -45,11 +40,9 @@ float3x4 ReadBoneMatrix_{DataInterfaceName}(uint VertexIndex, uint BoneIndex)
 float ReadBoneWeight_{DataInterfaceName}(uint VertexIndex, uint BoneIndex)
 {
 	return ReadBoneWeight(
-		{DataInterfaceName}_InputStreamStart,
 		{DataInterfaceName}_InputWeightLookupStream,
 		{DataInterfaceName}_InputWeightIndexSize,
 		{DataInterfaceName}_InputWeightStream,
-		{DataInterfaceName}_InputWeightStart,
 		{DataInterfaceName}_InputWeightStride,
 		{DataInterfaceName}_NumBoneInfluences,	
 		VertexIndex,
@@ -60,11 +53,9 @@ float ReadBoneWeight_{DataInterfaceName}(uint VertexIndex, uint BoneIndex)
 float3x4 ReadWeightedBoneMatrix_{DataInterfaceName}(uint VertexIndex)
 {
 	return ReadWeightedBoneMatrix(
-		{DataInterfaceName}_InputStreamStart,
 		{DataInterfaceName}_InputWeightLookupStream,
 		{DataInterfaceName}_InputWeightIndexSize,
 		{DataInterfaceName}_InputWeightStream,
-		{DataInterfaceName}_InputWeightStart,
 		{DataInterfaceName}_InputWeightStride,
 		{DataInterfaceName}_NumBoneInfluences,
 		{DataInterfaceName}_BoneMatrices,
diff --git a/Engine/Plugins/Animation/DeformerGraph/Shaders/Private/DataInterfaceSkeletonCommon.ush b/Engine/Plugins/Animation/DeformerGraph/Shaders/Private/DataInterfaceSkeletonCommon.ush
--- a/Engine/Plugins/Animation/DeformerGraph/Shaders/Private/DataInterfaceSkeletonCommon.ush
+++ b/Engine/Plugins/Animation/DeformerGraph/Shaders/Private/DataInterfaceSkeletonCommon.ush
@@ -3,7 +3,6 @@
 #include "/Engine/Private/GpuSkinCommon.ush"
 
 uint ReadNumBones(
-	uint InputStreamStart,
 	Buffer<uint> InputWeightLookupStream,
 	uint NumBoneInfluences,
 	uint VertexIndex
@@ -12,7 +11,7 @@ uint ReadNumBones(
 #if !ENABLE_DEFORMER_BONES
 	return 1;
 #elif GPUSKIN_UNLIMITED_BONE_INFLUENCE
-	const uint BufferIndex 		= InputStreamStart + VertexIndex;
+	const uint BufferIndex = VertexIndex;
 	uint BlendOffsetCount = InputWeightLookupStream[BufferIndex];
 	return BlendOffsetCount & 0xff;
 #else // !GPUSKIN_UNLIMITED_BONE_INFLUENCE
@@ -21,11 +20,9 @@ uint ReadNumBones(
 }
 
 float ReadBoneWeight(
-	uint InputStreamStart,
 	Buffer<uint> InputWeightLookupStream,
 	uint InputWeightIndexSize,
 	Buffer<uint> InputWeightStream,
-	uint InputWeightStart,
 	uint InputWeightStride,
 	uint NumBoneInfluences,
 	uint VertexIndex,
@@ -37,7 +34,7 @@ float ReadBoneWeight(
 	return 1;
 
 #elif GPUSKIN_UNLIMITED_BONE_INFLUENCE
-	const uint BufferIndex 		= InputStreamStart + VertexIndex;	
+	const uint BufferIndex 		= VertexIndex;	
 	uint BlendOffsetCount = InputWeightLookupStream[BufferIndex];
 	int NumBoneInfluencesLocal = BlendOffsetCount & 0xff;
 	int StreamOffset = BlendOffsetCount >> 8;
@@ -57,7 +54,7 @@ float ReadBoneWeight(
 
 #else // !GPUSKIN_UNLIMITED_BONE_INFLUENCE
 	
-	uint StreamOffset = InputWeightStart + (VertexIndex * (InputWeightStride / 4));
+	uint StreamOffset = (VertexIndex * (InputWeightStride / 4));
 	float BlendWeight = GetBlendWeights(InputWeightStream, StreamOffset, BoneIndex / 4, NumBoneInfluences)[BoneIndex & 0x3];
 	return BlendWeight;
 
@@ -71,11 +68,9 @@ float3x4 GetBoneMatrix(Buffer<float4> BoneMatrices, uint BoneIndex)
 }
 
 float3x4 ReadBoneMatrix(
-	uint InputStreamStart,
 	Buffer<uint> InputWeightLookupStream,
 	uint InputWeightIndexSize,
 	Buffer<uint> InputWeightStream,
-	uint InputWeightStart,
 	uint InputWeightStride,
 	Buffer<float4> BoneMatrices,	
 	uint VertexIndex,
@@ -87,7 +82,7 @@ float3x4 ReadBoneMatrix(
 	return float3x4(1, 0, 0, 0, 0, 1, 0, 0, 0, 0, 1, 0);
 
 #elif GPUSKIN_UNLIMITED_BONE_INFLUENCE
-	const uint BufferIndex 		= InputStreamStart + VertexIndex;
+	const uint BufferIndex 		= VertexIndex;
 	uint BlendOffsetCount = InputWeightLookupStream[BufferIndex];
 	int StreamOffset = BlendOffsetCount >> 8;
 	int BoneIndexSize = InputWeightIndexSize & 0xff;
@@ -102,7 +97,7 @@ float3x4 ReadBoneMatrix(
 
 #else // !GPUSKIN_UNLIMITED_BONE_INFLUENCE
 	
-	uint StreamOffset = InputWeightStart + (VertexIndex * (InputWeightStride / 4));
+	uint StreamOffset = (VertexIndex * (InputWeightStride / 4));
 	int BlendIndex = GetBlendIndices(InputWeightStream ,StreamOffset, BoneIndex / 4)[BoneIndex & 0x3];
 	return GetBoneMatrix(BoneMatrices,BlendIndex);
 
@@ -110,11 +105,9 @@ float3x4 ReadBoneMatrix(
 }
 
 float3x4 ReadWeightedBoneMatrix(
-	uint InputStreamStart,
 	Buffer<uint> InputWeightLookupStream,
 	uint InputWeightIndexSize,
 	Buffer<uint> InputWeightStream,
-	uint InputWeightStart,
 	uint InputWeightStride,
 	uint NumBoneInfluences,
 	Buffer<float4> BoneMatrices,		
@@ -127,13 +120,13 @@ float3x4 ReadWeightedBoneMatrix(
 
 #elif GPUSKIN_UNLIMITED_BONE_INFLUENCE
 
-	const uint BufferIndex 		= InputStreamStart + VertexIndex;
+	const uint BufferIndex 		= VertexIndex;
 	const uint BlendOffsetCount = InputWeightLookupStream[BufferIndex];
 	return ComputeBoneMatrixWithUnlimitedInfluences(BoneMatrices, InputWeightStream, InputWeightIndexSize, BlendOffsetCount);
 
 #else
 
-	const uint StreamOffset = InputWeightStart + (VertexIndex * (InputWeightStride / 4));
+	const uint StreamOffset = (VertexIndex * (InputWeightStride / 4));
 	const uint NumInfluences = NumBoneInfluences;
 
 	FGPUSkinIndexAndWeight IndicesAndWeights = (FGPUSkinIndexAndWeight)0;
diff --git a/Engine/Plugins/Animation/DeformerGraph/Shaders/Private/DataInterfaceSkinWeightsAsVertexMask.ush b/Engine/Plugins/Animation/DeformerGraph/Shaders/Private/DataInterfaceSkinWeightsAsVertexMask.ush
--- a/Engine/Plugins/Animation/DeformerGraph/Shaders/Private/DataInterfaceSkinWeightsAsVertexMask.ush
+++ b/Engine/Plugins/Animation/DeformerGraph/Shaders/Private/DataInterfaceSkinWeightsAsVertexMask.ush
@@ -3,9 +3,7 @@
 #include "/Plugin/Optimus/Private/DataInterfaceSkeletonCommon.ush"
 
 uint {DataInterfaceName}_NumVertices;
-uint {DataInterfaceName}_InputStreamStart;
 uint {DataInterfaceName}_NumBoneInfluences;
-uint {DataInterfaceName}_InputWeightStart;
 uint {DataInterfaceName}_InputWeightStride;
 uint {DataInterfaceName}_InputWeightIndexSize;
 Buffer<uint> {DataInterfaceName}_InputWeightStream;
@@ -21,7 +19,6 @@ uint ReadNumVertices_{DataInterfaceName}()
 uint ReadNumBones_{DataInterfaceName}(uint VertexIndex)
 {
 	return ReadNumBones(
-		{DataInterfaceName}_InputStreamStart,
 		{DataInterfaceName}_InputWeightLookupStream,
 		{DataInterfaceName}_NumBoneInfluences,
 		VertexIndex
@@ -31,11 +28,9 @@ uint ReadNumBones_{DataInterfaceName}(uint VertexIndex)
 float ReadBoneWeight_{DataInterfaceName}(uint VertexIndex, uint BoneIndex)
 {
 	return ReadBoneWeight(
-		{DataInterfaceName}_InputStreamStart,
 		{DataInterfaceName}_InputWeightLookupStream,
 		{DataInterfaceName}_InputWeightIndexSize,
 		{DataInterfaceName}_InputWeightStream,
-		{DataInterfaceName}_InputWeightStart,
 		{DataInterfaceName}_InputWeightStride,
 		{DataInterfaceName}_NumBoneInfluences,	
 		VertexIndex,
@@ -50,7 +45,7 @@ uint ReadBoneIsSelected_{DataInterfaceName}(uint VertexIndex, uint BoneIndex)
 	return 0;
 
 #elif GPUSKIN_UNLIMITED_BONE_INFLUENCE
-	const uint BufferIndex 		= {DataInterfaceName}_InputStreamStart + VertexIndex;
+	const uint BufferIndex 		= VertexIndex;
 	uint BlendOffsetCount = {DataInterfaceName}_InputWeightLookupStream[BufferIndex];
 	int StreamOffset = BlendOffsetCount >> 8;
 	int BoneIndexSize = {DataInterfaceName}_InputWeightIndexSize & 0xff;
@@ -65,7 +60,7 @@ uint ReadBoneIsSelected_{DataInterfaceName}(uint VertexIndex, uint BoneIndex)
 
 #else // !GPUSKIN_UNLIMITED_BONE_INFLUENCE
 	
-	uint StreamOffset = {DataInterfaceName}_InputWeightStart + (VertexIndex * ({DataInterfaceName}_InputWeightStride / 4));
+	uint StreamOffset = (VertexIndex * ({DataInterfaceName}_InputWeightStride / 4));
 	int BlendIndex = GetBlendIndices({DataInterfaceName}_InputWeightStream, StreamOffset, BoneIndex / 4)[BoneIndex & 0x3];
 	return {DataInterfaceName}_BoneIsSelected[BlendIndex];
 
diff --git a/Engine/Plugins/Animation/DeformerGraph/Shaders/Private/DataInterfaceSkinnedMesh.ush b/Engine/Plugins/Animation/DeformerGraph/Shaders/Private/DataInterfaceSkinnedMesh.ush
--- a/Engine/Plugins/Animation/DeformerGraph/Shaders/Private/DataInterfaceSkinnedMesh.ush
+++ b/Engine/Plugins/Animation/DeformerGraph/Shaders/Private/DataInterfaceSkinnedMesh.ush
@@ -6,8 +6,6 @@
 uint {DataInterfaceName}_NumVertices;
 uint {DataInterfaceName}_NumTriangles;
 uint {DataInterfaceName}_NumUVChannels;
-uint {DataInterfaceName}_IndexBufferStart;
-uint {DataInterfaceName}_InputStreamStart;
 Buffer<uint> {DataInterfaceName}_IndexBuffer;
 Buffer<float> {DataInterfaceName}_PositionInputBuffer;
 Buffer<SNORM float4> {DataInterfaceName}_TangentInputBuffer;
@@ -32,37 +30,37 @@ uint ReadNumUVChannels_{DataInterfaceName}()
 
 uint ReadIndexBuffer_{DataInterfaceName}(uint Index)
 {
-	uint BufferIndex = {DataInterfaceName}_IndexBufferStart + Index;
-	return {DataInterfaceName}_IndexBuffer[BufferIndex] - {DataInterfaceName}_InputStreamStart;
+	uint BufferIndex = Index;
+	return {DataInterfaceName}_IndexBuffer[BufferIndex];
 }
 
 float3 ReadPosition_{DataInterfaceName}(uint VertexIndex)
 {
-	uint BufferIndex = ({DataInterfaceName}_InputStreamStart + VertexIndex) * 3;
+	uint BufferIndex = ( VertexIndex) * 3;
 	return float3({DataInterfaceName}_PositionInputBuffer[BufferIndex], {DataInterfaceName}_PositionInputBuffer[BufferIndex + 1], {DataInterfaceName}_PositionInputBuffer[BufferIndex + 2]);
 }
 
 float4 ReadTangentX_{DataInterfaceName}(uint VertexIndex)
 {
-	uint BufferIndex = ({DataInterfaceName}_InputStreamStart + VertexIndex) * 2;
+	uint BufferIndex = (VertexIndex) * 2;
 	return TangentBias_SkinCache({DataInterfaceName}_TangentInputBuffer[BufferIndex]);
 }
 
 float4 ReadTangentZ_{DataInterfaceName}(uint VertexIndex)
 {
-	uint BufferIndex = ({DataInterfaceName}_InputStreamStart + VertexIndex) * 2;
+	uint BufferIndex = (VertexIndex) * 2;
 	return TangentBias_SkinCache({DataInterfaceName}_TangentInputBuffer[BufferIndex + 1]);
 }
 
 float4 ReadColor_{DataInterfaceName}(uint VertexIndex)
 {
 	// Index Mask is used here in case ColorStaticBuffer is bound to the GWhiteVertexBufferWithSRV that is only a few bytes
-	uint BufferIndex = ({DataInterfaceName}_InputStreamStart + VertexIndex) & ({DataInterfaceName}_ColorIndexMask);
+	uint BufferIndex = (VertexIndex) & ({DataInterfaceName}_ColorIndexMask);
 	return {DataInterfaceName}_ColorInputBuffer[BufferIndex] FMANUALFETCH_COLOR_COMPONENT_SWIZZLE;
 }
 
 float2 ReadUV_{DataInterfaceName}(uint VertexIndex, uint UVChannelIndex)
 {
-	uint BufferIndex = ({DataInterfaceName}_InputStreamStart + VertexIndex) * {DataInterfaceName}_NumUVChannels + UVChannelIndex;
+	uint BufferIndex = (VertexIndex) * {DataInterfaceName}_NumUVChannels + UVChannelIndex;
 	return {DataInterfaceName}_UVInputBuffer[BufferIndex];
 }
diff --git a/Engine/Plugins/Animation/DeformerGraph/Shaders/Private/DataInterfaceSkinnedMeshRead.ush b/Engine/Plugins/Animation/DeformerGraph/Shaders/Private/DataInterfaceSkinnedMeshRead.ush
--- a/Engine/Plugins/Animation/DeformerGraph/Shaders/Private/DataInterfaceSkinnedMeshRead.ush
+++ b/Engine/Plugins/Animation/DeformerGraph/Shaders/Private/DataInterfaceSkinnedMeshRead.ush
@@ -4,7 +4,6 @@
 #include "/Engine/Private/GpuSkinCommon.ush"
 
 uint {DataInterfaceName}_NumVertices;
-uint {DataInterfaceName}_InputStreamStart;
 
 RWBuffer<float> {DataInterfaceName}_PositionBufferUAV;
 RWBuffer<TANGENT_RWBUFFER_FORMAT> {DataInterfaceName}_TangentBufferUAV;
@@ -22,7 +21,7 @@ uint ReadNumVertices_{DataInterfaceName}()
 
 float3 ReadPosition_{DataInterfaceName}(uint VertexIndex)
 {
-	uint BufferIndex = ({DataInterfaceName}_InputStreamStart + VertexIndex) * 3;
+	uint BufferIndex = (VertexIndex) * 3;
 
 #if READABLE_OUTPUT_BUFFERS & OPTIMUS_SKINNED_MESH_READ_POSITION
 	return float3(
@@ -39,7 +38,7 @@ float3 ReadPosition_{DataInterfaceName}(uint VertexIndex)
 
 float4 ReadTangentX_{DataInterfaceName}(uint VertexIndex)
 {
-	uint BufferIndex = ({DataInterfaceName}_InputStreamStart + VertexIndex) * 2;
+	uint BufferIndex = (VertexIndex) * 2;
 	
 #if READABLE_OUTPUT_BUFFERS & OPTIMUS_SKINNED_MESH_READ_TANGENTS
 	return TangentBias_SkinCache({DataInterfaceName}_TangentBufferUAV[BufferIndex]) ;
@@ -50,7 +49,7 @@ float4 ReadTangentX_{DataInterfaceName}(uint VertexIndex)
 
 float4 ReadTangentZ_{DataInterfaceName}(uint VertexIndex)
 {
-	uint BufferIndex = ({DataInterfaceName}_InputStreamStart + VertexIndex) * 2;
+	uint BufferIndex = (VertexIndex) * 2;
 #if READABLE_OUTPUT_BUFFERS & OPTIMUS_SKINNED_MESH_READ_TANGENTS
 	return TangentBias_SkinCache({DataInterfaceName}_TangentBufferUAV[BufferIndex + 1]) ;
 #else	
@@ -61,7 +60,7 @@ float4 ReadTangentZ_{DataInterfaceName}(uint VertexIndex)
 float4 ReadColor_{DataInterfaceName}(uint VertexIndex)
 {
 	// Index Mask is used here in case ColorStaticBuffer is bound to the GWhiteVertexBufferWithSRV that is only a few bytes
-	uint BufferIndex = ({DataInterfaceName}_InputStreamStart + VertexIndex) & ({DataInterfaceName}_ColorIndexMask);
+	uint BufferIndex = (VertexIndex) & ({DataInterfaceName}_ColorIndexMask);
 #if READABLE_OUTPUT_BUFFERS & OPTIMUS_SKINNED_MESH_READ_COLOR
 	return {DataInterfaceName}_ColorBufferUAV[BufferIndex] FMANUALFETCH_COLOR_COMPONENT_SWIZZLE;
 #else	
diff --git a/Engine/Plugins/Animation/DeformerGraph/Shaders/Private/DataInterfaceSkinnedMeshVertexAttribute.ush b/Engine/Plugins/Animation/DeformerGraph/Shaders/Private/DataInterfaceSkinnedMeshVertexAttribute.ush
--- a/Engine/Plugins/Animation/DeformerGraph/Shaders/Private/DataInterfaceSkinnedMeshVertexAttribute.ush
+++ b/Engine/Plugins/Animation/DeformerGraph/Shaders/Private/DataInterfaceSkinnedMeshVertexAttribute.ush
@@ -1,7 +1,6 @@
 // Copyright Epic Games, Inc. All Rights Reserved.
 
 uint {DataInterfaceName}_NumVertices;
-uint {DataInterfaceName}_InputStreamStart;
 uint {DataInterfaceName}_bIsValid;
 Buffer<float> {DataInterfaceName}_ValueBuffer;
 
@@ -14,7 +13,7 @@ float ReadValue_{DataInterfaceName}(uint VertexIndex)
 {
 	if ({DataInterfaceName}_bIsValid)
 	{
-		uint BufferIndex = ({DataInterfaceName}_InputStreamStart + VertexIndex);
+		uint BufferIndex = (VertexIndex);
 		return {DataInterfaceName}_ValueBuffer[BufferIndex];
 	}
 	else
diff --git a/Engine/Plugins/Animation/DeformerGraph/Shaders/Private/DataInterfaceSkinnedMeshWrite.ush b/Engine/Plugins/Animation/DeformerGraph/Shaders/Private/DataInterfaceSkinnedMeshWrite.ush
--- a/Engine/Plugins/Animation/DeformerGraph/Shaders/Private/DataInterfaceSkinnedMeshWrite.ush
+++ b/Engine/Plugins/Animation/DeformerGraph/Shaders/Private/DataInterfaceSkinnedMeshWrite.ush
@@ -4,7 +4,6 @@
 #include "/Engine/Private/GpuSkinCommon.ush"
 
 uint {DataInterfaceName}_NumVertices;
-uint {DataInterfaceName}_OutputStreamStart;
 RWBuffer<float> {DataInterfaceName}_PositionBufferUAV;
 RWBuffer<TANGENT_RWBUFFER_FORMAT> {DataInterfaceName}_TangentBufferUAV;
 RWBuffer<UNORM float4> {DataInterfaceName}_ColorBufferUAV;
@@ -16,7 +15,7 @@ uint ReadNumVertices_{DataInterfaceName}()
 
 void WritePosition_{DataInterfaceName}(uint VertexIndex, float3 Position)
 {
-	uint BufferIndex = ({DataInterfaceName}_OutputStreamStart + VertexIndex) * 3;
+	uint BufferIndex = (VertexIndex) * 3;
 	{DataInterfaceName}_PositionBufferUAV[BufferIndex + 0] = Position.x;
 	{DataInterfaceName}_PositionBufferUAV[BufferIndex + 1] = Position.y;
 	{DataInterfaceName}_PositionBufferUAV[BufferIndex + 2] = Position.z;
@@ -24,18 +23,18 @@ void WritePosition_{DataInterfaceName}(uint VertexIndex, float3 Position)
 
 void WriteTangentX_{DataInterfaceName}(uint VertexIndex, float4 TangentX)
 {
-	uint BufferIndex = ({DataInterfaceName}_OutputStreamStart + VertexIndex) * 2;
+	uint BufferIndex = (VertexIndex) * 2;
 	{DataInterfaceName}_TangentBufferUAV[BufferIndex + 0] = TangentUnbias_SkinCache(TangentX);
 }
 
 void WriteTangentZ_{DataInterfaceName}(uint VertexIndex, float4 TangentZ)
 {
-	uint BufferIndex = ({DataInterfaceName}_OutputStreamStart + VertexIndex) * 2;
+	uint BufferIndex = (VertexIndex) * 2;
 	{DataInterfaceName}_TangentBufferUAV[BufferIndex + 1] = TangentUnbias_SkinCache(TangentZ);
 }
 
 void WriteColor_{DataInterfaceName}(uint VertexIndex, float4 Color)
 {
-	uint BufferIndex = {DataInterfaceName}_OutputStreamStart + VertexIndex;
+	uint BufferIndex = VertexIndex;
 	{DataInterfaceName}_ColorBufferUAV[BufferIndex] = Color FMANUALFETCH_COLOR_COMPONENT_SWIZZLE;
 }
diff --git a/Engine/Plugins/Animation/DeformerGraph/Source/OptimusCore/Private/DataInterfaces/OptimusDataInterfaceAdvancedSkeleton.cpp b/Engine/Plugins/Animation/DeformerGraph/Source/OptimusCore/Private/DataInterfaces/OptimusDataInterfaceAdvancedSkeleton.cpp
--- a/Engine/Plugins/Animation/DeformerGraph/Source/OptimusCore/Private/DataInterfaces/OptimusDataInterfaceAdvancedSkeleton.cpp
+++ b/Engine/Plugins/Animation/DeformerGraph/Source/OptimusCore/Private/DataInterfaces/OptimusDataInterfaceAdvancedSkeleton.cpp
@@ -378,9 +378,7 @@ void UOptimusAdvancedSkeletonDataInterface::GetSupportedInputs(TArray<FShaderFun
 
 BEGIN_SHADER_PARAMETER_STRUCT(FAdvancedSkeletonDataInterfaceDefaultParameters, )
 	SHADER_PARAMETER(uint32, NumVertices)
-	SHADER_PARAMETER(uint32, InputStreamStart)
 	SHADER_PARAMETER(uint32, NumBoneInfluences)
-	SHADER_PARAMETER(uint32, InputWeightStart)
 	SHADER_PARAMETER(uint32, InputWeightStride)
 	SHADER_PARAMETER(uint32, InputWeightIndexSize)
 	SHADER_PARAMETER_SRV(Buffer<float4>, BoneMatrices)
@@ -1016,34 +1014,29 @@ void FOptimusAdvancedSkeletonDataProviderProxy::GatherDispatchData(FDispatchData
 	FSkeletalMeshLODRenderData const* LodRenderData = &SkeletalMeshRenderData.LODRenderData[LodIndex];
 
 	FRHIShaderResourceView* NullSRVBinding = GWhiteVertexBufferWithSRV->ShaderResourceViewRHI.GetReference();
-	
+
+	FSkinWeightVertexBuffer const* WeightBuffer = LodRenderData->GetSkinWeightVertexBuffer();
+	if (LodRenderData->SkinWeightProfilesData.ContainsProfile(SkinWeightProfile))
+	{
+		const FSkinWeightProfileStack ProfileStack{SkinWeightProfile};
+		WeightBuffer = LodRenderData->SkinWeightProfilesData.GetOverrideBuffer(ProfileStack);
+	}
+	check(WeightBuffer != nullptr);
+		
+	FRHIShaderResourceView* SkinWeightBufferSRV = WeightBuffer->GetDataVertexBuffer()->GetSRV();
+	const bool bUnlimitedBoneInfluences = WeightBuffer->GetBoneInfluenceType() == GPUSkinBoneInfluenceType::UnlimitedBoneInfluence;
+	FRHIShaderResourceView* InputWeightLookupStreamSRV = bUnlimitedBoneInfluences ? WeightBuffer->GetLookupVertexBuffer()->GetSRV() : nullptr;
 
 	for (int32 InvocationIndex = 0; InvocationIndex < InDispatchData.NumInvocations; ++InvocationIndex)
 	{
 		uint8* ParameterBufferPtr = (InDispatchData.ParameterBuffer + InDispatchData.ParameterBufferOffset + InDispatchData.ParameterBufferStride * InvocationIndex);
 		
-		FSkelMeshRenderSection const& RenderSection = LodRenderData->RenderSections[InvocationIndex];
-
 		const bool bPreviousFrame = false;
 		FRHIShaderResourceView* BoneBufferSRV = FSkeletalMeshDeformerHelpers::GetBoneBufferForReading(SkeletalMeshObject, LodIndex, InvocationIndex, bPreviousFrame);
-
-		FSkinWeightVertexBuffer const* WeightBuffer = LodRenderData->GetSkinWeightVertexBuffer();
-		if (LodRenderData->SkinWeightProfilesData.ContainsProfile(SkinWeightProfile))
-		{
-			const FSkinWeightProfileStack ProfileStack{SkinWeightProfile};
-			WeightBuffer = LodRenderData->SkinWeightProfilesData.GetOverrideBuffer(ProfileStack);
-		}
-		check(WeightBuffer != nullptr);
-		
-		FRHIShaderResourceView* SkinWeightBufferSRV = WeightBuffer->GetDataVertexBuffer()->GetSRV();
-		const bool bUnlimitedBoneInfluences = WeightBuffer->GetBoneInfluenceType() == GPUSkinBoneInfluenceType::UnlimitedBoneInfluence;
-		FRHIShaderResourceView* InputWeightLookupStreamSRV = bUnlimitedBoneInfluences ? WeightBuffer->GetLookupVertexBuffer()->GetSRV() : nullptr;
-
+	
 		FDefaultParameters* Parameters = reinterpret_cast<FDefaultParameters*>(ParameterBufferPtr);
-		Parameters->NumVertices = RenderSection.NumVertices;
-		Parameters->InputStreamStart = RenderSection.BaseVertexIndex;
+		Parameters->NumVertices = LodRenderData->GetNumVertices();
 		Parameters->NumBoneInfluences = WeightBuffer->GetMaxBoneInfluences();
-		Parameters->InputWeightStart = (WeightBuffer->GetConstantInfluencesVertexStride() * RenderSection.GetVertexBufferIndex()) / sizeof(float);
 		Parameters->InputWeightStride = WeightBuffer->GetConstantInfluencesVertexStride();
 		Parameters->InputWeightIndexSize = WeightBuffer->GetBoneIndexByteSize() | (WeightBuffer->GetBoneWeightByteSize() << 8);
 		Parameters->BoneMatrices = BoneBufferSRV != nullptr ? BoneBufferSRV : NullSRVBinding;
diff --git a/Engine/Plugins/Animation/DeformerGraph/Source/OptimusCore/Private/DataInterfaces/OptimusDataInterfaceCloth.cpp b/Engine/Plugins/Animation/DeformerGraph/Source/OptimusCore/Private/DataInterfaces/OptimusDataInterfaceCloth.cpp
--- a/Engine/Plugins/Animation/DeformerGraph/Source/OptimusCore/Private/DataInterfaces/OptimusDataInterfaceCloth.cpp
+++ b/Engine/Plugins/Animation/DeformerGraph/Source/OptimusCore/Private/DataInterfaces/OptimusDataInterfaceCloth.cpp
@@ -72,6 +72,7 @@ void UOptimusClothDataInterface::GetSupportedInputs(TArray<FShaderFunctionDefini
 
 BEGIN_SHADER_PARAMETER_STRUCT(FClothDataInterfaceParameters, )
 	SHADER_PARAMETER(uint32, NumVertices)
+	SHADER_PARAMETER(uint32, BaseVertexIndex)
 	SHADER_PARAMETER(uint32, InputStreamStart)
 	SHADER_PARAMETER(uint32, NumInfluencesPerVertex)
 	SHADER_PARAMETER(float, ClothBlendWeight)
@@ -215,7 +216,8 @@ void FOptimusClothDataProviderProxy::GatherDispatchData(FDispatchData const& InD
 		const bool bValidCloth = (ClothBuffers.ClothInfluenceBuffer != nullptr) && (ClothBuffers.ClothSimulatedPositionAndNormalBuffer != nullptr);
 
 		FParameters& Parameters = ParameterArray[InvocationIndex];
-		Parameters.NumVertices = RenderSection.NumVertices;
+		Parameters.NumVertices = LodRenderData->GetNumVertices();
+		Parameters.BaseVertexIndex = RenderSection.BaseVertexIndex;
 		Parameters.InputStreamStart = ClothBuffers.ClothInfluenceBufferOffset;
 		Parameters.ClothBlendWeight = bValidCloth ? ClothBlendWeight : 0.f;
 		Parameters.MeshScale = bValidCloth ? MeshScale : FVector3f::OneVector;
diff --git a/Engine/Plugins/Animation/DeformerGraph/Source/OptimusCore/Private/DataInterfaces/OptimusDataInterfaceConnectivity.cpp b/Engine/Plugins/Animation/DeformerGraph/Source/OptimusCore/Private/DataInterfaces/OptimusDataInterfaceConnectivity.cpp
--- a/Engine/Plugins/Animation/DeformerGraph/Source/OptimusCore/Private/DataInterfaces/OptimusDataInterfaceConnectivity.cpp
+++ b/Engine/Plugins/Animation/DeformerGraph/Source/OptimusCore/Private/DataInterfaces/OptimusDataInterfaceConnectivity.cpp
@@ -168,7 +168,6 @@ void UOptimusConnectivityDataInterface::GetSupportedInputs(TArray<FShaderFunctio
 
 BEGIN_SHADER_PARAMETER_STRUCT(FConnectivityDataInterfaceParameters, )
 SHADER_PARAMETER(uint32, NumVertices)
-SHADER_PARAMETER(uint32, InputStreamStart)
 SHADER_PARAMETER(uint32, MaxConnectedVertexCount)
 SHADER_PARAMETER_RDG_BUFFER_SRV(Buffer<uint>, ConnectivityBuffer)
 END_SHADER_PARAMETER_STRUCT()
@@ -280,11 +279,8 @@ void FOptimusConnectivityDataProviderProxy::GatherDispatchData(FDispatchData con
 	const TStridedView<FParameters> ParameterArray = MakeStridedParameterView<FParameters>(InDispatchData);
 	for (int32 InvocationIndex = 0; InvocationIndex < ParameterArray.Num(); ++InvocationIndex)
 	{
-		FSkelMeshRenderSection const& RenderSection = LodRenderData->RenderSections[InvocationIndex];
-
 		FParameters& Parameters = ParameterArray[InvocationIndex];
-		Parameters.NumVertices = RenderSection.NumVertices;
-		Parameters.InputStreamStart = RenderSection.BaseVertexIndex;
+		Parameters.NumVertices = LodRenderData->GetNumVertices();
 		Parameters.MaxConnectedVertexCount = UOptimusConnectivityDataInterface::MaxConnectedVertexCount;
 		Parameters.ConnectivityBuffer = ConnectivityBufferSRV;
 	}
diff --git a/Engine/Plugins/Animation/DeformerGraph/Source/OptimusCore/Private/DataInterfaces/OptimusDataInterfaceConnectivity.h b/Engine/Plugins/Animation/DeformerGraph/Source/OptimusCore/Private/DataInterfaces/OptimusDataInterfaceConnectivity.h
--- a/Engine/Plugins/Animation/DeformerGraph/Source/OptimusCore/Private/DataInterfaces/OptimusDataInterfaceConnectivity.h
+++ b/Engine/Plugins/Animation/DeformerGraph/Source/OptimusCore/Private/DataInterfaces/OptimusDataInterfaceConnectivity.h
@@ -28,6 +28,7 @@ public:
 	
 	//~ Begin UComputeDataInterface Interface
 	TCHAR const* GetClassName() const override { return TEXT("Connectivity"); }
+	bool CanSupportUnifiedDispatch() const override { return true; }
 	void GetSupportedInputs(TArray<FShaderFunctionDefinition>& OutFunctions) const override;
 	void GetShaderParameters(TCHAR const* UID, FShaderParametersMetadataBuilder& InOutBuilder, FShaderParametersMetadataAllocations& InOutAllocations) const override;
 	TCHAR const* GetShaderVirtualPath() const override;
diff --git a/Engine/Plugins/Animation/DeformerGraph/Source/OptimusCore/Private/DataInterfaces/OptimusDataInterfaceCopyKernel.cpp b/Engine/Plugins/Animation/DeformerGraph/Source/OptimusCore/Private/DataInterfaces/OptimusDataInterfaceCopyKernel.cpp
--- a/Engine/Plugins/Animation/DeformerGraph/Source/OptimusCore/Private/DataInterfaces/OptimusDataInterfaceCopyKernel.cpp
+++ b/Engine/Plugins/Animation/DeformerGraph/Source/OptimusCore/Private/DataInterfaces/OptimusDataInterfaceCopyKernel.cpp
@@ -4,65 +4,73 @@
 
 #include "OptimusComponentSource.h"
 #include "OptimusExpressionEvaluator.h"
+#include "OptimusHelpers.h"
+#include "OptimusObjectVersion.h"
 #include "ShaderParameterMetadataBuilder.h"
 #include "ComputeFramework/ComputeMetadataBuilder.h"
 #include "ComputeFramework/ShaderParameterMetadataAllocation.h"
 #include "ComputeFramework/ShaderParamTypeDefinition.h"
+#include "ShaderCore.h"
+#include "ShaderCompilerCore.h"
 
 #include UE_INLINE_GENERATED_CPP_BY_NAME(OptimusDataInterfaceCopyKernel)
 
-const FString UOptimusCopyKernelDataInterface::NumThreadsReservedName = TEXT("NumThreads");
+
+const TCHAR* UOptimusCopyKernelDataInterface:: ReadNumThreadsFunctionName = TEXT("ReadNumThreads");
+const TCHAR* UOptimusCopyKernelDataInterface:: ReadNumThreadsPerInvocationFunctionName = TEXT("ReadNumThreadsPerInvocation");
+const TCHAR* UOptimusCopyKernelDataInterface:: ReadThreadIndexOffsetFunctionName = TEXT("ReadThreadIndexOffset");
+
 
 void UOptimusCopyKernelDataInterface::GetSupportedInputs(TArray<FShaderFunctionDefinition>& OutFunctions) const
 {
 	OutFunctions.AddDefaulted_GetRef()
-		.SetName(FString::Printf(TEXT("Read%s"), *NumThreadsReservedName))
+		.SetName(ReadNumThreadsFunctionName)
 		.AddReturnType(FShaderValueType::Get(EShaderFundamentalType::Int, 3));
+
+	OutFunctions.AddDefaulted_GetRef()
+		.SetName(ReadNumThreadsPerInvocationFunctionName)
+		.AddReturnType(FShaderValueType::Get(EShaderFundamentalType::Int));
+	
+	OutFunctions.AddDefaulted_GetRef()
+		.SetName(ReadThreadIndexOffsetFunctionName)
+		.AddReturnType(FShaderValueType::Get(EShaderFundamentalType::Int));
 }
 
+BEGIN_SHADER_PARAMETER_STRUCT(FCopyKernelDataInterfaceParameters, )
+	SHADER_PARAMETER(FIntVector3, NumThreads)
+	SHADER_PARAMETER(int32, NumThreadsPerInvocation)
+	SHADER_PARAMETER(int32, ThreadIndexOffset)
+END_SHADER_PARAMETER_STRUCT()
+
+
 void UOptimusCopyKernelDataInterface::GetShaderParameters(TCHAR const* UID,
 	FShaderParametersMetadataBuilder& InOutBuilder, FShaderParametersMetadataAllocations& InOutAllocations) const
 {
-	FShaderParametersMetadataBuilder Builder;
-	TArray<FShaderParametersMetadata*> DummyNestedStructs;
-	ComputeFramework::AddParamForType(Builder, *NumThreadsReservedName, FShaderValueType::Get(EShaderFundamentalType::Int, 3),DummyNestedStructs);
-
-	FShaderParametersMetadata* ShaderParameterMetadata = Builder.Build(FShaderParametersMetadata::EUseCase::ShaderParameterStruct, TEXT("UCopyKernelDataInterface"));
-
-	InOutAllocations.ShaderParameterMetadatas.Add(ShaderParameterMetadata);
-	InOutAllocations.ShaderParameterMetadatas.Append(DummyNestedStructs);
+	InOutBuilder.AddNestedStruct<FCopyKernelDataInterfaceParameters>(UID);
+}
 
-	// Add the generated nested struct to our builder.
-	InOutBuilder.AddNestedStruct(UID, ShaderParameterMetadata);
+const TCHAR* UOptimusCopyKernelDataInterface::TemplateFilePath = TEXT("/Plugin/Optimus/Private/DataInterfaceCopyKernel.ush");
 
+const TCHAR* UOptimusCopyKernelDataInterface::GetShaderVirtualPath() const
+{
+	return TemplateFilePath;
 }
 
 void UOptimusCopyKernelDataInterface::GetShaderHash(FString& InOutKey) const
 {
-	// UComputeGraph::BuildKernelSource hashes the result of GetHLSL()
-	// Only append additional hashes here if the HLSL contains any additional includes	
+	GetShaderFileHash(TemplateFilePath, EShaderPlatform::SP_PCD3D_SM5).AppendString(InOutKey);
 }
 
 void UOptimusCopyKernelDataInterface::GetHLSL(FString& OutHLSL, FString const& InDataInterfaceName) const
 {
-	const FString TypeName = FShaderValueType::Get(EShaderFundamentalType::Int, 3)->ToString();
-
-	if (ensure(!TypeName.IsEmpty()))
+	TMap<FString, FStringFormatArg> TemplateArgs =
 	{
-		// Add uniforms.
-		OutHLSL += FString::Printf(TEXT("%s %s_%s;\n"), 
-			*TypeName, 
-			*InDataInterfaceName, 
-			*NumThreadsReservedName);
-			
-		// Add function getters.
-		OutHLSL += FString::Printf(TEXT("%s Read%s_%s()\n{\n\treturn %s_%s;\n}\n"), 
-			*TypeName,
-			*NumThreadsReservedName,
-			*InDataInterfaceName, 
-			*InDataInterfaceName, 
-			*NumThreadsReservedName);
-	}
+		{ TEXT("DataInterfaceName"), InDataInterfaceName },
+	};
+
+	FString TemplateFile;
+	LoadShaderSourceFile(TemplateFilePath, EShaderPlatform::SP_PCD3D_SM5, &TemplateFile, nullptr);
+	OutHLSL += FString::Format(*TemplateFile, TemplateArgs);
 }
 
 UComputeDataProvider* UOptimusCopyKernelDataInterface::CreateDataProvider(TObjectPtr<UObject> InBinding,
@@ -85,198 +93,56 @@ void UOptimusCopyKernelDataInterface::SetComponentBinding(const UOptimusComponen
 	ComponentSourceBinding = InBinding;
 }
 
-void UOptimusCopyKernelDataProvider::InitFromDataInterface(const UOptimusCopyKernelDataInterface* InDataInterface, const UObject* InBinding)
+const TCHAR* UOptimusCopyKernelDataInterface::GetReadNumThreadsFunctionName() const
 {
-	WeakComponent = Cast<UActorComponent>(InBinding);
-	WeakDataInterface = InDataInterface;
+	return ReadNumThreadsFunctionName;
 }
 
-FComputeDataProviderRenderProxy* UOptimusCopyKernelDataProvider::GetRenderProxy()
+const TCHAR* UOptimusCopyKernelDataInterface::GetReadNumThreadsPerInvocationFunctionName() const
 {
-	TArray<int32> InvocationCounts;
-	int32 TotalThreadCount = 0;
-	
-	if (!GetInvocationThreadCounts(InvocationCounts, TotalThreadCount))
-	{
-		InvocationCounts.Reset();
-	}
-	
-	FOptimusCopyKernelDataProviderProxy* Proxy = new FOptimusCopyKernelDataProviderProxy(MoveTemp(InvocationCounts), TotalThreadCount);
-	return Proxy;
+	return ReadNumThreadsPerInvocationFunctionName;
 }
 
-bool UOptimusCopyKernelDataProvider::GetInvocationThreadCounts(
-	TArray<int32>& OutInvocationThreadCount,
-	int32& OutTotalThreadCount
-	) const
+const TCHAR* UOptimusCopyKernelDataInterface::GetReadThreadIndexOffsetFunctionName() const
 {
-	if (!WeakDataInterface.IsValid())
-	{
-		return false;
-	}
-
-	if (!WeakDataInterface->ComponentSourceBinding.IsValid())
-	{
-		return false;
-	}
-	
-	const UOptimusComponentSource* ComponentSource = WeakDataInterface->ComponentSourceBinding->GetComponentSource();
-	const FString& NumThreadsExpression = WeakDataInterface->NumThreadsExpression;
-	
-	if (!WeakComponent.IsValid() || !ComponentSource)
-	{
-		return false;
-	}
-	
-	const UActorComponent* Component = WeakComponent.Get();
-	int32 LodIndex = ComponentSource->GetLodIndex(Component);
-	int32 NumInvocations = ComponentSource->GetDefaultNumInvocations(Component, LodIndex);
-
-	// In some cases the component can be set to a unusable state intentionally
-	// such as when the OptimusEditor shuts down,
-	// in which case we don't have data to work with so simply do nothing
-	if (NumInvocations < 1)
-	{
-		return false;
-	}
-	
-	TMap<FName, float> EngineConstants;
-	TMap<FName, TArray<int32>> ElementCountsPerDomain;
-
-	for(FName ExecutionDomain: ComponentSource->GetExecutionDomains())
-	{
-		EngineConstants.Add(ExecutionDomain, 0);
-
-		TArray<int32>& ElementCounts = ElementCountsPerDomain.FindOrAdd(ExecutionDomain);
-		if (!ComponentSource->GetComponentElementCountsForExecutionDomain(
-				ExecutionDomain,
-				Component,
-				LodIndex,
-				ElementCounts))
-		{
-			return false;
-		}
-
-		if (!ensure(NumInvocations == ElementCounts.Num()))
-		{
-			// Component source needs to provide as many values for each of its execution domains as there are invocations
-			return false;
-		}
-	}
-
-	using namespace Optimus::Expression;
-
-	FEngine Engine;
-	TVariant<FExpressionObject, FParseError> ParseResult = Engine.Parse(NumThreadsExpression, [EngineConstants](FName InName)->TOptional<float>
-	{
-		if (const float* Value = EngineConstants.Find(InName))
-		{
-			return *Value;
-		};
-
-		return {};
-	});
-	
-	if (ParseResult.IsType<FParseError>())
-	{
-		return false;
-	}
-
-	OutInvocationThreadCount.Reset(NumInvocations);
-	for (int32 Index = 0; Index < NumInvocations; Index++)
-	{
-		for (TPair<FName, float>& Constant: EngineConstants)
-		{
-			const FName ConstantName = Constant.Key;
-			float& Value = Constant.Value;
-			
-			const TArray<int32>& ElementCounts = ElementCountsPerDomain[ConstantName]; 
-			Value = ElementCounts.IsValidIndex(Index) ? ElementCounts[Index] : 1;
-		}
-		
-		const int32 Count = static_cast<int32>(Engine.Execute(ParseResult.Get<FExpressionObject>(),
-			[EngineConstants](FName InName)->TOptional<float>
-			{
-				if (const float* Value = EngineConstants.Find(InName))
-				{
-					return *Value;
-				};
-
-				return {};
-			}));
-
-		if (Count < 0)
-		{
-			return false;
-		}
-		
-		OutInvocationThreadCount.Add(Count);
-	}
+	return ReadThreadIndexOffsetFunctionName;
+}
 
-	OutTotalThreadCount = 0;
-	for (int32 Count : OutInvocationThreadCount)
-	{
-		OutTotalThreadCount += Count;
-	}
+void UOptimusCopyKernelDataProvider::InitFromDataInterface(const UOptimusCopyKernelDataInterface* InDataInterface, const UObject* InBinding)
+{
+	WeakComponent = Cast<UActorComponent>(InBinding);
+	WeakComponentSource = InDataInterface->ComponentSourceBinding->GetComponentSource();
+	ParseResult = Optimus::ParseExecutionDomainExpression(InDataInterface->NumThreadsExpression, WeakComponentSource);
+}
 
-	// Extra validation for unified dispatch
+FComputeDataProviderRenderProxy* UOptimusCopyKernelDataProvider::GetRenderProxy()
+{
+	TArray<int32> InvocationCounts;
 
-	// We want to make sure the results are the same for unified and non-unified
-	// Do the sum across invocations first and then evaluate the expression
-	for (TPair<FName, float>& Constant : EngineConstants)
-	{
-		const FName ConstantName = Constant.Key;
-		float& Value = Constant.Value;
-		Value = 0;
-		const TArray<int32>& ElementCounts = ElementCountsPerDomain[ConstantName];
-		
-		for (int32 Count : ElementCounts )
-		{
-			Value += Count;
-		}
-	}
-	
-	
-	const int32 TotalElementCountForUnifiedDispatch = static_cast<int32>(Engine.Execute(ParseResult.Get<FExpressionObject>(), 
-		[EngineConstants](FName InName)->TOptional<float>
-		{
-			if (const float* Value = EngineConstants.Find(InName))
-			{
-				return *Value;
-			};
-
-			return {};
-		}));
+	Optimus::EvaluateExecutionDomainExpressionParseResult(ParseResult, WeakComponentSource, WeakComponent,InvocationCounts);
 	
-	if (OutTotalThreadCount != TotalElementCountForUnifiedDispatch)
-	{
-		return false;
-	}
-
-	return true;	
+	FOptimusCopyKernelDataProviderProxy* Proxy = new FOptimusCopyKernelDataProviderProxy(MoveTemp(InvocationCounts));
+	return Proxy;
 }
 
 FOptimusCopyKernelDataProviderProxy::FOptimusCopyKernelDataProviderProxy(
-	TArray<int32>&& InInvocationThreadCounts,
-	int32 InTotalThreadCount
+	TArray<int32>&& InInvocationThreadCounts
 	) :
-	InvocationThreadCounts(MoveTemp(InInvocationThreadCounts)),
-	TotalThreadCount(InTotalThreadCount)
+	InvocationThreadCounts(MoveTemp(InInvocationThreadCounts))
 {
+	for (int32 Count : InvocationThreadCounts)
+	{
+		TotalThreadCount += Count;
+	}
 }
 
 bool FOptimusCopyKernelDataProviderProxy::IsValid(FValidationData const& InValidationData) const
 {
-	if (InvocationThreadCounts.Num() == 0)
-	{
-		return false;
-	}
-
 	if (TotalThreadCount <= 0)
 	{
 		return false;
 	}
-	
+
 	return true;
 }
 
@@ -292,22 +158,20 @@ int32 FOptimusCopyKernelDataProviderProxy::GetDispatchThreadCount(TArray<FIntVec
 
 void FOptimusCopyKernelDataProviderProxy::GatherDispatchData(FDispatchData const& InDispatchData)
 {
-	for (int32 InvocationIndex = 0; InvocationIndex < InDispatchData.NumInvocations; ++InvocationIndex)
+	const TStridedView<FParameters> ParameterArray = MakeStridedParameterView<FParameters>(InDispatchData);
+
+	int32 NumDispatchedThreads = 0;
+	for (int32 InvocationIndex = 0; InvocationIndex < ParameterArray.Num(); ++InvocationIndex)
 	{
-		int32 NumThreads;
-		if (InDispatchData.bUnifiedDispatch)
-		{
-			NumThreads = TotalThreadCount;
-		}
-		else
-		{
-			NumThreads = InvocationThreadCounts[InvocationIndex];
-		}
-
-		uint8* ParameterBuffer = (InDispatchData.ParameterBuffer + InDispatchData.ParameterBufferOffset + InDispatchData.ParameterBufferStride * InvocationIndex);
+		FParameters& Parameters = ParameterArray[InvocationIndex];
+		Parameters.NumThreads = FIntVector3(TotalThreadCount, 1, 1);
+		
+		int32 NumThreadsPerInvocation = InDispatchData.bUnifiedDispatch ? TotalThreadCount : InvocationThreadCounts[InvocationIndex] ;
+		Parameters.NumThreadsPerInvocation = NumThreadsPerInvocation;
+		
+		Parameters.ThreadIndexOffset = InDispatchData.bUnifiedDispatch ? 0 : NumDispatchedThreads;
 		
-		const int32 ThreadCount3D[] = {NumThreads, 1, 1};
-		FMemory::Memcpy(ParameterBuffer, ThreadCount3D, 3 * sizeof(int32));
+		NumDispatchedThreads += NumThreadsPerInvocation;
 	}
 }
 
diff --git a/Engine/Plugins/Animation/DeformerGraph/Source/OptimusCore/Private/DataInterfaces/OptimusDataInterfaceCopyKernel.h b/Engine/Plugins/Animation/DeformerGraph/Source/OptimusCore/Private/DataInterfaces/OptimusDataInterfaceCopyKernel.h
--- a/Engine/Plugins/Animation/DeformerGraph/Source/OptimusCore/Private/DataInterfaces/OptimusDataInterfaceCopyKernel.h
+++ b/Engine/Plugins/Animation/DeformerGraph/Source/OptimusCore/Private/DataInterfaces/OptimusDataInterfaceCopyKernel.h
@@ -7,11 +7,13 @@
 #include "OptimusComputeDataInterface.h"
 #include "OptimusDataType.h"
 #include "ComputeFramework/ComputeDataProvider.h"
+#include "OptimusExpressionEvaluator.h"
 
 #include "OptimusDataInterfaceCopyKernel.generated.h"
 
 class UOptimusComponentSource;
 class UOptimusComponentSourceBinding;
+class FCopyKernelDataInterfaceParameters;
 
 UCLASS(Category = ComputeFramework)
 class OPTIMUSCORE_API UOptimusCopyKernelDataInterface :
@@ -21,14 +23,13 @@ class OPTIMUSCORE_API UOptimusCopyKernelDataInterface :
 	GENERATED_BODY()
 
 public:
-	static const FString NumThreadsReservedName;
-	
 	//~ Begin UComputeDataInterface Interface
 	TCHAR const* GetClassName() const override { return TEXT("CopyKernelData"); }
 	bool IsExecutionInterface() const override { return true; }
 	bool CanSupportUnifiedDispatch() const override { return true; }
 	void GetSupportedInputs(TArray<FShaderFunctionDefinition>& OutFunctions) const override;
 	void GetShaderParameters(TCHAR const* UID, FShaderParametersMetadataBuilder& InOutBuilder, FShaderParametersMetadataAllocations& InOutAllocations) const override;
+	const TCHAR* GetShaderVirtualPath() const override;
 	void GetShaderHash(FString& InOutKey) const override;
 	void GetHLSL(FString& OutHLSL, FString const& InDataInterfaceName) const override;
 	UComputeDataProvider* CreateDataProvider(TObjectPtr<UObject> InBinding, uint64 InInputMask, uint64 InOutputMask) const override;
@@ -37,6 +38,9 @@ public:
 	//~ Begin IOptimusComputeKernelDataInterface Interface
 	void SetExecutionDomain(const FString& InExecutionDomain) override;
 	void SetComponentBinding(const UOptimusComponentSourceBinding* InBinding) override;
+	const TCHAR* GetReadNumThreadsFunctionName() const override;
+	const TCHAR* GetReadNumThreadsPerInvocationFunctionName() const override;
+	const TCHAR* GetReadThreadIndexOffsetFunctionName() const override;
 	//~ End IOptimusComputeKernelDataInterface Interface
 	
 	UPROPERTY()
@@ -44,6 +48,13 @@ public:
 	
 	UPROPERTY()
 	FString NumThreadsExpression;
+	
+protected:	
+	static const TCHAR* TemplateFilePath;
+	
+	static const TCHAR* ReadNumThreadsFunctionName;
+	static const TCHAR* ReadNumThreadsPerInvocationFunctionName;
+	static const TCHAR* ReadThreadIndexOffsetFunctionName;
 };
 
 /** Compute Framework Data Provider for each copy kernel. */
@@ -60,23 +71,17 @@ public:
 	FComputeDataProviderRenderProxy* GetRenderProxy() override;
 	//~ End UComputeDataProvider Interface
 
-	TWeakObjectPtr<const UActorComponent> WeakComponent = nullptr;
-
-	TWeakObjectPtr<const UOptimusCopyKernelDataInterface> WeakDataInterface = nullptr;
-
 protected:
-	bool GetInvocationThreadCounts(
-		TArray<int32>& OutInvocationThreadCount,
-		int32& OutTotalThreadCount
-		) const;
+	TWeakObjectPtr<const UActorComponent> WeakComponent = nullptr;
+	TWeakObjectPtr<const UOptimusComponentSource> WeakComponentSource = nullptr;
+	TVariant<Optimus::Expression::FExpressionObject, Optimus::Expression::FParseError> ParseResult;
 };
 
 class FOptimusCopyKernelDataProviderProxy : public FComputeDataProviderRenderProxy
 {
 public:
 	FOptimusCopyKernelDataProviderProxy(
-		TArray<int32>&& InInvocationThreadCounts,
-		int32 InTotalThreadCount
+		TArray<int32>&& InInvocationThreadCounts
 		);
 
 	//~ Begin FComputeDataProviderRenderProxy Interface
@@ -84,9 +89,11 @@ public:
 	int32 GetDispatchThreadCount(TArray<FIntVector>& InOutThreadCounts) const override;
 	void GatherDispatchData(FDispatchData const& InDispatchData) override;
 	//~ End FComputeDataProviderRenderProxy Interface
-
+	
+private:
 	TArray<int32> InvocationThreadCounts;
-	int32 TotalThreadCount;
+	int32 TotalThreadCount = 0; 
+	using FParameters = FCopyKernelDataInterfaceParameters;
 };
 
 
diff --git a/Engine/Plugins/Animation/DeformerGraph/Source/OptimusCore/Private/DataInterfaces/OptimusDataInterfaceCustomComputeKernel.cpp b/Engine/Plugins/Animation/DeformerGraph/Source/OptimusCore/Private/DataInterfaces/OptimusDataInterfaceCustomComputeKernel.cpp
--- a/Engine/Plugins/Animation/DeformerGraph/Source/OptimusCore/Private/DataInterfaces/OptimusDataInterfaceCustomComputeKernel.cpp
+++ b/Engine/Plugins/Animation/DeformerGraph/Source/OptimusCore/Private/DataInterfaces/OptimusDataInterfaceCustomComputeKernel.cpp
@@ -5,67 +5,72 @@
 #include "OptimusComponentSource.h"
 #include "OptimusDeformerInstance.h"
 #include "OptimusExpressionEvaluator.h"
+#include "OptimusHelpers.h"
 #include "ShaderParameterMetadataBuilder.h"
 #include "ComputeFramework/ComputeMetadataBuilder.h"
 #include "ComputeFramework/ShaderParameterMetadataAllocation.h"
 #include "ComputeFramework/ShaderParamTypeDefinition.h"
+#include "ShaderCore.h"
+#include "ShaderCompilerCore.h"
 
 #include "Nodes/OptimusNode_CustomComputeKernel.h"
 
 #include UE_INLINE_GENERATED_CPP_BY_NAME(OptimusDataInterfaceCustomComputeKernel)
 
-const FString UOptimusCustomComputeKernelDataInterface::NumThreadsReservedName = TEXT("NumThreads");
+
+const TCHAR* UOptimusCustomComputeKernelDataInterface::ReadNumThreadsFunctionName = TEXT("ReadNumThreads");
+const TCHAR* UOptimusCustomComputeKernelDataInterface::ReadNumThreadsPerInvocationFunctionName = TEXT("ReadNumThreadsPerInvocation");
+const TCHAR* UOptimusCustomComputeKernelDataInterface::ReadThreadIndexOffsetFunctionName = TEXT("ReadThreadIndexOffset");
 
 void UOptimusCustomComputeKernelDataInterface::GetSupportedInputs(TArray<FShaderFunctionDefinition>& OutFunctions) const
 {
 	OutFunctions.AddDefaulted_GetRef()
-		.SetName(FString::Printf(TEXT("Read%s"), *NumThreadsReservedName))
+		.SetName(ReadNumThreadsFunctionName)
 		.AddReturnType(FShaderValueType::Get(EShaderFundamentalType::Int, 3));
+	
+	OutFunctions.AddDefaulted_GetRef()
+		.SetName(ReadNumThreadsPerInvocationFunctionName)
+		.AddReturnType(FShaderValueType::Get(EShaderFundamentalType::Int));
+	
+	OutFunctions.AddDefaulted_GetRef()
+		.SetName(ReadThreadIndexOffsetFunctionName)
+		.AddReturnType(FShaderValueType::Get(EShaderFundamentalType::Int));
 }
 
+BEGIN_SHADER_PARAMETER_STRUCT(FCustomComputeKernelDataInterfaceParameters, )
+	SHADER_PARAMETER(FIntVector3, NumThreads)
+	SHADER_PARAMETER(int32, NumThreadsPerInvocation)
+	SHADER_PARAMETER(int32, ThreadIndexOffset)
+END_SHADER_PARAMETER_STRUCT()
+
 void UOptimusCustomComputeKernelDataInterface::GetShaderParameters(TCHAR const* UID,
 	FShaderParametersMetadataBuilder& InOutBuilder, FShaderParametersMetadataAllocations& InOutAllocations) const
 {
-	FShaderParametersMetadataBuilder Builder;
-	TArray<FShaderParametersMetadata*> DummyNestedStructs;
-	ComputeFramework::AddParamForType(Builder, *NumThreadsReservedName, FShaderValueType::Get(EShaderFundamentalType::Int, 3),DummyNestedStructs);
-
-	FShaderParametersMetadata* ShaderParameterMetadata = Builder.Build(FShaderParametersMetadata::EUseCase::ShaderParameterStruct, TEXT("UCustomComputeKernelDataInterface"));
-
-	InOutAllocations.ShaderParameterMetadatas.Add(ShaderParameterMetadata);
-	InOutAllocations.ShaderParameterMetadatas.Append(DummyNestedStructs);
+	InOutBuilder.AddNestedStruct<FCustomComputeKernelDataInterfaceParameters>(UID);
+}
 
-	// Add the generated nested struct to our builder.
-	InOutBuilder.AddNestedStruct(UID, ShaderParameterMetadata);
+const TCHAR* UOptimusCustomComputeKernelDataInterface::TemplateFilePath = TEXT("/Plugin/Optimus/Private/DataInterfaceCustomComputeKernel.ush");
 
+const TCHAR* UOptimusCustomComputeKernelDataInterface::GetShaderVirtualPath() const
+{
+	return TemplateFilePath;
 }
 
 void UOptimusCustomComputeKernelDataInterface::GetShaderHash(FString& InOutKey) const
 {
-	// UComputeGraph::BuildKernelSource hashes the result of GetHLSL()
-	// Only append additional hashes here if the HLSL contains any additional includes	
+	GetShaderFileHash(TemplateFilePath, EShaderPlatform::SP_PCD3D_SM5).AppendString(InOutKey);
 }
 
 void UOptimusCustomComputeKernelDataInterface::GetHLSL(FString& OutHLSL, FString const& InDataInterfaceName) const
 {
-	const FString TypeName = FShaderValueType::Get(EShaderFundamentalType::Int, 3)->ToString();
-
-	if (ensure(!TypeName.IsEmpty()))
+	TMap<FString, FStringFormatArg> TemplateArgs =
 	{
-		// Add uniforms.
-		OutHLSL += FString::Printf(TEXT("%s %s_%s;\n"), 
-			*TypeName, 
-			*InDataInterfaceName, 
-			*NumThreadsReservedName);
-			
-		// Add function getters.
-		OutHLSL += FString::Printf(TEXT("%s Read%s_%s()\n{\n\treturn %s_%s;\n}\n"), 
-			*TypeName,
-			*NumThreadsReservedName,
-			*InDataInterfaceName, 
-			*InDataInterfaceName, 
-			*NumThreadsReservedName);
-	}
+		{ TEXT("DataInterfaceName"), InDataInterfaceName },
+	};
+
+	FString TemplateFile;
+	LoadShaderSourceFile(TemplateFilePath, EShaderPlatform::SP_PCD3D_SM5, &TemplateFile, nullptr);
+	OutHLSL += FString::Format(*TemplateFile, TemplateArgs);
 }
 
 UComputeDataProvider* UOptimusCustomComputeKernelDataInterface::CreateDataProvider(TObjectPtr<UObject> InBinding,
@@ -88,198 +93,56 @@ void UOptimusCustomComputeKernelDataInterface::SetComponentBinding(const UOptimu
 	ComponentSourceBinding = InBinding;
 }
 
-void UOptimusCustomComputeKernelDataProvider::InitFromDataInterface(const UOptimusCustomComputeKernelDataInterface* InDataInterface, const UObject* InBinding)
+const TCHAR* UOptimusCustomComputeKernelDataInterface::GetReadNumThreadsFunctionName() const
 {
-	WeakComponent = Cast<UActorComponent>(InBinding);
-	WeakDataInterface = InDataInterface;
+	return ReadNumThreadsFunctionName;
 }
 
-FComputeDataProviderRenderProxy* UOptimusCustomComputeKernelDataProvider::GetRenderProxy()
+const TCHAR* UOptimusCustomComputeKernelDataInterface::GetReadNumThreadsPerInvocationFunctionName() const
 {
-	TArray<int32> InvocationCounts;
-	int32 TotalThreadCount = 0;
-	
-	if (!GetInvocationThreadCounts(InvocationCounts, TotalThreadCount))
-	{
-		InvocationCounts.Reset();
-	}
-	
-	FOptimusCustomComputeKernelDataProviderProxy* Proxy = new FOptimusCustomComputeKernelDataProviderProxy(MoveTemp(InvocationCounts), TotalThreadCount);
-	return Proxy;
+	return ReadNumThreadsPerInvocationFunctionName;
 }
 
-bool UOptimusCustomComputeKernelDataProvider::GetInvocationThreadCounts(
-	TArray<int32>& OutInvocationThreadCount,
-	int32& OutTotalThreadCount
-	) const
+const TCHAR* UOptimusCustomComputeKernelDataInterface::GetReadThreadIndexOffsetFunctionName() const
 {
-	if (!WeakDataInterface.IsValid())
-	{
-		return false;
-	}
-
-	if (!WeakDataInterface->ComponentSourceBinding.IsValid())
-	{
-		return false;
-	}
-	
-	const UOptimusComponentSource* ComponentSource = WeakDataInterface->ComponentSourceBinding->GetComponentSource();
-	const FString& NumThreadsExpression = WeakDataInterface->NumThreadsExpression;
-	
-	if (!WeakComponent.IsValid() || !ComponentSource)
-	{
-		return false;
-	}
-	
-	const UActorComponent* Component = WeakComponent.Get();
-	int32 LodIndex = ComponentSource->GetLodIndex(Component);
-	int32 NumInvocations = ComponentSource->GetDefaultNumInvocations(Component, LodIndex);
-
-	// In some cases the component can be set to a unusable state intentionally
-	// such as when the OptimusEditor shuts down,
-	// in which case we don't have data to work with so simply do nothing
-	if (NumInvocations < 1)
-	{
-		return false;
-	}
-	
-	TMap<FName, float> EngineConstants;
-	TMap<FName, TArray<int32>> ElementCountsPerDomain;
-
-	for(FName ExecutionDomain: ComponentSource->GetExecutionDomains())
-	{
-		EngineConstants.Add(ExecutionDomain, 0);
-
-		TArray<int32>& ElementCounts = ElementCountsPerDomain.FindOrAdd(ExecutionDomain);
-		if (!ComponentSource->GetComponentElementCountsForExecutionDomain(
-				ExecutionDomain,
-				Component,
-				LodIndex,
-				ElementCounts))
-		{
-			return false;
-		}
-
-		if (!ensure(NumInvocations == ElementCounts.Num()))
-		{
-			// Component source needs to provide as many values for each of its execution domains as there are invocations
-			return false;
-		}
-	}
-
-	using namespace Optimus::Expression;
-
-	FEngine Engine;
-	TVariant<FExpressionObject, FParseError> ParseResult = Engine.Parse(NumThreadsExpression, [EngineConstants](FName InName)->TOptional<float>
-	{
-		if (const float* Value = EngineConstants.Find(InName))
-		{
-			return *Value;
-		};
-
-		return {};
-	});
-	
-	if (ParseResult.IsType<FParseError>())
-	{
-		return false;
-	}
-
-	OutInvocationThreadCount.Reset(NumInvocations);
-	for (int32 Index = 0; Index < NumInvocations; Index++)
-	{
-		for (TPair<FName, float>& Constant: EngineConstants)
-		{
-			const FName ConstantName = Constant.Key;
-			float& Value = Constant.Value;
-			
-			const TArray<int32>& ElementCounts = ElementCountsPerDomain[ConstantName]; 
-			Value = ElementCounts.IsValidIndex(Index) ? ElementCounts[Index] : 1;
-		}
-		
-		const int32 Count = static_cast<int32>(Engine.Execute(ParseResult.Get<FExpressionObject>(),
-			[EngineConstants](FName InName)->TOptional<float>
-			{
-				if (const float* Value = EngineConstants.Find(InName))
-				{
-					return *Value;
-				};
-
-				return {};
-			}));
-
-		if (Count < 0)
-		{
-			return false;
-		}
-		
-		OutInvocationThreadCount.Add(Count);
-	}
+	return ReadThreadIndexOffsetFunctionName;
+}
 
-	OutTotalThreadCount = 0;
-	for (int32 Count : OutInvocationThreadCount)
-	{
-		OutTotalThreadCount += Count;
-	}
+void UOptimusCustomComputeKernelDataProvider::InitFromDataInterface(const UOptimusCustomComputeKernelDataInterface* InDataInterface, const UObject* InBinding)
+{
+	WeakComponent = Cast<UActorComponent>(InBinding);
+	WeakComponentSource = InDataInterface->ComponentSourceBinding->GetComponentSource();
+	ParseResult = Optimus::ParseExecutionDomainExpression(InDataInterface->NumThreadsExpression, WeakComponentSource);
+}
 
-	// Extra validation for unified dispatch
+FComputeDataProviderRenderProxy* UOptimusCustomComputeKernelDataProvider::GetRenderProxy()
+{
+	TArray<int32> InvocationCounts;
 
-	// We want to make sure the results are the same for unified and non-unified
-	// Do the sum across invocations first and then evaluate the expression
-	for (TPair<FName, float>& Constant : EngineConstants)
-	{
-		const FName ConstantName = Constant.Key;
-		float& Value = Constant.Value;
-		Value = 0;
-		const TArray<int32>& ElementCounts = ElementCountsPerDomain[ConstantName];
-		
-		for (int32 Count : ElementCounts )
-		{
-			Value += Count;
-		}
-	}
+	Optimus::EvaluateExecutionDomainExpressionParseResult(ParseResult, WeakComponentSource, WeakComponent,InvocationCounts);
 	
-	
-	const int32 TotalElementCountForUnifiedDispatch = static_cast<int32>(Engine.Execute(ParseResult.Get<FExpressionObject>(), 
-		[EngineConstants](FName InName)->TOptional<float>
-		{
-			if (const float* Value = EngineConstants.Find(InName))
-			{
-				return *Value;
-			};
-
-			return {};
-		}));
-	
-	if (OutTotalThreadCount != TotalElementCountForUnifiedDispatch)
-	{
-		return false;
-	}
-
-	return true;	
+	FOptimusCustomComputeKernelDataProviderProxy* Proxy = new FOptimusCustomComputeKernelDataProviderProxy(MoveTemp(InvocationCounts));
+	return Proxy;
 }
 
 FOptimusCustomComputeKernelDataProviderProxy::FOptimusCustomComputeKernelDataProviderProxy(
-	TArray<int32>&& InInvocationThreadCounts,
-	int32 InTotalThreadCount
+	TArray<int32>&& InInvocationThreadCounts
 	) :
-	InvocationThreadCounts(MoveTemp(InInvocationThreadCounts)),
-	TotalThreadCount(InTotalThreadCount)
+	InvocationThreadCounts(MoveTemp(InInvocationThreadCounts))
 {
+	for (int32 Count : InvocationThreadCounts)
+	{
+		TotalThreadCount += Count;
+	}
 }
 
 bool FOptimusCustomComputeKernelDataProviderProxy::IsValid(FValidationData const& InValidationData) const
 {
-	if (InvocationThreadCounts.Num() == 0)
-	{
-		return false;
-	}
-
 	if (TotalThreadCount <= 0)
 	{
 		return false;
 	}
-	
+
 	return true;
 }
 
@@ -295,22 +158,20 @@ int32 FOptimusCustomComputeKernelDataProviderProxy::GetDispatchThreadCount(TArra
 
 void FOptimusCustomComputeKernelDataProviderProxy::GatherDispatchData(FDispatchData const& InDispatchData)
 {
-	for (int32 InvocationIndex = 0; InvocationIndex < InDispatchData.NumInvocations; ++InvocationIndex)
+	const TStridedView<FParameters> ParameterArray = MakeStridedParameterView<FParameters>(InDispatchData);
+
+	int32 NumDispatchedThreads = 0;
+	for (int32 InvocationIndex = 0; InvocationIndex < ParameterArray.Num(); ++InvocationIndex)
 	{
-		int32 NumThreads;
-		if (InDispatchData.bUnifiedDispatch)
-		{
-			NumThreads = TotalThreadCount;
-		}
-		else
-		{
-			NumThreads = InvocationThreadCounts[InvocationIndex];
-		}
-
-		uint8* ParameterBuffer = (InDispatchData.ParameterBuffer + InDispatchData.ParameterBufferOffset + InDispatchData.ParameterBufferStride * InvocationIndex);
+		FParameters& Parameters = ParameterArray[InvocationIndex];
+		Parameters.NumThreads = FIntVector3(TotalThreadCount, 1, 1);
+		
+		int32 NumThreadsPerInvocation = InDispatchData.bUnifiedDispatch ? TotalThreadCount : InvocationThreadCounts[InvocationIndex] ;
+		Parameters.NumThreadsPerInvocation = NumThreadsPerInvocation;
+		
+		Parameters.ThreadIndexOffset = InDispatchData.bUnifiedDispatch ? 0 : NumDispatchedThreads;
 		
-		const int32 ThreadCount3D[] = {NumThreads, 1, 1};
-		FMemory::Memcpy(ParameterBuffer, ThreadCount3D, 3 * sizeof(int32));
+		NumDispatchedThreads += NumThreadsPerInvocation;
 	}
 }
 
diff --git a/Engine/Plugins/Animation/DeformerGraph/Source/OptimusCore/Private/DataInterfaces/OptimusDataInterfaceCustomComputeKernel.h b/Engine/Plugins/Animation/DeformerGraph/Source/OptimusCore/Private/DataInterfaces/OptimusDataInterfaceCustomComputeKernel.h
--- a/Engine/Plugins/Animation/DeformerGraph/Source/OptimusCore/Private/DataInterfaces/OptimusDataInterfaceCustomComputeKernel.h
+++ b/Engine/Plugins/Animation/DeformerGraph/Source/OptimusCore/Private/DataInterfaces/OptimusDataInterfaceCustomComputeKernel.h
@@ -7,6 +7,7 @@
 #include "OptimusComputeDataInterface.h"
 #include "OptimusConstant.h"
 #include "OptimusDataType.h"
+#include "OptimusExpressionEvaluator.h"
 #include "ComputeFramework/ComputeDataProvider.h"
 
 #include "OptimusDataInterfaceCustomComputeKernel.generated.h"
@@ -14,6 +15,7 @@
 class UOptimusNode_CustomComputeKernel;
 class UOptimusComponentSource;
 class UOptimusComponentSourceBinding;
+class FCustomComputeKernelDataInterfaceParameters;
 
 UCLASS(Category = ComputeFramework)
 class OPTIMUSCORE_API UOptimusCustomComputeKernelDataInterface :
@@ -22,15 +24,15 @@ class OPTIMUSCORE_API UOptimusCustomComputeKernelDataInterface :
 {
 	GENERATED_BODY()
 
-public:
-	static const FString NumThreadsReservedName;
-	
+public:	
 	//~ Begin UComputeDataInterface Interface
 	TCHAR const* GetClassName() const override { return TEXT("CustomComputeKernelData"); }
 	bool IsExecutionInterface() const override { return true; }
 	bool CanSupportUnifiedDispatch() const override { return true; }
+
 	void GetSupportedInputs(TArray<FShaderFunctionDefinition>& OutFunctions) const override;
 	void GetShaderParameters(TCHAR const* UID, FShaderParametersMetadataBuilder& InOutBuilder, FShaderParametersMetadataAllocations& InOutAllocations) const override;
+	const TCHAR* GetShaderVirtualPath() const override;
 	void GetShaderHash(FString& InOutKey) const override;
 	void GetHLSL(FString& OutHLSL, FString const& InDataInterfaceName) const override;
 	UComputeDataProvider* CreateDataProvider(TObjectPtr<UObject> InBinding, uint64 InInputMask, uint64 InOutputMask) const override;
@@ -39,6 +41,9 @@ public:
 	//~ Begin IOptimusComputeKernelDataInterface Interface
 	void SetExecutionDomain(const FString& InExecutionDomain) override;
 	void SetComponentBinding(const UOptimusComponentSourceBinding* InBinding) override;
+	const TCHAR* GetReadNumThreadsFunctionName() const override;
+	const TCHAR* GetReadNumThreadsPerInvocationFunctionName() const override;
+	const TCHAR* GetReadThreadIndexOffsetFunctionName() const override;
 	//~ End IOptimusComputeKernelDataInterface Interface
 	
 	UPROPERTY()
@@ -49,6 +54,13 @@ public:
 
 	UPROPERTY()
 	FOptimusConstantIdentifier ExecutionDomainConstantIdentifier_DEPRECATED;
+	
+protected:	
+	static const TCHAR* TemplateFilePath;
+	
+	static const TCHAR* ReadNumThreadsFunctionName;
+	static const TCHAR* ReadNumThreadsPerInvocationFunctionName;
+	static const TCHAR* ReadThreadIndexOffsetFunctionName;
 };
 
 /** Compute Framework Data Provider for each custom compute kernel. */
@@ -65,23 +77,17 @@ public:
 	FComputeDataProviderRenderProxy* GetRenderProxy() override;
 	//~ End UComputeDataProvider Interface
 
-	TWeakObjectPtr<const UActorComponent> WeakComponent = nullptr;
-
-	TWeakObjectPtr<const UOptimusCustomComputeKernelDataInterface> WeakDataInterface = nullptr;
-
 protected:
-	bool GetInvocationThreadCounts(
-		TArray<int32>& OutInvocationThreadCount,
-		int32& OutTotalThreadCount
-		) const;
+	TWeakObjectPtr<const UActorComponent> WeakComponent = nullptr;
+	TWeakObjectPtr<const UOptimusComponentSource> WeakComponentSource = nullptr;
+	TVariant<Optimus::Expression::FExpressionObject, Optimus::Expression::FParseError> ParseResult;
 };
 
 class FOptimusCustomComputeKernelDataProviderProxy : public FComputeDataProviderRenderProxy
 {
 public:
 	FOptimusCustomComputeKernelDataProviderProxy(
-		TArray<int32>&& InInvocationThreadCounts,
-		int32 InTotalThreadCount
+		TArray<int32>&& InInvocationThreadCounts
 		);
 
 	//~ Begin FComputeDataProviderRenderProxy Interface
@@ -89,9 +95,11 @@ public:
 	int32 GetDispatchThreadCount(TArray<FIntVector>& InOutThreadCounts) const override;
 	void GatherDispatchData(FDispatchData const& InDispatchData) override;
 	//~ End FComputeDataProviderRenderProxy Interface
-
+	
+private:
 	TArray<int32> InvocationThreadCounts;
-	int32 TotalThreadCount;
+	int32 TotalThreadCount = 0; 
+	using FParameters = FCustomComputeKernelDataInterfaceParameters;
 };
 
 
diff --git a/Engine/Plugins/Animation/DeformerGraph/Source/OptimusCore/Private/DataInterfaces/OptimusDataInterfaceDuplicateVertices.cpp b/Engine/Plugins/Animation/DeformerGraph/Source/OptimusCore/Private/DataInterfaces/OptimusDataInterfaceDuplicateVertices.cpp
--- a/Engine/Plugins/Animation/DeformerGraph/Source/OptimusCore/Private/DataInterfaces/OptimusDataInterfaceDuplicateVertices.cpp
+++ b/Engine/Plugins/Animation/DeformerGraph/Source/OptimusCore/Private/DataInterfaces/OptimusDataInterfaceDuplicateVertices.cpp
@@ -64,7 +64,7 @@ void UOptimusDuplicateVerticesDataInterface::GetSupportedInputs(TArray<FShaderFu
 BEGIN_SHADER_PARAMETER_STRUCT(FDuplicateVerticesDataInterfaceParameters, )
 	SHADER_PARAMETER(uint32, NumVertices)
 	SHADER_PARAMETER(uint32, NumDuplicateVertices)
-	SHADER_PARAMETER(uint32, InputStreamStart)
+	SHADER_PARAMETER(uint32, BaseVertexIndex)
 	SHADER_PARAMETER_SRV(Buffer<uint>, DuplicateVertexStartAndLength)
 	SHADER_PARAMETER_SRV(Buffer<uint>, DuplicateVertices)
 END_SHADER_PARAMETER_STRUCT()
@@ -191,9 +191,9 @@ void FOptimusDuplicateVerticesDataProviderProxy::GatherDispatchData(FDispatchDat
 		const bool bValidDuplicateIndices = (DuplicateVertexStartAndLengthSRV != nullptr) && (DuplicateVerticesSRV != nullptr);
 
 		FParameters& Parameters = ParameterArray[InvocationIndex];
-		Parameters.NumVertices = RenderSection.NumVertices;
-		Parameters.NumDuplicateVertices = bValidDuplicateIndices ? RenderSection.DuplicatedVerticesBuffer.DupVertData.Num() : 0;
-		Parameters.InputStreamStart = RenderSection.BaseVertexIndex;
+		Parameters.NumVertices = LodRenderData->GetNumVertices();
+		Parameters.NumDuplicateVertices = bValidDuplicateIndices ? RenderSection.DuplicatedVerticesBuffer.DupVertData.Num() : 0; 
+		Parameters.BaseVertexIndex = RenderSection.BaseVertexIndex;
 		Parameters.DuplicateVertexStartAndLength = bValidDuplicateIndices ? DuplicateVertexStartAndLengthSRV : NullSRVBinding;
 		Parameters.DuplicateVertices = bValidDuplicateIndices ? DuplicateVerticesSRV : NullSRVBinding;
 	}
diff --git a/Engine/Plugins/Animation/DeformerGraph/Source/OptimusCore/Private/DataInterfaces/OptimusDataInterfaceHalfEdge.cpp b/Engine/Plugins/Animation/DeformerGraph/Source/OptimusCore/Private/DataInterfaces/OptimusDataInterfaceHalfEdge.cpp
--- a/Engine/Plugins/Animation/DeformerGraph/Source/OptimusCore/Private/DataInterfaces/OptimusDataInterfaceHalfEdge.cpp
+++ b/Engine/Plugins/Animation/DeformerGraph/Source/OptimusCore/Private/DataInterfaces/OptimusDataInterfaceHalfEdge.cpp
@@ -113,8 +113,6 @@ void UOptimusHalfEdgeDataInterface::GetSupportedInputs(TArray<FShaderFunctionDef
 BEGIN_SHADER_PARAMETER_STRUCT(FHalfEdgeDataInterfaceParameters, )
 	SHADER_PARAMETER(uint32, NumVertices)
 	SHADER_PARAMETER(uint32, NumTriangles)
-	SHADER_PARAMETER(uint32, IndexBufferStart)
-	SHADER_PARAMETER(uint32, InputStreamStart)
 	SHADER_PARAMETER_RDG_BUFFER_SRV(StructuredBuffer<int>, VertexToEdgeBuffer)
 	SHADER_PARAMETER_RDG_BUFFER_SRV(StructuredBuffer<int>, EdgeToTwinEdgeBuffer)
 
@@ -335,13 +333,9 @@ void FOptimusHalfEdgeDataProviderProxy::GatherDispatchData(FDispatchData const&
 	const TStridedView<FParameters> ParameterArray = MakeStridedParameterView<FParameters>(InDispatchData);
 	for (int32 InvocationIndex = 0; InvocationIndex < ParameterArray.Num(); ++InvocationIndex)
 	{
-		FSkelMeshRenderSection const& RenderSection = LodRenderData->RenderSections[InvocationIndex];
-
 		FParameters& Parameters = ParameterArray[InvocationIndex];
-		Parameters.NumVertices = InDispatchData.bUnifiedDispatch ? LodRenderData->GetNumVertices() : RenderSection.NumVertices;
-		Parameters.NumTriangles = InDispatchData.bUnifiedDispatch ? LodRenderData->GetTotalFaces() : RenderSection.NumTriangles;
-		Parameters.IndexBufferStart = InDispatchData.bUnifiedDispatch ? 0 :RenderSection.BaseIndex;
-		Parameters.InputStreamStart = InDispatchData.bUnifiedDispatch ? 0 : RenderSection.BaseVertexIndex;
+		Parameters.NumVertices = LodRenderData->GetNumVertices();
+		Parameters.NumTriangles = LodRenderData->GetTotalFaces();
 		Parameters.VertexToEdgeBuffer = VertexToEdgeBufferSRV;
 		Parameters.EdgeToTwinEdgeBuffer = EdgeToTwinEdgeBufferSRV;
 
diff --git a/Engine/Plugins/Animation/DeformerGraph/Source/OptimusCore/Private/DataInterfaces/OptimusDataInterfaceMorphTarget.cpp b/Engine/Plugins/Animation/DeformerGraph/Source/OptimusCore/Private/DataInterfaces/OptimusDataInterfaceMorphTarget.cpp
--- a/Engine/Plugins/Animation/DeformerGraph/Source/OptimusCore/Private/DataInterfaces/OptimusDataInterfaceMorphTarget.cpp
+++ b/Engine/Plugins/Animation/DeformerGraph/Source/OptimusCore/Private/DataInterfaces/OptimusDataInterfaceMorphTarget.cpp
@@ -53,7 +53,6 @@ void UOptimusMorphTargetDataInterface::GetSupportedInputs(TArray<FShaderFunction
 
 BEGIN_SHADER_PARAMETER_STRUCT(FMorphTargetDataInterfaceParameters, )
 	SHADER_PARAMETER(uint32, NumVertices)
-	SHADER_PARAMETER(uint32, InputStreamStart)
 	SHADER_PARAMETER_SRV(Buffer<float>, MorphBuffer)
 END_SHADER_PARAMETER_STRUCT()
 
@@ -176,14 +175,11 @@ void FOptimusMorphTargetDataProviderProxy::GatherDispatchData(FDispatchData cons
 	const TStridedView<FParameters> ParameterArray = MakeStridedParameterView<FParameters>(InDispatchData);
 	for (int32 InvocationIndex = 0; InvocationIndex < ParameterArray.Num(); ++InvocationIndex)
 	{
-		FSkelMeshRenderSection const& RenderSection = LodRenderData->RenderSections[InvocationIndex];
-
 		const bool bPreviousFrame = false;
 		FRHIShaderResourceView* MorphBufferSRV = FSkeletalMeshDeformerHelpers::GetMorphTargetBufferForReading(SkeletalMeshObject, LodIndex, InvocationIndex, FrameNumber, bPreviousFrame);
 
 		FParameters& Parameters = ParameterArray[InvocationIndex];
-		Parameters.NumVertices = RenderSection.NumVertices;
-		Parameters.InputStreamStart = RenderSection.BaseVertexIndex;
+		Parameters.NumVertices = LodRenderData->GetNumVertices();
 		Parameters.MorphBuffer = MorphBufferSRV != nullptr ? MorphBufferSRV : NullSRVBinding;
 	}
 }
diff --git a/Engine/Plugins/Animation/DeformerGraph/Source/OptimusCore/Private/DataInterfaces/OptimusDataInterfaceMorphTarget.h b/Engine/Plugins/Animation/DeformerGraph/Source/OptimusCore/Private/DataInterfaces/OptimusDataInterfaceMorphTarget.h
--- a/Engine/Plugins/Animation/DeformerGraph/Source/OptimusCore/Private/DataInterfaces/OptimusDataInterfaceMorphTarget.h
+++ b/Engine/Plugins/Animation/DeformerGraph/Source/OptimusCore/Private/DataInterfaces/OptimusDataInterfaceMorphTarget.h
@@ -26,6 +26,7 @@ public:
 	
 	//~ Begin UComputeDataInterface Interface
 	TCHAR const* GetClassName() const override { return TEXT("MorphTarget"); }
+	bool CanSupportUnifiedDispatch() const override {return true;}
 	void GetSupportedInputs(TArray<FShaderFunctionDefinition>& OutFunctions) const override;
 	void GetShaderParameters(TCHAR const* UID, FShaderParametersMetadataBuilder& InOutBuilder, FShaderParametersMetadataAllocations& InOutAllocations) const override;
 	void GetPermutations(FComputeKernelPermutationVector& OutPermutationVector) const override;
diff --git a/Engine/Plugins/Animation/DeformerGraph/Source/OptimusCore/Private/DataInterfaces/OptimusDataInterfaceRawBuffer.cpp b/Engine/Plugins/Animation/DeformerGraph/Source/OptimusCore/Private/DataInterfaces/OptimusDataInterfaceRawBuffer.cpp
--- a/Engine/Plugins/Animation/DeformerGraph/Source/OptimusCore/Private/DataInterfaces/OptimusDataInterfaceRawBuffer.cpp
+++ b/Engine/Plugins/Animation/DeformerGraph/Source/OptimusCore/Private/DataInterfaces/OptimusDataInterfaceRawBuffer.cpp
@@ -132,7 +132,6 @@ int32 UOptimusRawBufferDataInterface::GetWriteValueOutputIndex(EOptimusBufferWri
 }
 
 BEGIN_SHADER_PARAMETER_STRUCT(FTransientBufferDataInterfaceParameters, )
-	SHADER_PARAMETER(uint32, StartOffset)
 	SHADER_PARAMETER(uint32, BufferSize)
 	SHADER_PARAMETER_RDG_BUFFER_SRV(StructuredBuffer<int>, BufferSRV)
 	SHADER_PARAMETER_RDG_BUFFER_UAV(RWStructuredBuffer<int>, BufferUAV)
@@ -144,7 +143,6 @@ void UOptimusTransientBufferDataInterface::GetShaderParameters(TCHAR const* UID,
 }
 
 BEGIN_SHADER_PARAMETER_STRUCT(FImplicitPersistentBufferDataInterfaceParameters, )
-	SHADER_PARAMETER(uint32, StartOffset)
 	SHADER_PARAMETER(uint32, BufferSize)
 	SHADER_PARAMETER_RDG_BUFFER_SRV(StructuredBuffer<int>, BufferSRV)
 	SHADER_PARAMETER_RDG_BUFFER_UAV(RWStructuredBuffer<int>, BufferUAV)
@@ -157,7 +155,6 @@ void UOptimusImplicitPersistentBufferDataInterface::GetShaderParameters(TCHAR co
 }
 
 BEGIN_SHADER_PARAMETER_STRUCT(FPersistentBufferDataInterfaceParameters, )
-	SHADER_PARAMETER(uint32, StartOffset)
 	SHADER_PARAMETER(uint32, BufferSize)
 	SHADER_PARAMETER_RDG_BUFFER_UAV(RWStructuredBuffer<int>, BufferUAV)
 END_SHADER_PARAMETER_STRUCT()
@@ -267,157 +264,7 @@ bool UOptimusRawBufferDataProvider::GetLodAndInvocationElementCounts(
 
 	OutLodIndex = ComponentSourcePtr->GetLodIndex(ComponentPtr);
 
-	switch(DataDomain.Type)
-	{
-	case EOptimusDataDomainType::Dimensional:
-		{
-			const FName ExecutionDomain = !DataDomain.DimensionNames.IsEmpty() ? DataDomain.DimensionNames[0] : NAME_None;
-			TArray<int32> Values;
-			if (ComponentSourcePtr->GetComponentElementCountsForExecutionDomain(ExecutionDomain, ComponentPtr, OutLodIndex, Values))
-			{
-				if (DataDomain.DimensionNames.Num() == 1)
-				{
-					OutInvocationElementCounts.Reset(Values.Num());
-					for (int32& Count: Values)
-					{
-						Count *= FMath::Max(1, DataDomain.Multiplier);
-						OutInvocationElementCounts.Add(Count);
-					}
-					return true;
-				}
-			}
-
-			return false;
-		}
-	case EOptimusDataDomainType::Expression:
-		{
-			TMap<FName, float> EngineConstants;
-			TMap<FName, TArray<int32>> ElementCountsPerDomain;
-
-			int32 NumInvocations = -1;
-			for(FName ExecutionDomain: ComponentSource->GetExecutionDomains())
-			{
-				EngineConstants.Add(ExecutionDomain, 0);
-
-				TArray<int32>& ElementCounts = ElementCountsPerDomain.FindOrAdd(ExecutionDomain);
-				if (!ComponentSourcePtr->GetComponentElementCountsForExecutionDomain(
-						ExecutionDomain,
-						ComponentPtr,
-						OutLodIndex,
-						ElementCounts))
-				{
-					return false;
-				}
-
-				if (NumInvocations == -1)
-				{
-					NumInvocations = ElementCounts.Num();
-				}
-				else if (NumInvocations != ElementCounts.Num())
-				{
-					// Component source needs to provide as many values for each of its execution domains as there are invocations
-					return false;
-				}
-			}
-
-			if (NumInvocations < 1)
-			{
-				return false;
-			}
-			
-			using namespace Optimus::Expression;
-
-			FEngine Engine;
-			TVariant<FExpressionObject, FParseError> ParseResult = Engine.Parse(DataDomain.Expression, [EngineConstants](FName InName)->TOptional<float>
-			{
-				if (const float* Value = EngineConstants.Find(InName))
-				{
-					return *Value;
-				};
-
-				return {};
-			});
-			
-			if (ParseResult.IsType<FParseError>())
-			{
-				return false;
-			}
-
-			OutInvocationElementCounts.Reset(NumInvocations);
-			for (int32 Index = 0; Index < NumInvocations; Index++)
-			{
-				for (TPair<FName, float>& Constant: EngineConstants)
-				{
-					const FName ConstantName = Constant.Key;
-					float& Value = Constant.Value;
-					
-					const TArray<int32>& ElementCounts = ElementCountsPerDomain[ConstantName]; 
-					Value = ElementCounts.IsValidIndex(Index) ? ElementCounts[Index] : 1;
-				}
-				
-				const int32 Count = static_cast<int32>(Engine.Execute(ParseResult.Get<FExpressionObject>(),
-					[EngineConstants](FName InName)->TOptional<float>
-					{
-						if (const float* Value = EngineConstants.Find(InName))
-						{
-							return *Value;
-						};
-
-						return {};
-					}));
-
-				if (Count < 0)
-				{
-					return false;
-				}
-				
-				OutInvocationElementCounts.Add(Count);
-			}
-
-			int32 TotalElementCount = 0;
-			for (int32 Count : OutInvocationElementCounts)
-			{
-				TotalElementCount += Count;
-			}
-
-			// Extra validation for unified dispatch
-
-			// We want to make sure the results are the same for unified and non-unified
-			// Do the sum across invocations first and then evaluate the expression
-			for (TPair<FName, float>& Constant : EngineConstants)
-			{
-				const FName ConstantName = Constant.Key;
-				float& Value = Constant.Value;
-				
-				Value = 0;
-				const TArray<int32>& ElementCounts = ElementCountsPerDomain[ConstantName];
-				
-				for (int32 Count : ElementCounts)
-				{
-					Value += Count;
-				}
-			}
-			
-			const int32 TotalElementCountForUnifiedDispatch = static_cast<int32>(Engine.Execute(ParseResult.Get<FExpressionObject>(),
-				[EngineConstants](FName InName)->TOptional<float>
-				{
-					if (const float* Value = EngineConstants.Find(InName))
-					{
-						return *Value;
-					};
-
-					return {};
-				}));
-			
-			if (TotalElementCount != TotalElementCountForUnifiedDispatch)
-			{
-				return false;
-			}
-
-			return true;
-		}
-	}
-	return false;
+	return Optimus::EvaluateExecutionDomainExpressionParseResult(DataDomainExpressionParseResult, ComponentSource, Component, OutInvocationElementCounts);
 }
 
 FComputeDataProviderRenderProxy* UOptimusTransientBufferDataProvider::GetRenderProxy()
@@ -517,15 +364,12 @@ void FOptimusTransientBufferDataProviderProxy::AllocateResources(FRDGBuilder& Gr
 void FOptimusTransientBufferDataProviderProxy::GatherDispatchData(FDispatchData const& InDispatchData)
 {
 	TStridedView<FParameters> ParameterArray = MakeStridedParameterView<FParameters>(InDispatchData);
-	for (int32 InvocationIndex = 0, StartOffset = 0; InvocationIndex < ParameterArray.Num(); ++InvocationIndex)
+	for (int32 InvocationIndex = 0; InvocationIndex < ParameterArray.Num(); ++InvocationIndex)
 	{
 		FParameters& Parameters = ParameterArray[InvocationIndex];
-		Parameters.StartOffset = InDispatchData.bUnifiedDispatch ? 0 : StartOffset;
-		Parameters.BufferSize = InDispatchData.bUnifiedDispatch ? TotalElementCount : InvocationElementCounts[InvocationIndex];
+		Parameters.BufferSize =TotalElementCount;
 		Parameters.BufferSRV = BufferSRV;
 		Parameters.BufferUAV = BufferUAV;
-		
-		StartOffset += InvocationElementCounts[InvocationIndex];
 	}
 }
 
@@ -596,15 +440,12 @@ void FOptimusImplicitPersistentBufferDataProviderProxy::AllocateResources(FRDGBu
 void FOptimusImplicitPersistentBufferDataProviderProxy::GatherDispatchData(FDispatchData const& InDispatchData)
 {
 	TStridedView<FParameters> ParameterArray = MakeStridedParameterView<FParameters>(InDispatchData);
-	for (int32 InvocationIndex = 0, StartOffset = 0; InvocationIndex < ParameterArray.Num(); ++InvocationIndex)
+	for (int32 InvocationIndex = 0; InvocationIndex < ParameterArray.Num(); ++InvocationIndex)
 	{
 		FParameters& Parameters = ParameterArray[InvocationIndex];
-		Parameters.StartOffset = InDispatchData.bUnifiedDispatch ? 0 : StartOffset;
-		Parameters.BufferSize = InDispatchData.bUnifiedDispatch ? TotalElementCount : InvocationElementCounts[InvocationIndex];
+		Parameters.BufferSize = TotalElementCount;
 		Parameters.BufferSRV = BufferSRV;
 		Parameters.BufferUAV = BufferUAV;
-		
-		StartOffset += InvocationElementCounts[InvocationIndex];
 	}
 }
 
@@ -670,16 +511,10 @@ void FOptimusPersistentBufferDataProviderProxy::AllocateResources(FRDGBuilder& G
 void FOptimusPersistentBufferDataProviderProxy::GatherDispatchData(FDispatchData const& InDispatchData)
 {
 	TStridedView<FParameters> ParameterArray = MakeStridedParameterView<FParameters>(InDispatchData);
-	for (int32 InvocationIndex = 0, StartOffset = 0; InvocationIndex < ParameterArray.Num(); ++InvocationIndex)
+	for (int32 InvocationIndex = 0;  InvocationIndex < ParameterArray.Num(); ++InvocationIndex)
 	{
 		FParameters& Parameters = ParameterArray[InvocationIndex];
-		Parameters.StartOffset = InDispatchData.bUnifiedDispatch ? 0 : StartOffset;
-		Parameters.BufferSize = InDispatchData.bUnifiedDispatch ? TotalElementCount : InvocationElementCounts[InvocationIndex];
+		Parameters.BufferSize = TotalElementCount;
 		Parameters.BufferUAV = BufferUAVs[InDispatchData.GraphKernelIndex];
-
-		if(!InDispatchData.bUnifiedDispatch)
-		{
-			StartOffset += InvocationElementCounts[InvocationIndex];
-		}
 	}
 }
diff --git a/Engine/Plugins/Animation/DeformerGraph/Source/OptimusCore/Private/DataInterfaces/OptimusDataInterfaceRawBuffer.h b/Engine/Plugins/Animation/DeformerGraph/Source/OptimusCore/Private/DataInterfaces/OptimusDataInterfaceRawBuffer.h
--- a/Engine/Plugins/Animation/DeformerGraph/Source/OptimusCore/Private/DataInterfaces/OptimusDataInterfaceRawBuffer.h
+++ b/Engine/Plugins/Animation/DeformerGraph/Source/OptimusCore/Private/DataInterfaces/OptimusDataInterfaceRawBuffer.h
@@ -9,6 +9,7 @@
 #include "OptimusConstant.h"
 #include "OptimusDataDomain.h"
 #include "OptimusDeformerInstance.h"
+#include "OptimusHelpers.h"
 #include "RenderGraphFwd.h"
 
 #include "OptimusDataInterfaceRawBuffer.generated.h"
@@ -90,9 +91,9 @@ protected:
 		{
 			Provider->Component = Component;
 			Provider->ComponentSource = GetComponentSource();
-			Provider->DataDomain = DataDomain;
 			Provider->ElementStride = ValueType->GetResourceElementSize();
 			Provider->RawStride = GetRawStride();
+			Provider->DataDomainExpressionParseResult = Optimus::ParseExecutionDomainExpression(DataDomain.AsExpression().Get({}), Provider->ComponentSource);
 		}
 		return Provider;
 	}
@@ -202,12 +203,11 @@ public:
 
 	TWeakObjectPtr<const UOptimusComponentSource> ComponentSource = nullptr;
 
-	/** The data domain this buffer covers */
-	FOptimusDataDomain DataDomain;
-
 	int32 ElementStride = 4;
 
 	int32 RawStride = 0;
+
+	TVariant<Optimus::Expression::FExpressionObject, Optimus::Expression::FParseError> DataDomainExpressionParseResult;
 };
 
 
diff --git a/Engine/Plugins/Animation/DeformerGraph/Source/OptimusCore/Private/DataInterfaces/OptimusDataInterfaceSkeleton.cpp b/Engine/Plugins/Animation/DeformerGraph/Source/OptimusCore/Private/DataInterfaces/OptimusDataInterfaceSkeleton.cpp
--- a/Engine/Plugins/Animation/DeformerGraph/Source/OptimusCore/Private/DataInterfaces/OptimusDataInterfaceSkeleton.cpp
+++ b/Engine/Plugins/Animation/DeformerGraph/Source/OptimusCore/Private/DataInterfaces/OptimusDataInterfaceSkeleton.cpp
@@ -67,9 +67,7 @@ void UOptimusSkeletonDataInterface::GetSupportedInputs(TArray<FShaderFunctionDef
 
 BEGIN_SHADER_PARAMETER_STRUCT(FSkeletonDataInterfaceParameters, )
 	SHADER_PARAMETER(uint32, NumVertices)
-	SHADER_PARAMETER(uint32, InputStreamStart)
 	SHADER_PARAMETER(uint32, NumBoneInfluences)
-	SHADER_PARAMETER(uint32, InputWeightStart)
 	SHADER_PARAMETER(uint32, InputWeightStride)
 	SHADER_PARAMETER(uint32, InputWeightIndexSize)
 	SHADER_PARAMETER_SRV(Buffer<float4>, BoneMatrices)
@@ -224,25 +222,21 @@ void FOptimusSkeletonDataProviderProxy::GatherDispatchData(FDispatchData const&
 
 	FRHIShaderResourceView* NullSRVBinding = GWhiteVertexBufferWithSRV->ShaderResourceViewRHI.GetReference();
 
+	FSkinWeightVertexBuffer const* WeightBuffer = LodRenderData->GetSkinWeightVertexBuffer();
+	check(WeightBuffer != nullptr);
+	FRHIShaderResourceView* SkinWeightBufferSRV = WeightBuffer->GetDataVertexBuffer()->GetSRV();
+	const bool bUnlimitedBoneInfluences = WeightBuffer->GetBoneInfluenceType() == GPUSkinBoneInfluenceType::UnlimitedBoneInfluence;
+	FRHIShaderResourceView* InputWeightLookupStreamSRV = bUnlimitedBoneInfluences ? WeightBuffer->GetLookupVertexBuffer()->GetSRV() : nullptr;
+
 	const TStridedView<FParameters> ParameterArray = MakeStridedParameterView<FParameters>(InDispatchData);
 	for (int32 InvocationIndex = 0; InvocationIndex < ParameterArray.Num(); ++InvocationIndex)
 	{
-		FSkelMeshRenderSection const& RenderSection = LodRenderData->RenderSections[InvocationIndex];
-
 		const bool bPreviousFrame = false;
 		FRHIShaderResourceView* BoneBufferSRV = FSkeletalMeshDeformerHelpers::GetBoneBufferForReading(SkeletalMeshObject, LodIndex, InvocationIndex, bPreviousFrame);
 
-		FSkinWeightVertexBuffer const* WeightBuffer = LodRenderData->GetSkinWeightVertexBuffer();
-		check(WeightBuffer != nullptr);
-		FRHIShaderResourceView* SkinWeightBufferSRV = WeightBuffer->GetDataVertexBuffer()->GetSRV();
-		const bool bUnlimitedBoneInfluences = WeightBuffer->GetBoneInfluenceType() == GPUSkinBoneInfluenceType::UnlimitedBoneInfluence;
-		FRHIShaderResourceView* InputWeightLookupStreamSRV = bUnlimitedBoneInfluences ? WeightBuffer->GetLookupVertexBuffer()->GetSRV() : nullptr;
-
 		FParameters& Parameters = ParameterArray[InvocationIndex];
-		Parameters.NumVertices = RenderSection.NumVertices;
-		Parameters.InputStreamStart = RenderSection.BaseVertexIndex;
+		Parameters.NumVertices = LodRenderData->GetNumVertices();
 		Parameters.NumBoneInfluences = WeightBuffer->GetMaxBoneInfluences();
-		Parameters.InputWeightStart = (WeightBuffer->GetConstantInfluencesVertexStride() * RenderSection.GetVertexBufferIndex()) / sizeof(float);
 		Parameters.InputWeightStride = WeightBuffer->GetConstantInfluencesVertexStride();
 		Parameters.InputWeightIndexSize = WeightBuffer->GetBoneIndexByteSize() | (WeightBuffer->GetBoneWeightByteSize() << 8);
 		Parameters.BoneMatrices = BoneBufferSRV != nullptr ? BoneBufferSRV : NullSRVBinding;
diff --git a/Engine/Plugins/Animation/DeformerGraph/Source/OptimusCore/Private/DataInterfaces/OptimusDataInterfaceSkinWeightsAsVertexMask.cpp b/Engine/Plugins/Animation/DeformerGraph/Source/OptimusCore/Private/DataInterfaces/OptimusDataInterfaceSkinWeightsAsVertexMask.cpp
--- a/Engine/Plugins/Animation/DeformerGraph/Source/OptimusCore/Private/DataInterfaces/OptimusDataInterfaceSkinWeightsAsVertexMask.cpp
+++ b/Engine/Plugins/Animation/DeformerGraph/Source/OptimusCore/Private/DataInterfaces/OptimusDataInterfaceSkinWeightsAsVertexMask.cpp
@@ -101,9 +101,7 @@ void UOptimusSkinWeightsAsVertexMaskDataInterface::GetSupportedInputs(TArray<FSh
 
 BEGIN_SHADER_PARAMETER_STRUCT(FSkinWeightsAsVertexMaskDataInterfaceParameters, )
 	SHADER_PARAMETER(uint32, NumVertices)
-	SHADER_PARAMETER(uint32, InputStreamStart)
 	SHADER_PARAMETER(uint32, NumBoneInfluences)
-	SHADER_PARAMETER(uint32, InputWeightStart)
 	SHADER_PARAMETER(uint32, InputWeightStride)
 	SHADER_PARAMETER(uint32, InputWeightIndexSize)
 	SHADER_PARAMETER_SRV(Buffer<uint>, InputWeightStream)
@@ -581,22 +579,19 @@ void FOptimusSkinWeightsAsVertexMaskDataProviderProxy::GatherDispatchData(FDispa
 
 	FRHIShaderResourceView* NullSRVBinding = GWhiteVertexBufferWithSRV->ShaderResourceViewRHI.GetReference();
 	
+	FSkinWeightVertexBuffer const* WeightBuffer = LodRenderData->GetSkinWeightVertexBuffer();
+	check(WeightBuffer != nullptr);
+	FRHIShaderResourceView* SkinWeightBufferSRV = WeightBuffer->GetDataVertexBuffer()->GetSRV();
+	const bool bUnlimitedBoneInfluences = WeightBuffer->GetBoneInfluenceType() == GPUSkinBoneInfluenceType::UnlimitedBoneInfluence;
+	FRHIShaderResourceView* InputWeightLookupStreamSRV = bUnlimitedBoneInfluences ? WeightBuffer->GetLookupVertexBuffer()->GetSRV() : nullptr;
+
+	
 	const TStridedView<FParameters> ParameterArray = MakeStridedParameterView<FParameters>(InDispatchData);
 	for (int32 InvocationIndex = 0; InvocationIndex < InDispatchData.NumInvocations; ++InvocationIndex)
 	{
-		FSkelMeshRenderSection const& RenderSection = LodRenderData->RenderSections[InvocationIndex];
-
-		FSkinWeightVertexBuffer const* WeightBuffer = LodRenderData->GetSkinWeightVertexBuffer();
-		check(WeightBuffer != nullptr);
-		FRHIShaderResourceView* SkinWeightBufferSRV = WeightBuffer->GetDataVertexBuffer()->GetSRV();
-		const bool bUnlimitedBoneInfluences = WeightBuffer->GetBoneInfluenceType() == GPUSkinBoneInfluenceType::UnlimitedBoneInfluence;
-		FRHIShaderResourceView* InputWeightLookupStreamSRV = bUnlimitedBoneInfluences ? WeightBuffer->GetLookupVertexBuffer()->GetSRV() : nullptr;
-
 		FParameters& Parameters = ParameterArray[InvocationIndex];
-		Parameters.NumVertices = RenderSection.NumVertices;
-		Parameters.InputStreamStart = RenderSection.BaseVertexIndex;
+		Parameters.NumVertices = LodRenderData->GetNumVertices();
 		Parameters.NumBoneInfluences = WeightBuffer->GetMaxBoneInfluences();
-		Parameters.InputWeightStart = (WeightBuffer->GetConstantInfluencesVertexStride() * RenderSection.GetVertexBufferIndex()) / sizeof(float);
 		Parameters.InputWeightStride = WeightBuffer->GetConstantInfluencesVertexStride();
 		Parameters.InputWeightIndexSize = WeightBuffer->GetBoneIndexByteSize() | (WeightBuffer->GetBoneWeightByteSize() << 8);
 		Parameters.InputWeightStream = SkinWeightBufferSRV != nullptr ? SkinWeightBufferSRV : NullSRVBinding;
diff --git a/Engine/Plugins/Animation/DeformerGraph/Source/OptimusCore/Private/DataInterfaces/OptimusDataInterfaceSkinWeightsAsVertexMask.h b/Engine/Plugins/Animation/DeformerGraph/Source/OptimusCore/Private/DataInterfaces/OptimusDataInterfaceSkinWeightsAsVertexMask.h
--- a/Engine/Plugins/Animation/DeformerGraph/Source/OptimusCore/Private/DataInterfaces/OptimusDataInterfaceSkinWeightsAsVertexMask.h
+++ b/Engine/Plugins/Animation/DeformerGraph/Source/OptimusCore/Private/DataInterfaces/OptimusDataInterfaceSkinWeightsAsVertexMask.h
@@ -36,6 +36,7 @@ public:
 
 	//~ Begin UComputeDataInterface Interface
 	TCHAR const* GetClassName() const override { return TEXT("SkinWeightMask"); }
+	bool CanSupportUnifiedDispatch() const override {return true;};
 	void GetSupportedInputs(TArray<FShaderFunctionDefinition>& OutFunctions) const override;
 	void GetShaderParameters(TCHAR const* UID, FShaderParametersMetadataBuilder& InOutBuilder, FShaderParametersMetadataAllocations& InOutAllocations) const override;
 	void GetPermutations(FComputeKernelPermutationVector& OutPermutationVector) const override;
diff --git a/Engine/Plugins/Animation/DeformerGraph/Source/OptimusCore/Private/DataInterfaces/OptimusDataInterfaceSkinnedMesh.cpp b/Engine/Plugins/Animation/DeformerGraph/Source/OptimusCore/Private/DataInterfaces/OptimusDataInterfaceSkinnedMesh.cpp
--- a/Engine/Plugins/Animation/DeformerGraph/Source/OptimusCore/Private/DataInterfaces/OptimusDataInterfaceSkinnedMesh.cpp
+++ b/Engine/Plugins/Animation/DeformerGraph/Source/OptimusCore/Private/DataInterfaces/OptimusDataInterfaceSkinnedMesh.cpp
@@ -95,8 +95,6 @@ BEGIN_SHADER_PARAMETER_STRUCT(FSkinnedMeshDataInterfaceParameters, )
 	SHADER_PARAMETER(uint32, NumVertices)
 	SHADER_PARAMETER(uint32, NumTriangles)
 	SHADER_PARAMETER(uint32, NumUVChannels)
-	SHADER_PARAMETER(uint32, IndexBufferStart)
-	SHADER_PARAMETER(uint32, InputStreamStart)
 	SHADER_PARAMETER_SRV(Buffer<uint>, IndexBuffer)
 	SHADER_PARAMETER_SRV(Buffer<float>, PositionInputBuffer)
 	SHADER_PARAMETER_SRV(Buffer<SNORM float4>, TangentInputBuffer)
@@ -182,8 +180,6 @@ void FOptimusSkinnedMeshDataProviderProxy::GatherDispatchData(FDispatchData cons
 	const TStridedView<FParameters> ParameterArray = MakeStridedParameterView<FParameters>(InDispatchData);
 	for (int32 InvocationIndex = 0; InvocationIndex < ParameterArray.Num(); ++InvocationIndex)
 	{
-		FSkelMeshRenderSection const& RenderSection = LodRenderData->RenderSections[InvocationIndex];
-
 		FRHIShaderResourceView* IndexBufferSRV = LodRenderData->MultiSizeIndexContainer.GetIndexBuffer()->GetSRV();
 		FRHIShaderResourceView* MeshVertexBufferSRV = LodRenderData->StaticVertexBuffers.PositionVertexBuffer.GetSRV();
 		FRHIShaderResourceView* MeshTangentBufferSRV = LodRenderData->StaticVertexBuffers.StaticMeshVertexBuffer.GetTangentsSRV();
@@ -191,11 +187,9 @@ void FOptimusSkinnedMeshDataProviderProxy::GatherDispatchData(FDispatchData cons
 		FRHIShaderResourceView* MeshColorBufferSRV = LodRenderData->StaticVertexBuffers.ColorVertexBuffer.GetColorComponentsSRV();
 
 		FParameters& Parameters = ParameterArray[InvocationIndex];
-		Parameters.NumVertices = InDispatchData.bUnifiedDispatch ? LodRenderData->GetNumVertices() : RenderSection.NumVertices;
-		Parameters.NumTriangles = InDispatchData.bUnifiedDispatch ? LodRenderData->GetTotalFaces() : RenderSection.NumTriangles;
+		Parameters.NumVertices = LodRenderData->GetNumVertices();
+		Parameters.NumTriangles =  LodRenderData->GetTotalFaces();
 		Parameters.NumUVChannels = LodRenderData->StaticVertexBuffers.StaticMeshVertexBuffer.GetNumTexCoords();
-		Parameters.IndexBufferStart = InDispatchData.bUnifiedDispatch ? 0 : RenderSection.BaseIndex;
-		Parameters.InputStreamStart = InDispatchData.bUnifiedDispatch ? 0 : RenderSection.BaseVertexIndex;
 		Parameters.IndexBuffer = IndexBufferSRV != nullptr ? IndexBufferSRV : NullSRVBinding;
 		Parameters.PositionInputBuffer = MeshVertexBufferSRV != nullptr ? MeshVertexBufferSRV : NullSRVBinding;
 		Parameters.TangentInputBuffer = MeshTangentBufferSRV != nullptr ? MeshTangentBufferSRV : NullSRVBinding;
diff --git a/Engine/Plugins/Animation/DeformerGraph/Source/OptimusCore/Private/DataInterfaces/OptimusDataInterfaceSkinnedMeshExec.cpp b/Engine/Plugins/Animation/DeformerGraph/Source/OptimusCore/Private/DataInterfaces/OptimusDataInterfaceSkinnedMeshExec.cpp
--- a/Engine/Plugins/Animation/DeformerGraph/Source/OptimusCore/Private/DataInterfaces/OptimusDataInterfaceSkinnedMeshExec.cpp
+++ b/Engine/Plugins/Animation/DeformerGraph/Source/OptimusCore/Private/DataInterfaces/OptimusDataInterfaceSkinnedMeshExec.cpp
@@ -13,17 +13,17 @@
 
 #include UE_INLINE_GENERATED_CPP_BY_NAME(OptimusDataInterfaceSkinnedMeshExec)
 
-FString UOptimusSkinnedMeshExecDataInterface::GetDisplayName() const
+FString UDEPRECATED_OptimusSkinnedMeshExecDataInterface::GetDisplayName() const
 {
 	return TEXT("Execute Skinned Mesh");
 }
 
-FName UOptimusSkinnedMeshExecDataInterface::GetCategory() const
+FName UDEPRECATED_OptimusSkinnedMeshExecDataInterface::GetCategory() const
 {
 	return CategoryName::ExecutionDataInterfaces;
 }
 
-TArray<FOptimusCDIPinDefinition> UOptimusSkinnedMeshExecDataInterface::GetPinDefinitions() const
+TArray<FOptimusCDIPinDefinition> UDEPRECATED_OptimusSkinnedMeshExecDataInterface::GetPinDefinitions() const
 {
 	TArray<FOptimusCDIPinDefinition> Defs;
 	Defs.Add({ "NumThreads", "ReadNumThreads" });
@@ -31,13 +31,13 @@ TArray<FOptimusCDIPinDefinition> UOptimusSkinnedMeshExecDataInterface::GetPinDef
 }
 
 
-TSubclassOf<UActorComponent> UOptimusSkinnedMeshExecDataInterface::GetRequiredComponentClass() const
+TSubclassOf<UActorComponent> UDEPRECATED_OptimusSkinnedMeshExecDataInterface::GetRequiredComponentClass() const
 {
 	return USkinnedMeshComponent::StaticClass();
 }
 
 
-void UOptimusSkinnedMeshExecDataInterface::GetSupportedInputs(TArray<FShaderFunctionDefinition>& OutFunctions) const
+void UDEPRECATED_OptimusSkinnedMeshExecDataInterface::GetSupportedInputs(TArray<FShaderFunctionDefinition>& OutFunctions) const
 {
 	OutFunctions.AddDefaulted_GetRef()
 		.SetName(TEXT("ReadNumThreads"))
@@ -49,24 +49,24 @@ BEGIN_SHADER_PARAMETER_STRUCT(FSkinedMeshExecDataInterfaceParameters, )
 	SHADER_PARAMETER(FIntVector, NumThreads)
 END_SHADER_PARAMETER_STRUCT()
 
-void UOptimusSkinnedMeshExecDataInterface::GetShaderParameters(TCHAR const* UID, FShaderParametersMetadataBuilder& InOutBuilder, FShaderParametersMetadataAllocations& InOutAllocations) const
+void UDEPRECATED_OptimusSkinnedMeshExecDataInterface::GetShaderParameters(TCHAR const* UID, FShaderParametersMetadataBuilder& InOutBuilder, FShaderParametersMetadataAllocations& InOutAllocations) const
 {
 	InOutBuilder.AddNestedStruct<FSkinedMeshExecDataInterfaceParameters>(UID);
 }
 
-TCHAR const* UOptimusSkinnedMeshExecDataInterface::TemplateFilePath = TEXT("/Plugin/Optimus/Private/DataInterfaceSkinnedMeshExec.ush");
+TCHAR const* UDEPRECATED_OptimusSkinnedMeshExecDataInterface::TemplateFilePath = TEXT("/Plugin/Optimus/Private/DataInterfaceSkinnedMeshExec.ush");
 
-TCHAR const* UOptimusSkinnedMeshExecDataInterface::GetShaderVirtualPath() const
+TCHAR const* UDEPRECATED_OptimusSkinnedMeshExecDataInterface::GetShaderVirtualPath() const
 {
 	return TemplateFilePath;
 }
 
-void UOptimusSkinnedMeshExecDataInterface::GetShaderHash(FString& InOutKey) const
+void UDEPRECATED_OptimusSkinnedMeshExecDataInterface::GetShaderHash(FString& InOutKey) const
 {
 	GetShaderFileHash(TemplateFilePath, EShaderPlatform::SP_PCD3D_SM5).AppendString(InOutKey);
 }
 
-void UOptimusSkinnedMeshExecDataInterface::GetHLSL(FString& OutHLSL, FString const& InDataInterfaceName) const
+void UDEPRECATED_OptimusSkinnedMeshExecDataInterface::GetHLSL(FString& OutHLSL, FString const& InDataInterfaceName) const
 {
 	TMap<FString, FStringFormatArg> TemplateArgs =
 	{
@@ -78,15 +78,12 @@ void UOptimusSkinnedMeshExecDataInterface::GetHLSL(FString& OutHLSL, FString con
 	OutHLSL += FString::Format(*TemplateFile, TemplateArgs);
 }
 
-UComputeDataProvider* UOptimusSkinnedMeshExecDataInterface::CreateDataProvider(TObjectPtr<UObject> InBinding, uint64 InInputMask, uint64 InOutputMask) const
+UComputeDataProvider* UDEPRECATED_OptimusSkinnedMeshExecDataInterface::CreateDataProvider(TObjectPtr<UObject> InBinding, uint64 InInputMask, uint64 InOutputMask) const
 {
-	UOptimusSkinnedMeshExecDataProvider* Provider = NewObject<UOptimusSkinnedMeshExecDataProvider>();
-	Provider->SkinnedMesh = Cast<USkinnedMeshComponent>(InBinding);
-	Provider->Domain = Domain;
-	return Provider;
+	return nullptr;
 }
 
-FName UOptimusSkinnedMeshExecDataInterface::GetSelectedExecutionDomainName() const
+FName UDEPRECATED_OptimusSkinnedMeshExecDataInterface::GetSelectedExecutionDomainName() const
 {
 	if (Domain == EOptimusSkinnedMeshExecDomain::Vertex)
 	{
@@ -101,83 +98,3 @@ FName UOptimusSkinnedMeshExecDataInterface::GetSelectedExecutionDomainName() con
 	return NAME_None;
 }
 
-
-FComputeDataProviderRenderProxy* UOptimusSkinnedMeshExecDataProvider::GetRenderProxy()
-{
-	return new FOptimusSkinnedMeshExecDataProviderProxy(SkinnedMesh, Domain);
-}
-
-
-FOptimusSkinnedMeshExecDataProviderProxy::FOptimusSkinnedMeshExecDataProviderProxy(USkinnedMeshComponent* InSkinnedMeshComponent, EOptimusSkinnedMeshExecDomain InDomain)
-{
-	SkeletalMeshObject = InSkinnedMeshComponent != nullptr ? InSkinnedMeshComponent->MeshObject : nullptr;
-	Domain = InDomain;
-}
-
-int32 FOptimusSkinnedMeshExecDataProviderProxy::GetDispatchThreadCount(TArray<FIntVector>& ThreadCounts) const
-{
-	ThreadCounts.Reset();
-	if (SkeletalMeshObject == nullptr)
-	{
-		return 0;
-	}
-
-	const int32 LodIndex = SkeletalMeshObject->GetLOD();
-	FSkeletalMeshRenderData const& SkeletalMeshRenderData = SkeletalMeshObject->GetSkeletalMeshRenderData();
-	FSkeletalMeshLODRenderData const* LodRenderData = &SkeletalMeshRenderData.LODRenderData[LodIndex];
-	const int32 NumInvocations = LodRenderData->RenderSections.Num();
-
-	ThreadCounts.Reserve(NumInvocations);
-	for (int32 InvocationIndex = 0; InvocationIndex < NumInvocations; ++InvocationIndex)
-	{
-		FSkelMeshRenderSection const& RenderSection = LodRenderData->RenderSections[InvocationIndex];
-		const int32 NumThreads = Domain == EOptimusSkinnedMeshExecDomain::Vertex ? RenderSection.NumVertices : RenderSection.NumTriangles;
-		ThreadCounts.Add(FIntVector(NumThreads, 1, 1));
-	}
-	
-	return NumInvocations;
-}
-
-bool FOptimusSkinnedMeshExecDataProviderProxy::IsValid(FValidationData const& InValidationData) const
-{
-	if (InValidationData.ParameterStructSize != sizeof(FParameters))
-	{
-		return false;
-	}
-	if (SkeletalMeshObject == nullptr)
-	{
-		return false;
-	}
-	if (SkeletalMeshObject->GetSkeletalMeshRenderData().LODRenderData[SkeletalMeshObject->GetLOD()].RenderSections.Num() != InValidationData.NumInvocations)
-	{
-		return false;
-	}
-
-	return true;
-}
-
-void FOptimusSkinnedMeshExecDataProviderProxy::GatherDispatchData(FDispatchData const& InDispatchData)
-{
-	const int32 LodIndex = SkeletalMeshObject->GetLOD();
-	FSkeletalMeshRenderData const& SkeletalMeshRenderData = SkeletalMeshObject->GetSkeletalMeshRenderData();
-	FSkeletalMeshLODRenderData const* LodRenderData = &SkeletalMeshRenderData.LODRenderData[LodIndex];
-
-	const TStridedView<FParameters> ParameterArray = MakeStridedParameterView<FParameters>(InDispatchData);
-	for (int32 InvocationIndex = 0; InvocationIndex < ParameterArray.Num(); ++InvocationIndex)
-	{
-		int32 NumThreads = 0;
-
-		if (InDispatchData.bUnifiedDispatch)
-		{
-			NumThreads = Domain == EOptimusSkinnedMeshExecDomain::Vertex ? LodRenderData->GetNumVertices() : LodRenderData->GetTotalFaces();
-		}
-		else
-		{
-			FSkelMeshRenderSection const& RenderSection = LodRenderData->RenderSections[InvocationIndex];
-			NumThreads = Domain == EOptimusSkinnedMeshExecDomain::Vertex ? RenderSection.NumVertices : RenderSection.NumTriangles;
-		}
-
-		FParameters& Parameters = ParameterArray[InvocationIndex];
-		Parameters.NumThreads = FIntVector(NumThreads, 1, 1);
-	}
-}
diff --git a/Engine/Plugins/Animation/DeformerGraph/Source/OptimusCore/Private/DataInterfaces/OptimusDataInterfaceSkinnedMeshExec.h b/Engine/Plugins/Animation/DeformerGraph/Source/OptimusCore/Private/DataInterfaces/OptimusDataInterfaceSkinnedMeshExec.h
--- a/Engine/Plugins/Animation/DeformerGraph/Source/OptimusCore/Private/DataInterfaces/OptimusDataInterfaceSkinnedMeshExec.h
+++ b/Engine/Plugins/Animation/DeformerGraph/Source/OptimusCore/Private/DataInterfaces/OptimusDataInterfaceSkinnedMeshExec.h
@@ -25,8 +25,8 @@ enum class EOptimusSkinnedMeshExecDomain : uint8
 };
 
 /** Compute Framework Data Interface for executing kernels over a skinned mesh domain. */
-UCLASS(Category = ComputeFramework)
-class OPTIMUSCORE_API UOptimusSkinnedMeshExecDataInterface :
+UCLASS(Deprecated, Category = ComputeFramework, meta = (DeprecationMessage ="This execution interface has been replaced with kernel specific execution data interface, see UOptimusCustomComputeKernelDataInterface"))
+class OPTIMUSCORE_API UDEPRECATED_OptimusSkinnedMeshExecDataInterface :
 	public UOptimusComputeDataInterface,
 	public IOptimusDeprecatedExecutionDataInterface
 {
@@ -64,38 +64,4 @@ private:
 	static TCHAR const* TemplateFilePath;
 };
 
-/** Compute Framework Data Provider for executing kernels over a skinned mesh domain. */
-UCLASS(BlueprintType, editinlinenew, Category = ComputeFramework)
-class UOptimusSkinnedMeshExecDataProvider : public UComputeDataProvider
-{
-	GENERATED_BODY()
-
-public:
-	UPROPERTY()
-	TObjectPtr<USkinnedMeshComponent> SkinnedMesh = nullptr;
-
-	UPROPERTY()
-	EOptimusSkinnedMeshExecDomain Domain = EOptimusSkinnedMeshExecDomain::Vertex;
-
-	//~ Begin UComputeDataProvider Interface
-	FComputeDataProviderRenderProxy* GetRenderProxy() override;
-	//~ End UComputeDataProvider Interface
-};
 
-class FOptimusSkinnedMeshExecDataProviderProxy : public FComputeDataProviderRenderProxy
-{
-public:
-	FOptimusSkinnedMeshExecDataProviderProxy(USkinnedMeshComponent* InSkinnedMeshComponent, EOptimusSkinnedMeshExecDomain InDomain);
-
-	//~ Begin FComputeDataProviderRenderProxy Interface
-	int32 GetDispatchThreadCount(TArray<FIntVector>& ThreadCounts) const override;
-	bool IsValid(FValidationData const& InValidationData) const override;
-	void GatherDispatchData(FDispatchData const& InDispatchData) override;
-	//~ End FComputeDataProviderRenderProxy Interface
-
-private:
-	using FParameters = FSkinedMeshExecDataInterfaceParameters;
-
-	FSkeletalMeshObject* SkeletalMeshObject = nullptr;
-	EOptimusSkinnedMeshExecDomain Domain = EOptimusSkinnedMeshExecDomain::Vertex;
-};
diff --git a/Engine/Plugins/Animation/DeformerGraph/Source/OptimusCore/Private/DataInterfaces/OptimusDataInterfaceSkinnedMeshRead.cpp b/Engine/Plugins/Animation/DeformerGraph/Source/OptimusCore/Private/DataInterfaces/OptimusDataInterfaceSkinnedMeshRead.cpp
--- a/Engine/Plugins/Animation/DeformerGraph/Source/OptimusCore/Private/DataInterfaces/OptimusDataInterfaceSkinnedMeshRead.cpp
+++ b/Engine/Plugins/Animation/DeformerGraph/Source/OptimusCore/Private/DataInterfaces/OptimusDataInterfaceSkinnedMeshRead.cpp
@@ -88,7 +88,6 @@ void UOptimusSkinnedMeshReadDataInterface::GetSupportedInputs(TArray<FShaderFunc
 
 BEGIN_SHADER_PARAMETER_STRUCT(FSkinnedMeshReadDataInterfaceParameters, )
 	SHADER_PARAMETER(uint32, NumVertices)
-	SHADER_PARAMETER(uint32, InputStreamStart)
 	SHADER_PARAMETER_RDG_BUFFER_UAV(RWBuffer<float>, PositionBufferUAV)
 	SHADER_PARAMETER_RDG_BUFFER_UAV(RWBuffer<SNORM float4>, TangentBufferUAV)
 	SHADER_PARAMETER_RDG_BUFFER_UAV(RWBuffer<UNORM float4>, ColorBufferUAV)
@@ -297,8 +296,7 @@ void FOptimusSkinnedMeshReadDataProviderProxy::GatherDispatchData(FDispatchData
 		FSkelMeshRenderSection const& RenderSection = LodRenderData->RenderSections[InvocationIndex];
 
 		FParameters& Parameters = ParameterArray[InvocationIndex];
-		Parameters.NumVertices = InDispatchData.bUnifiedDispatch ? LodRenderData->GetNumVertices() : RenderSection.GetNumVertices();
-		Parameters.InputStreamStart= InDispatchData.bUnifiedDispatch ? 0 : RenderSection.GetVertexBufferIndex();
+		Parameters.NumVertices = LodRenderData->GetNumVertices();
 
 		Parameters.PositionBufferUAV = PositionBufferUAV;
 		Parameters.TangentBufferUAV = TangentBufferUAV;
diff --git a/Engine/Plugins/Animation/DeformerGraph/Source/OptimusCore/Private/DataInterfaces/OptimusDataInterfaceSkinnedMeshVertexAttribute.cpp b/Engine/Plugins/Animation/DeformerGraph/Source/OptimusCore/Private/DataInterfaces/OptimusDataInterfaceSkinnedMeshVertexAttribute.cpp
--- a/Engine/Plugins/Animation/DeformerGraph/Source/OptimusCore/Private/DataInterfaces/OptimusDataInterfaceSkinnedMeshVertexAttribute.cpp
+++ b/Engine/Plugins/Animation/DeformerGraph/Source/OptimusCore/Private/DataInterfaces/OptimusDataInterfaceSkinnedMeshVertexAttribute.cpp
@@ -62,7 +62,6 @@ void UOptimusSkinnedMeshVertexAttributeDataInterface::GetSupportedInputs(TArray<
 
 BEGIN_SHADER_PARAMETER_STRUCT(FSkinnedMeshVertexAttributeDataInterfaceParameters, )
 	SHADER_PARAMETER(uint32, NumVertices)
-	SHADER_PARAMETER(uint32, InputStreamStart)
 	SHADER_PARAMETER(uint32, bIsValid)
 	SHADER_PARAMETER_SRV(Buffer<float>, ValueBuffer)
 END_SHADER_PARAMETER_STRUCT()
@@ -199,11 +198,8 @@ void FOptimusSkinnedMeshVertexAttributeDataProviderProxy::GatherDispatchData(
 
 	for (int32 InvocationIndex = 0; InvocationIndex < ParameterArray.Num(); ++InvocationIndex)
 	{
-		FSkelMeshRenderSection const& RenderSection = LodRenderData->RenderSections[InvocationIndex];
-
 		FSkinnedMeshVertexAttributeDataInterfaceParameters& Parameters = ParameterArray[InvocationIndex];
-		Parameters.NumVertices = InDispatchData.bUnifiedDispatch ? LodRenderData->GetNumVertices() : RenderSection.NumVertices;
-		Parameters.InputStreamStart = InDispatchData.bUnifiedDispatch ? 0 : RenderSection.BaseVertexIndex;
+		Parameters.NumVertices = LodRenderData->GetNumVertices();
 		Parameters.ValueBuffer = AttributeSRV;
 		Parameters.bIsValid = AttributeBuffer != nullptr;
 	}
diff --git a/Engine/Plugins/Animation/DeformerGraph/Source/OptimusCore/Private/DataInterfaces/OptimusDataInterfaceSkinnedMeshWrite.cpp b/Engine/Plugins/Animation/DeformerGraph/Source/OptimusCore/Private/DataInterfaces/OptimusDataInterfaceSkinnedMeshWrite.cpp
--- a/Engine/Plugins/Animation/DeformerGraph/Source/OptimusCore/Private/DataInterfaces/OptimusDataInterfaceSkinnedMeshWrite.cpp
+++ b/Engine/Plugins/Animation/DeformerGraph/Source/OptimusCore/Private/DataInterfaces/OptimusDataInterfaceSkinnedMeshWrite.cpp
@@ -99,7 +99,6 @@ void UOptimusSkinnedMeshWriteDataInterface::GetSupportedOutputs(TArray<FShaderFu
 
 BEGIN_SHADER_PARAMETER_STRUCT(FSkinedMeshWriteDataInterfaceParameters, )
 	SHADER_PARAMETER(uint32, NumVertices)
-	SHADER_PARAMETER(uint32, OutputStreamStart)
 	SHADER_PARAMETER_RDG_BUFFER_UAV(RWBuffer<float>, PositionBufferUAV)
 	SHADER_PARAMETER_RDG_BUFFER_UAV(RWBuffer<SNORM float4>, TangentBufferUAV)
 	SHADER_PARAMETER_RDG_BUFFER_UAV(RWBuffer<UNORM float4>, ColorBufferUAV)
@@ -307,8 +306,7 @@ void FOptimusSkinnedMeshWriteDataProviderProxy::GatherDispatchData(FDispatchData
 		FSkelMeshRenderSection const& RenderSection = LodRenderData->RenderSections[InvocationIndex];
 
 		FParameters& Parameters = ParameterArray[InvocationIndex];
-		Parameters.NumVertices = InDispatchData.bUnifiedDispatch ? LodRenderData->GetNumVertices() : RenderSection.GetNumVertices();
-		Parameters.OutputStreamStart = InDispatchData.bUnifiedDispatch ? 0 : RenderSection.GetVertexBufferIndex();
+		Parameters.NumVertices =LodRenderData->GetNumVertices();
 		Parameters.PositionBufferUAV = PositionBufferUAV;
 		Parameters.TangentBufferUAV = TangentBufferUAV;
 		Parameters.ColorBufferUAV = ColorBufferUAV;
diff --git a/Engine/Plugins/Animation/DeformerGraph/Source/OptimusCore/Private/Nodes/OptimusNode_ComputeKernelBase.cpp b/Engine/Plugins/Animation/DeformerGraph/Source/OptimusCore/Private/Nodes/OptimusNode_ComputeKernelBase.cpp
--- a/Engine/Plugins/Animation/DeformerGraph/Source/OptimusCore/Private/Nodes/OptimusNode_ComputeKernelBase.cpp
+++ b/Engine/Plugins/Animation/DeformerGraph/Source/OptimusCore/Private/Nodes/OptimusNode_ComputeKernelBase.cpp
@@ -143,8 +143,29 @@ FOptimus_ComputeKernelResult UOptimusNode_ComputeKernelBase::CreateComputeKernel
 	{
 		return ReturnError(LOCTEXT("ZeroOrMultiplePrimaryBindings", "Primary Group has zero or more than one Component Bindings"));
 	}
+	
+	UOptimusComponentSourceBinding* PrimaryBinding = *(PrimaryBindings.CreateIterator());
+	bool bUnifiedDispatchOnly = true;
+	{
+		using namespace Optimus::Expression;
+		TVariant<FExpressionObject, FParseError> ParseResult = Optimus::ParseExecutionDomainExpression(GetExecutionDomain().AsExpression(), PrimaryBinding->GetComponentSource());
+		if (ParseResult.IsType<FParseError>())
+		{
+			return ReturnError(
+				FText::Format(
+					LOCTEXT("ExecutionDomainParsingError", "Error while parsing Execution Domain: {0}"),
+					FText::FromString(ParseResult.Get<FParseError>().Message)));
+		}
 
+		FExpressionObject ExpressionObject = ParseResult.Get<FExpressionObject>();
 
+		// Only fixed domain supports both unified and non-unified dispatch
+		if (TOptional<FName> FixedDomain = ExpressionObject.GetAsSingleConstant())
+		{
+			bUnifiedDispatchOnly = false;
+		}
+	}
+	
 	UOptimusKernelSource* KernelSource = NewObject<UOptimusKernelSource>(InKernelSourceOuter);
 
 	// Wrap functions for unconnected resource pins (or value pins) that return default values
@@ -159,7 +180,9 @@ FOptimus_ComputeKernelResult UOptimusNode_ComputeKernelBase::CreateComputeKernel
 			for (const UOptimusNodePin* SubPin: Pin->GetSubPins())
 			{
 				TOptional<FText> Result = ProcessInputPinForComputeKernel(
-					InTraversalContext, SubPin, GroupIndex == 0 ? FString() : Pin->GetName(),
+					bUnifiedDispatchOnly,
+					InTraversalContext, SubPin,
+					GroupIndex == 0 ? FString() : Pin->GetName(),
 					InKernelInputs, KernelSource, GeneratedFunctions, OutInputDataBindings
 					);
 				if (Result.IsSet())
@@ -172,10 +195,17 @@ FOptimus_ComputeKernelResult UOptimusNode_ComputeKernelBase::CreateComputeKernel
 		}
 		else if (Pin->GetDirection() == EOptimusNodePinDirection::Output)
 		{
-			ProcessOutputPinForComputeKernel(
+			TOptional<FText> Result = ProcessOutputPinForComputeKernel(
+				bUnifiedDispatchOnly,
 				InTraversalContext, Pin,
 				InKernelOutputs,
-				KernelSource, GeneratedFunctions, OutInputDataBindings, OutOutputDataBindings);
+				KernelSource, GeneratedFunctions, OutInputDataBindings, OutOutputDataBindings
+				);
+				
+			if (Result.IsSet())
+			{
+				return ReturnError(Result.GetValue());
+			}
 		}
 	}
 
@@ -185,7 +215,8 @@ FOptimus_ComputeKernelResult UOptimusNode_ComputeKernelBase::CreateComputeKernel
 			PrimaryBindings.Array()[0],
 			InOutKernelDataInterface,
 			KernelSource,
-			OutInputDataBindings
+			OutInputDataBindings,
+			GeneratedFunctions
 			);
 	}
 
@@ -361,8 +392,15 @@ FString UOptimusNode_ComputeKernelBase::GetReadFunctionName(const FString& InBin
 	return FString::Printf(TEXT("Read%s"), *InBindingName);
 }
 
+const TCHAR* UOptimusNode_ComputeKernelBase::GetNumThreadsReservedName()
+{
+	static const TCHAR* NumThreadsReservedName = TEXT("NumThreads");
+	return NumThreadsReservedName;
+}
+
 
 TOptional<FText> UOptimusNode_ComputeKernelBase::ProcessInputPinForComputeKernel(
+	bool bUnifiedDispatchOnly,
 	const FOptimusPinTraversalContext& InTraversalContext,
 	const UOptimusNodePin* InInputPin,
 	const FString& InGroupName,
@@ -384,11 +422,19 @@ TOptional<FText> UOptimusNode_ComputeKernelBase::ProcessInputPinForComputeKernel
 		const UOptimusNode* OutputNode = OutputPin->GetOwningNode();
 		UComputeDataInterface* DataInterface = InputConnection->DataInterface;
 
-		if (!DataInterface->CanSupportUnifiedDispatch() && !InGroupName.IsEmpty())
+		if (!DataInterface->CanSupportUnifiedDispatch() )
 		{
-			return FText::Format(LOCTEXT("SecondaryGroupOnlySupportsDataInterfaceUnifiedDispatch", "Cannot connect Secondary group input binding pin '{0}' to a data interface that does not support unified dispatch, consider inserting a kernel inbetween"), FText::FromName(InInputPin->GetUniqueName()));
+			if (!InGroupName.IsEmpty())
+			{
+				return FText::Format(LOCTEXT("SecondaryGroupOnlySupportsDataInterfaceUnifiedDispatch", "Cannot connect Secondary group input binding pin '{0}' to a data interface that does not support unified dispatch, consider inserting a kernel inbetween"), FText::FromName(InInputPin->GetUniqueName()));
+			}
+
+			if (bUnifiedDispatchOnly)
+			{
+				return FText::Format(LOCTEXT("InvalidInput_KernelOnlySupportsDataInterfaceUnifiedDispatch", "Cannot connect input binding pin '{0}' to a data interface that does not support unified dispatch, consider using a simple kernel execution domain or inserting a kernel inbetween"), FText::FromName(InInputPin->GetUniqueName()));
+			}
 		}
-		
+
 		int32 DataInterfaceFuncIndex = INDEX_NONE;
 
 		if (Cast<const IOptimusComputeKernelProvider>(OutputNode))
@@ -502,7 +548,8 @@ TOptional<FText> UOptimusNode_ComputeKernelBase::ProcessInputPinForComputeKernel
 }
 
 
-void UOptimusNode_ComputeKernelBase::ProcessOutputPinForComputeKernel(
+TOptional<FText> UOptimusNode_ComputeKernelBase::ProcessOutputPinForComputeKernel(
+	bool bUnifiedDispatchOnly,
 	const FOptimusPinTraversalContext& InTraversalContext,
 	const UOptimusNodePin* InOutputPin,
 	const FOptimus_KernelOutputMap& InKernelOutputs,
@@ -544,6 +591,14 @@ void UOptimusNode_ComputeKernelBase::ProcessOutputPinForComputeKernel(
 			TArray<FShaderFunctionDefinition> WriteFunctions;
 			DataInterface->GetSupportedOutputs(WriteFunctions);
 
+			if (!DataInterface->CanSupportUnifiedDispatch())
+			{
+				if (bUnifiedDispatchOnly)
+				{
+					return FText::Format(LOCTEXT("InvalidOutput_KernelOnlySupportsDataInterfaceUnifiedDispatch", "Cannot connect output binding pin '{0}' to a data interface that does not support unified dispatch, consider using to a simple kernel execution domain or inserting a kernel inbetween"), FText::FromName(InOutputPin->GetUniqueName()));
+				}
+			}
+			
 			if (!ConnectedNode)
 			{
 				// No connected node indicates the kernel output is cached in a generated raw buffer
@@ -700,24 +755,26 @@ void UOptimusNode_ComputeKernelBase::ProcessOutputPinForComputeKernel(
 			FString::Printf(TEXT("void Write%s(%s, %s Value) { }"),
 				*InOutputPin->GetName(), *FString::Join(StubIndexes, TEXT(", ")), *ValueType->ToString()));
 	}
+
+	return {};
 }
 
 void UOptimusNode_ComputeKernelBase::BindKernelDataInterfaceForComputeKernel(
 	const UOptimusComponentSourceBinding* InKernelPrimaryComponentSourceBinding,
 	UComputeDataInterface* InOutKernelDataInterface,
 	UOptimusKernelSource* InKernelSource,
-	FOptimus_InterfaceBindingMap& OutInputDataBindings
+	FOptimus_InterfaceBindingMap& OutInputDataBindings,
+	TArray<FString>& OutGeneratedFunctions
 	) const
 {
 	check(InOutKernelDataInterface);
 	check(InKernelPrimaryComponentSourceBinding);
+	
 	IOptimusComputeKernelDataInterface* KernelDataInterface = Cast<IOptimusComputeKernelDataInterface>(InOutKernelDataInterface);
-	if (ensure(KernelDataInterface))
-	{
-		KernelDataInterface->SetExecutionDomain(GetExecutionDomain().AsExpression());
-		KernelDataInterface->SetComponentBinding(InKernelPrimaryComponentSourceBinding);
-	}
-
+	check(KernelDataInterface);
+	
+	KernelDataInterface->SetExecutionDomain(GetExecutionDomain().AsExpression());
+	KernelDataInterface->SetComponentBinding(InKernelPrimaryComponentSourceBinding);
 	
 	TArray<FShaderFunctionDefinition> ReadFunctions;
 	InOutKernelDataInterface->GetSupportedInputs(ReadFunctions);	
@@ -737,12 +794,16 @@ void UOptimusNode_ComputeKernelBase::BindKernelDataInterfaceForComputeKernel(
 		InterfaceBinding.DataInterface = InOutKernelDataInterface;
 		InterfaceBinding.DataInterfaceBindingIndex = FuncIndex;
 		InterfaceBinding.BindingFunctionName = FuncDef.Name;
-		InterfaceBinding.BindingFunctionNamespace = FString();
+		InterfaceBinding.BindingFunctionNamespace = Optimus::GetKernelInternalNamespaceName() ;
 				
 		OutInputDataBindings.Add(InKernelSource->ExternalInputs.Num(), InterfaceBinding);  
 				
 		InKernelSource->ExternalInputs.Emplace(FuncDef);	
 	}
+
+	OutGeneratedFunctions.Add(FString::Printf(
+			TEXT("int3 Read%s() { return %s::%s(); }"),
+			GetNumThreadsReservedName(), Optimus::GetKernelInternalNamespaceName(), KernelDataInterface->GetReadNumThreadsFunctionName() ));
 }
 
 #undef LOCTEXT_NAMESPACE
diff --git a/Engine/Plugins/Animation/DeformerGraph/Source/OptimusCore/Private/Nodes/OptimusNode_ComputeKernelBase.h b/Engine/Plugins/Animation/DeformerGraph/Source/OptimusCore/Private/Nodes/OptimusNode_ComputeKernelBase.h
--- a/Engine/Plugins/Animation/DeformerGraph/Source/OptimusCore/Private/Nodes/OptimusNode_ComputeKernelBase.h
+++ b/Engine/Plugins/Animation/DeformerGraph/Source/OptimusCore/Private/Nodes/OptimusNode_ComputeKernelBase.h
@@ -94,9 +94,11 @@ protected:
 	
 	static FString GetAtomicWriteFunctionName(EOptimusBufferWriteType InWriteType, const FString& InBindingName);
 	static FString GetReadFunctionName(const FString& InBindingName);
+	static const TCHAR* GetNumThreadsReservedName();
 	
 private:
 	TOptional<FText> ProcessInputPinForComputeKernel(
+		bool bUnifiedDispatchOnly,
 		const FOptimusPinTraversalContext& InTraversalContext,
 		const UOptimusNodePin* InInputPin,
 		const FString& InGroupName,
@@ -106,7 +108,8 @@ private:
 		FOptimus_InterfaceBindingMap& OutInputDataBindings
 		) const;
 
-	void ProcessOutputPinForComputeKernel(
+	TOptional<FText> ProcessOutputPinForComputeKernel(
+		bool bUnifiedDispatchOnly,
 		const FOptimusPinTraversalContext& InTraversalContext,
 		const UOptimusNodePin* InOutputPin,
 		const FOptimus_KernelOutputMap& InKernelOutputs,
@@ -120,7 +123,8 @@ private:
 		const UOptimusComponentSourceBinding* InKernelPrimaryComponentSourceBinding,
 		UComputeDataInterface* InOutKernelDataInterface,
 		UOptimusKernelSource* InKernelSource,
-		FOptimus_InterfaceBindingMap& OutInputDataBindings
+		FOptimus_InterfaceBindingMap& OutInputDataBindings,
+		TArray<FString>& OutGeneratedFunctions
 		) const;
 };
 
diff --git a/Engine/Plugins/Animation/DeformerGraph/Source/OptimusCore/Private/Nodes/OptimusNode_ComputeKernelFunction.cpp b/Engine/Plugins/Animation/DeformerGraph/Source/OptimusCore/Private/Nodes/OptimusNode_ComputeKernelFunction.cpp
--- a/Engine/Plugins/Animation/DeformerGraph/Source/OptimusCore/Private/Nodes/OptimusNode_ComputeKernelFunction.cpp
+++ b/Engine/Plugins/Animation/DeformerGraph/Source/OptimusCore/Private/Nodes/OptimusNode_ComputeKernelFunction.cpp
@@ -181,7 +181,8 @@ FIntVector UOptimusNode_ComputeKernelFunction::GetGroupSize() const
 
 FString UOptimusNode_ComputeKernelFunction::GetKernelSourceText() const
 {
-	return Optimus::GetCookedKernelSource(GetPathName(), GetGeneratorClass()->ShaderSource, GetKernelHlslName(), GetGroupSize());
+	checkNoEntry();
+	return {};
 }
 
 
diff --git a/Engine/Plugins/Animation/DeformerGraph/Source/OptimusCore/Private/Nodes/OptimusNode_CustomComputeKernel.cpp b/Engine/Plugins/Animation/DeformerGraph/Source/OptimusCore/Private/Nodes/OptimusNode_CustomComputeKernel.cpp
--- a/Engine/Plugins/Animation/DeformerGraph/Source/OptimusCore/Private/Nodes/OptimusNode_CustomComputeKernel.cpp
+++ b/Engine/Plugins/Animation/DeformerGraph/Source/OptimusCore/Private/Nodes/OptimusNode_CustomComputeKernel.cpp
@@ -102,7 +102,16 @@ UOptimusNode_CustomComputeKernel::UOptimusNode_CustomComputeKernel()
 
 FString UOptimusNode_CustomComputeKernel::GetKernelSourceText() const
 {
-	return Optimus::GetCookedKernelSource(GetPathName(), ShaderSource.ShaderText, GetKernelHlslName(), GroupSize);
+	using UKernelDI = UOptimusCustomComputeKernelDataInterface;
+	
+	return Optimus::GetCookedKernelSource(
+		GetPathName(),
+		ShaderSource.ShaderText,
+		GetKernelHlslName(),
+		GroupSize,
+		UKernelDI::StaticClass()->GetDefaultObject<UKernelDI>()->GetReadNumThreadsPerInvocationFunctionName(),
+		UKernelDI::StaticClass()->GetDefaultObject<UKernelDI>()->GetReadThreadIndexOffsetFunctionName()
+		);
 }
 
 FOptimusExecutionDomain UOptimusNode_CustomComputeKernel::GetExecutionDomain() const
@@ -633,7 +642,7 @@ FName UOptimusNode_CustomComputeKernel::GetSanitizedNewPinName(
 	FName NewName = Optimus::GetSanitizedNameForHlsl(InPinName);
 	
 	// "NumThreads" is reserved for the kernel data interface
-	if (NewName == *UOptimusCustomComputeKernelDataInterface::NumThreadsReservedName)
+	if (NewName == GetNumThreadsReservedName())
 	{
 		NewName.SetNumber(1);
 	}
@@ -1904,7 +1913,7 @@ void UOptimusNode_CustomComputeKernel::PostLoadExtractAtomicModeFromConnectedRes
 bool UOptimusNode_CustomComputeKernel::PostLoadRemoveDeprecatedNumThreadsPin()
 {
 	// NumThreads is now a reserved name that pins cannot use
-	const FName NumThreadsPinName = *UOptimusCustomComputeKernelDataInterface::NumThreadsReservedName;
+	const FName NumThreadsPinName = GetNumThreadsReservedName();
 	
 	UOptimusNodePin* PrimaryGroupPin = GetPrimaryGroupPin_Internal();
 	
diff --git a/Engine/Plugins/Animation/DeformerGraph/Source/OptimusCore/Private/OptimusDeformer.cpp b/Engine/Plugins/Animation/DeformerGraph/Source/OptimusCore/Private/OptimusDeformer.cpp
--- a/Engine/Plugins/Animation/DeformerGraph/Source/OptimusCore/Private/OptimusDeformer.cpp
+++ b/Engine/Plugins/Animation/DeformerGraph/Source/OptimusCore/Private/OptimusDeformer.cpp
@@ -2847,8 +2847,6 @@ FOptimusNodeGraphCompilationResult UOptimusDeformer::CompileNodeGraphToComputeGr
 
 				UOptimusKernelSource* KernelSource = NewObject<UOptimusKernelSource>(BoundCopyKernel.Kernel);
 				FString SourceText;
-				SourceText = TEXT("if (Index >= ReadNumThreads().x) return;\n");
-
 				{
 					UComputeDataInterface* CopyKernelDataInterface = CopyKernelDataInterfaceMap[SourceInstancedPin];
 					TArray<FShaderFunctionDefinition> Functions;
@@ -2867,7 +2865,7 @@ FOptimusNodeGraphCompilationResult UOptimusDeformer::CompileNodeGraphToComputeGr
 						InterfaceBinding.DataInterface = CopyKernelDataInterface;
 						InterfaceBinding.DataInterfaceBindingIndex = FuncIndex;
 						InterfaceBinding.BindingFunctionName = FuncDef.Name;
-						InterfaceBinding.BindingFunctionNamespace = FString();
+						InterfaceBinding.BindingFunctionNamespace = Optimus::GetKernelInternalNamespaceName();
 				
 						InputDataBindings.Add(KernelSource->ExternalInputs.Num(), InterfaceBinding);  
 				
@@ -2935,7 +2933,15 @@ FOptimusNodeGraphCompilationResult UOptimusDeformer::CompileNodeGraphToComputeGr
 
 				static const FString CopyKernelName = TEXT("CopyKernel");
 				static const FIntVector GroupSize = FIntVector(64, 1, 1);
-				FString CookedSource = Optimus::GetCookedKernelSource(BoundCopyKernel.Kernel->GetPathName(), SourceText, CopyKernelName, GroupSize);
+				using UKernelDI = UOptimusCopyKernelDataInterface;
+				FString CookedSource = Optimus::GetCookedKernelSource(
+					BoundCopyKernel.Kernel->GetPathName(),
+					SourceText,
+					CopyKernelName,
+					GroupSize,
+					UKernelDI::StaticClass()->GetDefaultObject<UKernelDI>()->GetReadNumThreadsPerInvocationFunctionName(),
+					UKernelDI::StaticClass()->GetDefaultObject<UKernelDI>()->GetReadThreadIndexOffsetFunctionName()	
+					);
 				KernelSource->SetSource(CookedSource);
 				KernelSource->EntryPoint = CopyKernelName;
 				KernelSource->GroupSize = GroupSize;
@@ -3340,8 +3346,9 @@ void UOptimusDeformer::PostLoad()
 			Optimus::RenameObject(ClassObject, nullptr, GetPackage());
 		}
 	}
-	
-	if (GetLinkerCustomVersion(FOptimusObjectVersion::GUID) < FOptimusObjectVersion::ReparentResourcesAndVariables)
+
+	FOptimusObjectVersion::Type ObjectVersion = static_cast<FOptimusObjectVersion::Type>(GetLinkerCustomVersion(FOptimusObjectVersion::GUID));
+	if (ObjectVersion < FOptimusObjectVersion::ReparentResourcesAndVariables)
 	{
 		// Move any resource or variable descriptor owned by this deformer to their own container.
 		// This is to fix a bug where variables/resources were put in their respective container
@@ -3362,7 +3369,7 @@ void UOptimusDeformer::PostLoad()
 			}
 		}
 	}
-	if (GetLinkerCustomVersion(FOptimusObjectVersion::GUID) < FOptimusObjectVersion::ComponentProviderSupport)
+	if (ObjectVersion < FOptimusObjectVersion::ComponentProviderSupport)
 	{
 		if (ensure(Bindings->Bindings.IsEmpty()))
 		{
@@ -3378,7 +3385,7 @@ void UOptimusDeformer::PostLoad()
 		// Fix up any data providers to ensure they have a binding.
 		PostLoadFixupMissingComponentBindingsCompat();
 	}
-	if (GetLinkerCustomVersion(FOptimusObjectVersion::GUID) < FOptimusObjectVersion::SetPrimaryBindingName)
+	if (ObjectVersion < FOptimusObjectVersion::SetPrimaryBindingName)
 	{
 		FName PrimaryBindingName = UOptimusComponentSourceBinding::GetPrimaryBindingName();
 		for (UOptimusComponentSourceBinding* Binding : Bindings->Bindings)
@@ -3403,31 +3410,31 @@ void UOptimusDeformer::PostLoad()
 	}
 	// Fix any resource data domains if the component binding is valid but the domain is not. This will mostly cut links
 	// to kernels with mismatched domain info.
-	if (GetLinkerCustomVersion(FOptimusObjectVersion::GUID) < FOptimusObjectVersion::DataDomainExpansion)
+	if (ObjectVersion < FOptimusObjectVersion::DataDomainExpansion)
 	{
 		PostLoadFixupMismatchedResourceDataDomains();
 	}
 
-	if (GetLinkerCustomVersion(FOptimusObjectVersion::GUID) < FOptimusObjectVersion::KernelDataInterface)
+	if (ObjectVersion < FOptimusObjectVersion::KernelDataInterface)
 	{
 		PostLoadRemoveDeprecatedExecutionNodes();
 	}
 
-	if (GetLinkerCustomVersion(FOptimusObjectVersion::GUID) < FOptimusObjectVersion::PropertyBagValueContainer)
+	if (ObjectVersion < FOptimusObjectVersion::PropertyBagValueContainer)
 	{
 		PostLoadRemoveDeprecatedValueContainerGeneratorClass();
 	}
 
-	if (GetLinkerCustomVersion(FOptimusObjectVersion::GUID) < FOptimusObjectVersion::PropertyPinSupport)
+	if (ObjectVersion < FOptimusObjectVersion::PropertyPinSupport)
 	{
 		PostLoadMoveValueFromGraphDataInterfaceToDeformerValueMap();
 	}
 	
-	// If the graph was saved at any previous version, and was clean, mark the status now as modified.
-	if (Status == EOptimusDeformerStatus::Compiled &&
-		GetLinkerCustomVersion(FOptimusObjectVersion::GUID) < FOptimusObjectVersion::LatestVersion)
+	// If the graph was saved at any previous version, and was clean, recompile it to the latest version.
+	if (Status != EOptimusDeformerStatus::HasErrors &&
+		ObjectVersion < FOptimusObjectVersion::LatestVersion)
 	{
-		Status = EOptimusDeformerStatus::Modified;
+		Compile();
 	}
 }
 
diff --git a/Engine/Plugins/Animation/DeformerGraph/Source/OptimusCore/Private/OptimusExpressionEvaluator.cpp b/Engine/Plugins/Animation/DeformerGraph/Source/OptimusCore/Private/OptimusExpressionEvaluator.cpp
--- a/Engine/Plugins/Animation/DeformerGraph/Source/OptimusCore/Private/OptimusExpressionEvaluator.cpp
+++ b/Engine/Plugins/Animation/DeformerGraph/Source/OptimusCore/Private/OptimusExpressionEvaluator.cpp
@@ -140,6 +140,19 @@ TSet<FName> FExpressionObject::GetUsedConstants() const
 	return UsedConstants;	
 }
 
+TOptional<FName> FExpressionObject::GetAsSingleConstant() const
+{
+	if (Expression.Num() == 1)
+	{
+		if (const FName* ConstantName = Expression[0].TryGet<FName>())
+		{
+			return TOptional<FName>(*ConstantName);
+		}
+	}
+
+	return {};
+}
+
 
 void FExpressionObject::Load(FArchive& Ar)
 {
diff --git a/Engine/Plugins/Animation/DeformerGraph/Source/OptimusCore/Private/OptimusHelpers.cpp b/Engine/Plugins/Animation/DeformerGraph/Source/OptimusCore/Private/OptimusHelpers.cpp
--- a/Engine/Plugins/Animation/DeformerGraph/Source/OptimusCore/Private/OptimusHelpers.cpp
+++ b/Engine/Plugins/Animation/DeformerGraph/Source/OptimusCore/Private/OptimusHelpers.cpp
@@ -3,7 +3,9 @@
 #include "OptimusHelpers.h"
 
 #include "Matrix3x4.h"
+#include "OptimusComponentSource.h"
 #include "OptimusDataTypeRegistry.h"
+#include "OptimusExpressionEvaluator.h"
 #include "ComputeFramework/ShaderParamTypeDefinition.h"
 #include "ShaderParameterMetadataBuilder.h"
 #include "ShaderParameterMetadata.h"
@@ -222,7 +224,9 @@ FString Optimus::GetCookedKernelSource(
 	const FString& InObjectPathName,
 	const FString& InShaderSource,
 	const FString& InKernelName,
-	FIntVector InGroupSize
+	FIntVector InGroupSize,
+	const TCHAR* InReadNumThreadsPerInvocationFunctionName,
+	const TCHAR* InReadThreadIndexOffsetFunctionName
 	)
 {
 	// FIXME: Create source range mappings so that we can go from error location to
@@ -245,11 +249,25 @@ FString Optimus::GetCookedKernelSource(
 		TEXT("[numthreads(%d,%d,%d)]\nvoid %s(uint3 GroupId : SV_GroupID, uint GroupIndex : SV_GroupIndex)"), 
 		InGroupSize.X, InGroupSize.Y, InGroupSize.Z, *InKernelName);
 	
-	const FString UnWrappedDispatchThreadId = FString::Printf(
+	const FString TheadIndexForInvocation = FString::Printf(
 	TEXT("GetUnWrappedDispatchThreadId(GroupId, GroupIndex, %d)"),
 		InGroupSize.X * InGroupSize.Y * InGroupSize.Z
 	);
 
+	const FString ThreadIndexOffsetForKernel = FString::Printf(TEXT(""));
+
+	const FString IndexCheckAndApplyOffset= FString::Printf(
+		TEXT(
+		"uint IndexForInvocation = %s;\n"
+		"if (IndexForInvocation >= %s::%s()) return;\n"
+		"uint Index = IndexForInvocation + %s::%s();\n"),
+		*TheadIndexForInvocation,
+		GetKernelInternalNamespaceName(),
+		InReadNumThreadsPerInvocationFunctionName,
+		GetKernelInternalNamespaceName(),
+		InReadThreadIndexOffsetFunctionName
+		);
+	
 	if (bHasKernelKeyword)
 	{
 		Source.ReplaceInline(TEXT("KERNEL"), TEXT("void __kernel_func(uint Index)"), ESearchCase::CaseSensitive);
@@ -259,8 +277,12 @@ FString Optimus::GetCookedKernelSource(
 				"#line 1 \"%s\"\n"
 				"%s\n\n"
 				"%s\n\n"
-				"%s { __kernel_func(%s); }\n"
-				), *ShaderPathName, *Source, *ComputeShaderUtilsInclude,*KernelFunc, *UnWrappedDispatchThreadId);
+				"%s\n"
+				"{\n"
+				"%s\n"
+				"__kernel_func(Index);\n"
+				"}\n"
+				), *ShaderPathName, *Source, *ComputeShaderUtilsInclude,*KernelFunc, *IndexCheckAndApplyOffset);
 	}
 	else
 	{
@@ -269,12 +291,133 @@ FString Optimus::GetCookedKernelSource(
 			"%s\n"
 			"%s\n"
 			"{\n"
-			"uint Index = %s;\n"
+			"%s\n"
 			"#line 1 \"%s\"\n"
 			"%s\n"
 			"}\n"
-			), *ComputeShaderUtilsInclude,*KernelFunc, *UnWrappedDispatchThreadId, *ShaderPathName, *Source);
+			), *ComputeShaderUtilsInclude,*KernelFunc, *IndexCheckAndApplyOffset, *ShaderPathName, *Source);
+	}
+}
+
+const TCHAR* Optimus::GetKernelInternalNamespaceName()
+{
+	static const TCHAR* KernelInternalNamespaceName = TEXT("KernelInternal");
+	return KernelInternalNamespaceName;
+}
+
+
+Optimus::Expression::FParseResult Optimus::ParseExecutionDomainExpression(
+	const FString& InExpression,
+	TWeakObjectPtr<const UOptimusComponentSource> InComponentSource
+	)
+{
+	TMap<FName, float> EngineConstants;
+
+	for(FName ExecutionDomain: InComponentSource->GetExecutionDomains())
+	{
+		EngineConstants.Add(ExecutionDomain, 0);
+	}
+
+	Expression::FEngine Engine;
+	return Engine.Parse(InExpression, [EngineConstants](FName InName)->TOptional<float>
+	{
+		if (const float* Value = EngineConstants.Find(InName))
+		{
+			return *Value;
+		};
+
+		return {};
+	});
+}
+
+bool Optimus::EvaluateExecutionDomainExpressionParseResult(
+	const TVariant<Expression::FExpressionObject, Expression::FParseError>& InParseResult,
+	TWeakObjectPtr<const UOptimusComponentSource> InComponentSource,
+	TWeakObjectPtr<const UActorComponent> InWeakComponent,
+	TArray<int32>& OutInvocationThreadCount
+	)
+{
+	if (!InWeakComponent.IsValid() || !InComponentSource.IsValid())
+	{
+		return false;
+	}
+	
+	using namespace Expression;
+	if (InParseResult.IsType<FParseError>())
+	{
+		return false;
+	}
+
+	const UActorComponent* Component = InWeakComponent.Get();
+	int32 LodIndex = InComponentSource->GetLodIndex(Component);
+	int32 NumInvocations = InComponentSource->GetDefaultNumInvocations(Component, LodIndex);
+
+	// In some cases the component can be set to an unusable state intentionally
+	// such as when the OptimusEditor shuts down,
+	// in which case we don't have data to work with so simply do nothing
+	if (NumInvocations < 1)
+	{
+		return false;
 	}
+
+	// In the case of Fixed Domain, multiple invocations are supported when the data interface does not support unified dispatch
+	// in the case of unified dispatch, it run a single invocation with thread count equal to the sum of all invocations
+	if (TOptional<FName> FixedDomain = InParseResult.Get<FExpressionObject>().GetAsSingleConstant())
+	{
+		InComponentSource->GetComponentElementCountsForExecutionDomain(
+				*FixedDomain,
+				Component,
+				LodIndex,
+				OutInvocationThreadCount);
+		
+		return OutInvocationThreadCount.Num() > 0;
+	}
+	
+	// Other custom domains only support single invocation / unified dispatch
+	TMap<FName, float> EngineConstants;
+
+	for(FName ExecutionDomain: InComponentSource->GetExecutionDomains())
+	{
+		TArray<int32> ElementCounts;
+		if (!InComponentSource->GetComponentElementCountsForExecutionDomain(
+				ExecutionDomain,
+				Component,
+				LodIndex,
+				ElementCounts))
+		{
+			return false;
+		}
+
+		if (!ensure(NumInvocations == ElementCounts.Num()))
+		{
+			// Component source needs to provide as many values for each of its execution domains as there are invocations
+			return false;
+		}
+
+		int32 TotalCount = 0;
+		for (int32 Count : ElementCounts)
+		{
+			TotalCount += Count;
+		}
+
+		EngineConstants.Add(ExecutionDomain, TotalCount);	
+	}
+
+	FEngine Engine;
+	const int32 TotalElementCountForUnifiedDispatch = static_cast<int32>(Engine.Execute(InParseResult.Get<FExpressionObject>(), 
+		[&EngineConstants](FName InName)->TOptional<float>
+		{
+			if (const float* Value = EngineConstants.Find(InName))
+			{
+				return *Value;
+			};
+
+			return {};
+		}));
+	
+	OutInvocationThreadCount.Add(TotalElementCountForUnifiedDispatch);
+	
+	return true;	
 }
 
 bool Optimus::FindMovedItemInNameArray(const TArray<FName>& Old, const TArray<FName>& New, FName& OutSubjectName, FName& OutNextName)
diff --git a/Engine/Plugins/Animation/DeformerGraph/Source/OptimusCore/Public/IOptimusComputeKernelDataInterface.h b/Engine/Plugins/Animation/DeformerGraph/Source/OptimusCore/Public/IOptimusComputeKernelDataInterface.h
--- a/Engine/Plugins/Animation/DeformerGraph/Source/OptimusCore/Public/IOptimusComputeKernelDataInterface.h
+++ b/Engine/Plugins/Animation/DeformerGraph/Source/OptimusCore/Public/IOptimusComputeKernelDataInterface.h
@@ -24,5 +24,8 @@ public:
 	GENERATED_BODY()
 	virtual void SetExecutionDomain(const FString& InExecutionDomain) = 0;
 	virtual void SetComponentBinding(const UOptimusComponentSourceBinding* InBinding) = 0;
+	virtual const TCHAR* GetReadNumThreadsFunctionName() const = 0;
+	virtual const TCHAR* GetReadNumThreadsPerInvocationFunctionName() const = 0;
+	virtual const TCHAR* GetReadThreadIndexOffsetFunctionName() const = 0;
 	
 };
diff --git a/Engine/Plugins/Animation/DeformerGraph/Source/OptimusCore/Public/OptimusExpressionEvaluator.h b/Engine/Plugins/Animation/DeformerGraph/Source/OptimusCore/Public/OptimusExpressionEvaluator.h
--- a/Engine/Plugins/Animation/DeformerGraph/Source/OptimusCore/Public/OptimusExpressionEvaluator.h
+++ b/Engine/Plugins/Animation/DeformerGraph/Source/OptimusCore/Public/OptimusExpressionEvaluator.h
@@ -55,6 +55,8 @@ struct FExpressionObject
 	/** Returns all named constants used by this expression */
 	TSet<FName> GetUsedConstants() const;
 
+	TOptional<FName> GetAsSingleConstant() const;
+
 	void Load(FArchive& Ar);
 	void Save(FArchive& Ar) const;
 	
@@ -83,6 +85,8 @@ private:
 	TArray<OpElement, TInlineAllocator<8>> Expression;
 };
 
+using FParseResult = TVariant<FExpressionObject, FParseError>;
+	
 /** Simple, fast expression evaluation engine. Can evaluate simple arithmetic expressions,
  *  including user-settable constants. 
  *		1 + 1
@@ -113,7 +117,7 @@ class OPTIMUSCORE_API FEngine
 public:
 	/** Parse an expression and either receive an error or a expression execution object that
 	 *  matches this evaluator. */
-	TVariant<FExpressionObject, FParseError> Parse(
+	FParseResult Parse(
 		FStringView InExpression,
 		TFunction<TOptional<float>(FName InConstantName)> InConstantEvaluator = {} 
 		) const;
diff --git a/Engine/Plugins/Animation/DeformerGraph/Source/OptimusCore/Public/OptimusHelpers.h b/Engine/Plugins/Animation/DeformerGraph/Source/OptimusCore/Public/OptimusHelpers.h
--- a/Engine/Plugins/Animation/DeformerGraph/Source/OptimusCore/Public/OptimusHelpers.h
+++ b/Engine/Plugins/Animation/DeformerGraph/Source/OptimusCore/Public/OptimusHelpers.h
@@ -11,11 +11,15 @@
 #include "UObject/ObjectMacros.h"
 #include "UObject/Package.h"
 #include "UObject/Class.h"
+#include "OptimusExpressionEvaluator.h"
 
+
+class UOptimusComponentSource;
 struct FShaderValueContainerView;
 class FShaderParametersMetadata;
 class FShaderParametersMetadataBuilder;
 struct FShaderValueTypeHandle;
+class UActorComponent;
 
 namespace Optimus
 {
@@ -128,9 +132,25 @@ namespace Optimus
 		const FString& InObjectPathName,
 		const FString& InShaderSource,
 		const FString& InKernelName,
-		FIntVector InGroupSize
+		FIntVector InGroupSize,
+		const TCHAR* InReadNumThreadsPerInvocationFunctionName,
+		const TCHAR* InReadThreadIndexOffsetFunctionName
 	);
 
+	const TCHAR* GetKernelInternalNamespaceName();
+
+	Expression::FParseResult ParseExecutionDomainExpression(
+		const FString& InExpression,
+		TWeakObjectPtr<const UOptimusComponentSource> InComponentSource
+		);
+
+	bool EvaluateExecutionDomainExpressionParseResult(
+		const TVariant<Expression::FExpressionObject, Expression::FParseError>& InParseResult,
+		TWeakObjectPtr<const UOptimusComponentSource> InComponentSource,
+		TWeakObjectPtr<const UActorComponent> InWeakComponent,
+		TArray<int32>& OutInvocationThreadCount
+		);
+	
 	bool FindMovedItemInNameArray(const TArray<FName>& Old, const TArray<FName>& New, FName& OutSubjectName, FName& OutNextName);
 
 	FName GenerateUniqueNameFromExistingNames(FName InBaseName, const TArray<FName>& InExistingNames);
diff --git a/Engine/Plugins/Animation/DeformerGraph/Source/OptimusCore/Public/OptimusObjectVersion.h b/Engine/Plugins/Animation/DeformerGraph/Source/OptimusCore/Public/OptimusObjectVersion.h
--- a/Engine/Plugins/Animation/DeformerGraph/Source/OptimusCore/Public/OptimusObjectVersion.h
+++ b/Engine/Plugins/Animation/DeformerGraph/Source/OptimusCore/Public/OptimusObjectVersion.h
@@ -24,6 +24,7 @@ struct OPTIMUSCORE_API FOptimusObjectVersion
 		PropertyBagValueContainer,
 		SkinnedMeshWriteDIColorBufferManualFetchSwizzle,
 		PropertyPinSupport,
+		KernelThreadIndexUsingTotalThreadCount,
 
 		// -----<new versions can be added above this line>-------------------------------------------------
 		VersionPlusOne,
diff --git a/Engine/Plugins/Animation/MLDeformer/MLDeformerFramework/Shaders/Private/MLDeformerModelHeatMap.ush b/Engine/Plugins/Animation/MLDeformer/MLDeformerFramework/Shaders/Private/MLDeformerModelHeatMap.ush
--- a/Engine/Plugins/Animation/MLDeformer/MLDeformerFramework/Shaders/Private/MLDeformerModelHeatMap.ush
+++ b/Engine/Plugins/Animation/MLDeformer/MLDeformerFramework/Shaders/Private/MLDeformerModelHeatMap.ush
@@ -1,7 +1,6 @@
 // Copyright Epic Games, Inc. All Rights Reserved.
 
 uint {DataInterfaceName}_NumVertices;
-uint {DataInterfaceName}_InputStreamStart;
 uint {DataInterfaceName}_HeatMapMode;
 float {DataInterfaceName}_HeatMapMax;
 float {DataInterfaceName}_GroundTruthLerp;
@@ -16,7 +15,7 @@ uint ReadNumVertices_{DataInterfaceName}()
 
 float3 ReadPositionGroundTruth_{DataInterfaceName}(uint VertexIndex)
 {
-	uint BufferIndex = {DataInterfaceName}_VertexMapBuffer[VertexIndex + {DataInterfaceName}_InputStreamStart];
+	uint BufferIndex = {DataInterfaceName}_VertexMapBuffer[VertexIndex];
 	if (BufferIndex >= {DataInterfaceName}_GroundTruthBufferSize)
 	{
 		return float3(0, 0, 0);
diff --git a/Engine/Plugins/Animation/MLDeformer/MLDeformerFramework/Shaders/Private/SkeletonWithWeightedQuats.ush b/Engine/Plugins/Animation/MLDeformer/MLDeformerFramework/Shaders/Private/SkeletonWithWeightedQuats.ush
--- a/Engine/Plugins/Animation/MLDeformer/MLDeformerFramework/Shaders/Private/SkeletonWithWeightedQuats.ush
+++ b/Engine/Plugins/Animation/MLDeformer/MLDeformerFramework/Shaders/Private/SkeletonWithWeightedQuats.ush
@@ -1,9 +1,7 @@
 // Copyright Epic Games, Inc. All Rights Reserved.
 
 uint {DataInterfaceName}_NumVertices;
-uint {DataInterfaceName}_InputStreamStart;
 uint {DataInterfaceName}_NumBoneInfluences;
-uint {DataInterfaceName}_InputWeightStart;
 uint {DataInterfaceName}_InputWeightStride;
 uint {DataInterfaceName}_InputWeightIndexSize;
 Buffer<float4> {DataInterfaceName}_BoneMatrices;
@@ -70,7 +68,7 @@ float3x4 ReadBoneMatrix_{DataInterfaceName}(uint VertexIndex, uint BoneIndex)
 
 #else // !GPUSKIN_UNLIMITED_BONE_INFLUENCE
 	
-	uint StreamOffset = {DataInterfaceName}_InputWeightStart + (VertexIndex * ({DataInterfaceName}_InputWeightStride / 4));
+	uint StreamOffset = (VertexIndex * ({DataInterfaceName}_InputWeightStride / 4));
 	int BlendIndex = GetBlendIndices_{DataInterfaceName}(StreamOffset, BoneIndex / 4)[BoneIndex & 0x3];
 	return GetBoneMatrix_{DataInterfaceName}(BlendIndex);
 
@@ -104,7 +102,7 @@ float ReadBoneWeight_{DataInterfaceName}(uint VertexIndex, uint BoneIndex)
 
 #else // !GPUSKIN_UNLIMITED_BONE_INFLUENCE
 	
-	uint StreamOffset = {DataInterfaceName}_InputWeightStart + (VertexIndex * ({DataInterfaceName}_InputWeightStride / 4));
+	uint StreamOffset = (VertexIndex * ({DataInterfaceName}_InputWeightStride / 4));
 	float BlendWeight = GetBlendWeights_{DataInterfaceName}(StreamOffset, BoneIndex / 4)[BoneIndex & 0x3];
 	return BlendWeight;
 
@@ -119,13 +117,13 @@ float3x4 ReadWeightedBoneMatrix_{DataInterfaceName}(uint VertexIndex)
 
 #elif GPUSKIN_UNLIMITED_BONE_INFLUENCE
 
-	const uint BufferIndex 		= {DataInterfaceName}_InputStreamStart + VertexIndex;
+	const uint BufferIndex 		= VertexIndex;
 	const uint BlendOffsetCount = {DataInterfaceName}_InputWeightLookupStream[BufferIndex];
 	return ComputeBoneMatrixWithUnlimitedInfluences({DataInterfaceName}_BoneMatrices, {DataInterfaceName}_InputWeightStream, {DataInterfaceName}_InputWeightIndexSize, BlendOffsetCount);
 
 #else
 
-	const uint StreamOffset = {DataInterfaceName}_InputWeightStart + (VertexIndex * ({DataInterfaceName}_InputWeightStride / 4));
+	const uint StreamOffset = (VertexIndex * ({DataInterfaceName}_InputWeightStride / 4));
 	const uint NumInfluences = {DataInterfaceName}_NumBoneInfluences;
 	FGPUSkinIndexAndWeight IndicesAndWeights = (FGPUSkinIndexAndWeight)0;
 	if (NumInfluences > 0)
@@ -161,7 +159,7 @@ float4 ReadWeightedBoneQuat_{DataInterfaceName}(uint VertexIndex)
 	float4 BlendQuat = (float4)0;
 	float4 BoneQuat0 = (float4)0;
 
-	uint BufferIndex = {DataInterfaceName}_InputStreamStart + VertexIndex;
+	uint BufferIndex = VertexIndex;
 	uint BlendOffsetCount = {DataInterfaceName}_InputWeightLookupStream[BufferIndex];
 	int NumBoneInfluencesLocal = BlendOffsetCount & 0xff;
 	int IndicesOffset = BlendOffsetCount >> 8;
@@ -199,7 +197,7 @@ float4 ReadWeightedBoneQuat_{DataInterfaceName}(uint VertexIndex)
 	float4 BlendQuat = (float4)0;
 	float4 BoneQuat0 = (float4)0;
 
-	uint StreamOffset = {DataInterfaceName}_InputWeightStart + (VertexIndex * ({DataInterfaceName}_InputWeightStride / 4));
+	uint StreamOffset = (VertexIndex * ({DataInterfaceName}_InputWeightStride / 4));
 
 	if ({DataInterfaceName}_NumBoneInfluences > 0)
 	{
diff --git a/Engine/Plugins/Animation/MLDeformer/MLDeformerFramework/Source/MLDeformerFramework/Private/DeformerGraphDataInterfaceSkeletonWithQuats.cpp b/Engine/Plugins/Animation/MLDeformer/MLDeformerFramework/Source/MLDeformerFramework/Private/DeformerGraphDataInterfaceSkeletonWithQuats.cpp
--- a/Engine/Plugins/Animation/MLDeformer/MLDeformerFramework/Source/MLDeformerFramework/Private/DeformerGraphDataInterfaceSkeletonWithQuats.cpp
+++ b/Engine/Plugins/Animation/MLDeformer/MLDeformerFramework/Source/MLDeformerFramework/Private/DeformerGraphDataInterfaceSkeletonWithQuats.cpp
@@ -71,9 +71,7 @@ void UDeformerGraphSkeletonWithWeightedQuatsDataInterface::GetSupportedInputs(TA
 
 BEGIN_SHADER_PARAMETER_STRUCT(FSkeletonWithQuatsDataInterfaceParameters, )
 	SHADER_PARAMETER(uint32, NumVertices)
-	SHADER_PARAMETER(uint32, InputStreamStart)
 	SHADER_PARAMETER(uint32, NumBoneInfluences)
-	SHADER_PARAMETER(uint32, InputWeightStart)
 	SHADER_PARAMETER(uint32, InputWeightStride)
 	SHADER_PARAMETER(uint32, InputWeightIndexSize)
 	SHADER_PARAMETER_SRV(Buffer<float4>, BoneMatrices)
@@ -271,25 +269,21 @@ namespace UE::MLDeformer
 
 		FRHIShaderResourceView* NullSRVBinding = GWhiteVertexBufferWithSRV->ShaderResourceViewRHI.GetReference();
 
+		FSkinWeightVertexBuffer const* WeightBuffer = LodRenderData->GetSkinWeightVertexBuffer();
+		check(WeightBuffer != nullptr);
+		FRHIShaderResourceView* SkinWeightBufferSRV = WeightBuffer->GetDataVertexBuffer()->GetSRV();
+		const bool bUnlimitedBoneInfluences = WeightBuffer->GetBoneInfluenceType() == GPUSkinBoneInfluenceType::UnlimitedBoneInfluence;
+		FRHIShaderResourceView* InputWeightLookupStreamSRV = bUnlimitedBoneInfluences ? WeightBuffer->GetLookupVertexBuffer()->GetSRV() : nullptr;
+
 		const TStridedView<FParameters> ParameterArray = MakeStridedParameterView<FParameters>(InDispatchData);
 		for (int32 InvocationIndex = 0; InvocationIndex < ParameterArray.Num(); ++InvocationIndex)
 		{
-			FSkelMeshRenderSection const& RenderSection = LodRenderData->RenderSections[InvocationIndex];
-
 			const bool bPreviousFrame = false;
 			FRHIShaderResourceView* BoneBufferSRV = FSkeletalMeshDeformerHelpers::GetBoneBufferForReading(SkeletalMeshObject, LodIndex, InvocationIndex, bPreviousFrame);
-
-			FSkinWeightVertexBuffer const* WeightBuffer = LodRenderData->GetSkinWeightVertexBuffer();
-			check(WeightBuffer != nullptr);
-			FRHIShaderResourceView* SkinWeightBufferSRV = WeightBuffer->GetDataVertexBuffer()->GetSRV();
-			const bool bUnlimitedBoneInfluences = WeightBuffer->GetBoneInfluenceType() == GPUSkinBoneInfluenceType::UnlimitedBoneInfluence;
-			FRHIShaderResourceView* InputWeightLookupStreamSRV = bUnlimitedBoneInfluences ? WeightBuffer->GetLookupVertexBuffer()->GetSRV() : nullptr;
-
+			
 			FParameters& Parameters = ParameterArray[InvocationIndex];
-			Parameters.NumVertices = RenderSection.NumVertices;
-			Parameters.InputStreamStart = RenderSection.BaseVertexIndex;
+			Parameters.NumVertices = LodRenderData->GetNumVertices();
 			Parameters.NumBoneInfluences = WeightBuffer->GetMaxBoneInfluences();
-			Parameters.InputWeightStart = (WeightBuffer->GetConstantInfluencesVertexStride() * RenderSection.GetVertexBufferIndex()) / sizeof(float);
 			Parameters.InputWeightStride = WeightBuffer->GetConstantInfluencesVertexStride();
 			Parameters.InputWeightIndexSize = WeightBuffer->GetBoneIndexByteSize() | (WeightBuffer->GetBoneWeightByteSize() << 8);
 			Parameters.BoneMatrices = BoneBufferSRV != nullptr ? BoneBufferSRV : NullSRVBinding;
diff --git a/Engine/Plugins/Animation/MLDeformer/MLDeformerFramework/Source/MLDeformerFramework/Private/MLDeformerGraphDebugDataInterface.cpp b/Engine/Plugins/Animation/MLDeformer/MLDeformerFramework/Source/MLDeformerFramework/Private/MLDeformerGraphDebugDataInterface.cpp
--- a/Engine/Plugins/Animation/MLDeformer/MLDeformerFramework/Source/MLDeformerFramework/Private/MLDeformerGraphDebugDataInterface.cpp
+++ b/Engine/Plugins/Animation/MLDeformer/MLDeformerFramework/Source/MLDeformerFramework/Private/MLDeformerGraphDebugDataInterface.cpp
@@ -62,7 +62,6 @@ void UMLDeformerGraphDebugDataInterface::GetSupportedInputs(TArray<FShaderFuncti
 
 BEGIN_SHADER_PARAMETER_STRUCT(FMLDeformerGraphDebugDataInterfaceParameters, )
 	SHADER_PARAMETER(uint32, NumVertices)
-	SHADER_PARAMETER(uint32, InputStreamStart)
 	SHADER_PARAMETER(int32, HeatMapMode)
 	SHADER_PARAMETER(float, HeatMapMax)
 	SHADER_PARAMETER(float, GroundTruthLerp)
@@ -221,10 +220,8 @@ namespace UE::MLDeformer
 		const TStridedView<FMLDeformerGraphDebugDataInterfaceParameters> ParameterArray = MakeStridedParameterView<FMLDeformerGraphDebugDataInterfaceParameters>(InDispatchData);
 		for (int32 InvocationIndex = 0; InvocationIndex < ParameterArray.Num(); ++InvocationIndex)
 		{
-			const FSkelMeshRenderSection& RenderSection = LodRenderData->RenderSections[InvocationIndex];
 			FMLDeformerGraphDebugDataInterfaceParameters& Parameters = ParameterArray[InvocationIndex];
-			Parameters.NumVertices = InDispatchData.bUnifiedDispatch ? LodRenderData->GetNumVertices() : RenderSection.GetNumVertices();
-			Parameters.InputStreamStart = InDispatchData.bUnifiedDispatch ? 0 : RenderSection.BaseVertexIndex;
+			Parameters.NumVertices = LodRenderData->GetNumVertices();
 			Parameters.HeatMapMode = HeatMapMode;
 			Parameters.HeatMapMax = HeatMapMax;
 			Parameters.GroundTruthLerp = GroundTruthLerp;
diff --git a/Engine/Plugins/Animation/MLDeformer/NearestNeighborModel/Shaders/Private/DataInterfaceSkeletonWithQuats.ush b/Engine/Plugins/Animation/MLDeformer/NearestNeighborModel/Shaders/Private/DataInterfaceSkeletonWithQuats.ush
--- a/Engine/Plugins/Animation/MLDeformer/NearestNeighborModel/Shaders/Private/DataInterfaceSkeletonWithQuats.ush
+++ b/Engine/Plugins/Animation/MLDeformer/NearestNeighborModel/Shaders/Private/DataInterfaceSkeletonWithQuats.ush
@@ -1,9 +1,7 @@
 // Copyright Epic Games, Inc. All Rights Reserved.
 
 uint {DataInterfaceName}_NumVertices;
-uint {DataInterfaceName}_InputStreamStart;
 uint {DataInterfaceName}_NumBoneInfluences;
-uint {DataInterfaceName}_InputWeightStart;
 uint {DataInterfaceName}_InputWeightStride;
 uint {DataInterfaceName}_InputWeightIndexSize;
 Buffer<float4> {DataInterfaceName}_BoneMatrices;
@@ -70,7 +68,7 @@ float3x4 ReadBoneMatrix_{DataInterfaceName}(uint VertexIndex, uint BoneIndex)
 
 #else // !GPUSKIN_UNLIMITED_BONE_INFLUENCE
 	
-	uint StreamOffset = {DataInterfaceName}_InputWeightStart + (VertexIndex * ({DataInterfaceName}_InputWeightStride / 4));
+	uint StreamOffset = (VertexIndex * ({DataInterfaceName}_InputWeightStride / 4));
 	int BlendIndex = GetBlendIndices_{DataInterfaceName}(StreamOffset, BoneIndex / 4)[BoneIndex & 0x3];
 	return GetBoneMatrix_{DataInterfaceName}(BlendIndex);
 
@@ -104,7 +102,7 @@ float ReadBoneWeight_{DataInterfaceName}(uint VertexIndex, uint BoneIndex)
 
 #else // !GPUSKIN_UNLIMITED_BONE_INFLUENCE
 	
-	uint StreamOffset = {DataInterfaceName}_InputWeightStart + (VertexIndex * ({DataInterfaceName}_InputWeightStride / 4));
+	uint StreamOffset = (VertexIndex * ({DataInterfaceName}_InputWeightStride / 4));
 	float BlendWeight = GetBlendWeights_{DataInterfaceName}(StreamOffset, BoneIndex / 4)[BoneIndex & 0x3];
 	return BlendWeight;
 
@@ -119,13 +117,13 @@ float3x4 ReadWeightedBoneMatrix_{DataInterfaceName}(uint VertexIndex)
 
 #elif GPUSKIN_UNLIMITED_BONE_INFLUENCE
 
-	const uint BufferIndex 		= {DataInterfaceName}_InputStreamStart + VertexIndex;
+	const uint BufferIndex 		= VertexIndex;
 	const uint BlendOffsetCount = {DataInterfaceName}_InputWeightLookupStream[BufferIndex];
 	return ComputeBoneMatrixWithUnlimitedInfluences({DataInterfaceName}_BoneMatrices, {DataInterfaceName}_InputWeightStream, {DataInterfaceName}_InputWeightIndexSize, BlendOffsetCount);
 
 #else
 
-	const uint StreamOffset = {DataInterfaceName}_InputWeightStart + (VertexIndex * ({DataInterfaceName}_InputWeightStride / 4));
+	const uint StreamOffset = (VertexIndex * ({DataInterfaceName}_InputWeightStride / 4));
 	const uint NumInfluences = {DataInterfaceName}_NumBoneInfluences;
 	FGPUSkinIndexAndWeight IndicesAndWeights = (FGPUSkinIndexAndWeight)0;
 	if (NumInfluences > 0)
@@ -161,7 +159,7 @@ float4 ReadWeightedBoneQuat_{DataInterfaceName}(uint VertexIndex)
 	float4 BlendQuat = (float4)0;
 	float4 BoneQuat0 = (float4)0;
 
-	uint BufferIndex = {DataInterfaceName}_InputStreamStart + VertexIndex;
+	uint BufferIndex = VertexIndex;
 	uint BlendOffsetCount = {DataInterfaceName}_InputWeightLookupStream[BufferIndex];
 	int NumBoneInfluencesLocal = BlendOffsetCount & 0xff;
 	int IndicesOffset = BlendOffsetCount >> 8;
@@ -199,7 +197,7 @@ float4 ReadWeightedBoneQuat_{DataInterfaceName}(uint VertexIndex)
 	float4 BlendQuat = (float4)0;
 	float4 BoneQuat0 = (float4)0;
 
-	uint StreamOffset = {DataInterfaceName}_InputWeightStart + (VertexIndex * ({DataInterfaceName}_InputWeightStride / 4));
+	uint StreamOffset = (VertexIndex * ({DataInterfaceName}_InputWeightStride / 4));
 
 	if ({DataInterfaceName}_NumBoneInfluences > 0)
 	{
diff --git a/Engine/Plugins/Animation/MLDeformer/NearestNeighborModel/Source/NearestNeighborModel/Private/OptimusDataInterfaceSkeletonWithQuats.cpp b/Engine/Plugins/Animation/MLDeformer/NearestNeighborModel/Source/NearestNeighborModel/Private/OptimusDataInterfaceSkeletonWithQuats.cpp
--- a/Engine/Plugins/Animation/MLDeformer/NearestNeighborModel/Source/NearestNeighborModel/Private/OptimusDataInterfaceSkeletonWithQuats.cpp
+++ b/Engine/Plugins/Animation/MLDeformer/NearestNeighborModel/Source/NearestNeighborModel/Private/OptimusDataInterfaceSkeletonWithQuats.cpp
@@ -74,9 +74,7 @@ void UOptimusSkeletonWithQuatsDataInterface::GetSupportedInputs(TArray<FShaderFu
 
 BEGIN_SHADER_PARAMETER_STRUCT(FSkeletonWithQuatsDataInterfaceParameters, )
 	SHADER_PARAMETER(uint32, NumVertices)
-	SHADER_PARAMETER(uint32, InputStreamStart)
 	SHADER_PARAMETER(uint32, NumBoneInfluences)
-	SHADER_PARAMETER(uint32, InputWeightStart)
 	SHADER_PARAMETER(uint32, InputWeightStride)
 	SHADER_PARAMETER(uint32, InputWeightIndexSize)
 	SHADER_PARAMETER_SRV(Buffer<float4>, BoneMatrices)
@@ -290,10 +288,8 @@ void FOptimusSkeletonWithQuatsDataProviderProxy::GatherDispatchData(FDispatchDat
 		FRHIShaderResourceView* InputWeightLookupStreamSRV = bUnlimitedBoneInfluences ? WeightBuffer->GetLookupVertexBuffer()->GetSRV() : nullptr;
 
 		FParameters& Parameters = ParameterArray[InvocationIndex];
-		Parameters.NumVertices = RenderSection.NumVertices;
-		Parameters.InputStreamStart = RenderSection.BaseVertexIndex;
+		Parameters.NumVertices = LodRenderData->GetNumVertices();
 		Parameters.NumBoneInfluences = WeightBuffer->GetMaxBoneInfluences();
-		Parameters.InputWeightStart = (WeightBuffer->GetConstantInfluencesVertexStride() * RenderSection.GetVertexBufferIndex()) / sizeof(float);
 		Parameters.InputWeightStride = WeightBuffer->GetConstantInfluencesVertexStride();
 		Parameters.InputWeightIndexSize = WeightBuffer->GetBoneIndexByteSize() | (WeightBuffer->GetBoneWeightByteSize() << 8);
 		Parameters.BoneMatrices = BoneBufferSRV != nullptr ? BoneBufferSRV : NullSRVBinding;
diff --git a/Engine/Plugins/Animation/MLDeformer/VertexDeltaModel/Shaders/Private/VertexDeltaModel.ush b/Engine/Plugins/Animation/MLDeformer/VertexDeltaModel/Shaders/Private/VertexDeltaModel.ush
--- a/Engine/Plugins/Animation/MLDeformer/VertexDeltaModel/Shaders/Private/VertexDeltaModel.ush
+++ b/Engine/Plugins/Animation/MLDeformer/VertexDeltaModel/Shaders/Private/VertexDeltaModel.ush
@@ -1,7 +1,6 @@
 // Copyright Epic Games, Inc. All Rights Reserved.
 
 uint {DataInterfaceName}_NumVertices;
-uint {DataInterfaceName}_InputStreamStart;
 float {DataInterfaceName}_Weight;
 Buffer<float> {DataInterfaceName}_PositionDeltaBuffer;
 Buffer<uint> {DataInterfaceName}_VertexMapBuffer;
@@ -13,7 +12,7 @@ uint ReadNumVertices_{DataInterfaceName}()
 
 float3 ReadPositionDelta_{DataInterfaceName}(uint VertexIndex)
 {
-	uint BufferIndex = {DataInterfaceName}_VertexMapBuffer[VertexIndex + {DataInterfaceName}_InputStreamStart] * 3;
+	uint BufferIndex = {DataInterfaceName}_VertexMapBuffer[VertexIndex] * 3;
 	float3 Delta = float3({DataInterfaceName}_PositionDeltaBuffer[BufferIndex], {DataInterfaceName}_PositionDeltaBuffer[BufferIndex + 1], {DataInterfaceName}_PositionDeltaBuffer[BufferIndex + 2]);
 	return Delta * {DataInterfaceName}_Weight;
 }
diff --git a/Engine/Plugins/Animation/MLDeformer/VertexDeltaModel/Shaders/Private/VertexDeltaModelHeatMap.ush b/Engine/Plugins/Animation/MLDeformer/VertexDeltaModel/Shaders/Private/VertexDeltaModelHeatMap.ush
--- a/Engine/Plugins/Animation/MLDeformer/VertexDeltaModel/Shaders/Private/VertexDeltaModelHeatMap.ush
+++ b/Engine/Plugins/Animation/MLDeformer/VertexDeltaModel/Shaders/Private/VertexDeltaModelHeatMap.ush
@@ -1,7 +1,6 @@
 // Copyright Epic Games, Inc. All Rights Reserved.
 
 uint {DataInterfaceName}_NumVertices;
-uint {DataInterfaceName}_InputStreamStart;
 uint {DataInterfaceName}_HeatMapMode;
 float {DataInterfaceName}_HeatMapMax;
 float {DataInterfaceName}_GroundTruthLerp;
@@ -16,7 +15,7 @@ uint ReadNumVertices_{DataInterfaceName}()
 
 float3 ReadPositionGroundTruth_{DataInterfaceName}(uint VertexIndex)
 {
-	uint BufferIndex = {DataInterfaceName}_VertexMapBuffer[VertexIndex + {DataInterfaceName}_InputStreamStart];
+	uint BufferIndex = {DataInterfaceName}_VertexMapBuffer[VertexIndex];
 	if (BufferIndex >= {DataInterfaceName}_GroundTruthBufferSize)
 	{
 		return float3(0, 0, 0);
diff --git a/Engine/Plugins/Animation/MLDeformer/VertexDeltaModel/Source/VertexDeltaModel/Private/VertexDeltaGraphDataInterface.cpp b/Engine/Plugins/Animation/MLDeformer/VertexDeltaModel/Source/VertexDeltaModel/Private/VertexDeltaGraphDataInterface.cpp
--- a/Engine/Plugins/Animation/MLDeformer/VertexDeltaModel/Source/VertexDeltaModel/Private/VertexDeltaGraphDataInterface.cpp
+++ b/Engine/Plugins/Animation/MLDeformer/VertexDeltaModel/Source/VertexDeltaModel/Private/VertexDeltaGraphDataInterface.cpp
@@ -46,7 +46,6 @@ void UVertexDeltaGraphDataInterface::GetSupportedInputs(TArray<FShaderFunctionDe
 
 BEGIN_SHADER_PARAMETER_STRUCT(FVertexDeltaGraphDataInterfaceParameters, )
 	SHADER_PARAMETER(uint32, NumVertices)
-	SHADER_PARAMETER(uint32, InputStreamStart)
 	SHADER_PARAMETER(float, Weight)
 	SHADER_PARAMETER_RDG_BUFFER_SRV(Buffer<float>, PositionDeltaBuffer)
 	SHADER_PARAMETER_SRV(Buffer<uint>, VertexMapBuffer)
@@ -159,10 +158,8 @@ namespace UE::VertexDeltaModel
 		const TStridedView<FVertexDeltaGraphDataInterfaceParameters> ParameterArray = MakeStridedParameterView<FVertexDeltaGraphDataInterfaceParameters>(InDispatchData);
 		for (int32 InvocationIndex = 0; InvocationIndex < ParameterArray.Num(); ++InvocationIndex)
 		{
-			const FSkelMeshRenderSection& RenderSection = LodRenderData->RenderSections[InvocationIndex];
 			FVertexDeltaGraphDataInterfaceParameters& Parameters = ParameterArray[InvocationIndex];
-			Parameters.NumVertices = InDispatchData.bUnifiedDispatch ? LodRenderData->GetNumVertices() : RenderSection.GetNumVertices();
-			Parameters.InputStreamStart = InDispatchData.bUnifiedDispatch ? 0 : RenderSection.BaseVertexIndex;
+			Parameters.NumVertices = LodRenderData->GetNumVertices();
 			Parameters.Weight = Weight;
 			Parameters.PositionDeltaBuffer = BufferSRV;
 			Parameters.VertexMapBuffer = VertexMapBufferSRV;
diff --git a/Engine/Plugins/Animation/MLDeformer/VertexDeltaModel/Source/VertexDeltaModel/Private/VertexDeltaGraphDebugDataInterface.cpp b/Engine/Plugins/Animation/MLDeformer/VertexDeltaModel/Source/VertexDeltaModel/Private/VertexDeltaGraphDebugDataInterface.cpp
--- a/Engine/Plugins/Animation/MLDeformer/VertexDeltaModel/Source/VertexDeltaModel/Private/VertexDeltaGraphDebugDataInterface.cpp
+++ b/Engine/Plugins/Animation/MLDeformer/VertexDeltaModel/Source/VertexDeltaModel/Private/VertexDeltaGraphDebugDataInterface.cpp
@@ -63,7 +63,6 @@ void UDEPRECATED_VertexDeltaGraphDebugDataInterface::GetSupportedInputs(TArray<F
 
 BEGIN_SHADER_PARAMETER_STRUCT(FVertexDeltaGraphDebugDataInterfaceParameters, )
 	SHADER_PARAMETER(uint32, NumVertices)
-	SHADER_PARAMETER(uint32, InputStreamStart)
 	SHADER_PARAMETER(int32, HeatMapMode)
 	SHADER_PARAMETER(float, HeatMapMax)
 	SHADER_PARAMETER(float, GroundTruthLerp)
@@ -213,10 +212,8 @@ namespace UE::VertexDeltaModel
 		const TStridedView<FVertexDeltaGraphDebugDataInterfaceParameters> ParameterArray = MakeStridedParameterView<FVertexDeltaGraphDebugDataInterfaceParameters>(InDispatchData);
 		for (int32 InvocationIndex = 0; InvocationIndex < ParameterArray.Num(); ++InvocationIndex)
 		{
-			const FSkelMeshRenderSection& RenderSection = LodRenderData->RenderSections[InvocationIndex];
 			FVertexDeltaGraphDebugDataInterfaceParameters& Parameters = ParameterArray[InvocationIndex];
-			Parameters.NumVertices = InDispatchData.bUnifiedDispatch ? LodRenderData->GetNumVertices() : RenderSection.GetNumVertices();
-			Parameters.InputStreamStart = InDispatchData.bUnifiedDispatch ? 0 : RenderSection.BaseVertexIndex;
+			Parameters.NumVertices = LodRenderData->GetNumVertices();
 			Parameters.HeatMapMode = HeatMapMode;
 			Parameters.HeatMapMax = HeatMapMax;
 			Parameters.GroundTruthLerp = GroundTruthLerp;
diff --git a/Engine/Plugins/Experimental/ChaosFlesh/Shaders/Private/DIFleshDeformer.ush b/Engine/Plugins/Experimental/ChaosFlesh/Shaders/Private/DIFleshDeformer.ush
--- a/Engine/Plugins/Experimental/ChaosFlesh/Shaders/Private/DIFleshDeformer.ush
+++ b/Engine/Plugins/Experimental/ChaosFlesh/Shaders/Private/DIFleshDeformer.ush
@@ -1,7 +1,6 @@
 // Copyright Epic Games, Inc.  All Rights Reserved.
 
 uint {DataInterfaceName}_NumVertices;
-uint {DataInterfaceName}_BaseVertexIndex; // InputStreamStart
 
 uint {DataInterfaceName}_NumTetVertexBuffer;
 uint {DataInterfaceName}_NumTetRestVertexBuffer;
@@ -146,9 +145,7 @@ int4 ReadParents8_{DataInterfaceName}(uint BufferIndex)
 int4 ReadParents_{DataInterfaceName}(uint VertexIndex)
 {
 #if ENABLE_DEFORMER_FLESH
-	// The shader only sees one section of the mesh at a time, each of which is indexed starting at 0,
-	// so we need to offset VertexIndex by the base vertex index, to get our buffer relative index.
-	uint BufferIndex = VertexIndex + {DataInterfaceName}_BaseVertexIndex;
+	uint BufferIndex = VertexIndex;
 
 	// Many render vertices will likely share the same set of parents, so the set of parents is sparsely
 	// represented.  The VertsToParentsBuffer contains the mapping from render vertices to parents.
@@ -183,7 +180,7 @@ half4 ReadWeights_{DataInterfaceName}(uint VertexIndex)
 #if ENABLE_DEFORMER_FLESH
 	if(VertexIndex < {DataInterfaceName}_NumWeightsBuffer)
 	{
-		uint BufferIndex = VertexIndex + {DataInterfaceName}_BaseVertexIndex;
+		uint BufferIndex = VertexIndex;
 		return half4(
 			{DataInterfaceName}_WeightsBuffer[BufferIndex * 4], 
 			{DataInterfaceName}_WeightsBuffer[BufferIndex * 4 + 1], 
@@ -199,7 +196,7 @@ half3 ReadOffset_{DataInterfaceName}(uint VertexIndex)
 #if ENABLE_DEFORMER_FLESH
 	if(VertexIndex < {DataInterfaceName}_NumOffsetBuffer)
 	{
-		uint BufferIndex = VertexIndex + {DataInterfaceName}_BaseVertexIndex;
+		uint BufferIndex = VertexIndex;
 		return half3(
 			{DataInterfaceName}_OffsetBuffer[BufferIndex * 3], 
 			{DataInterfaceName}_OffsetBuffer[BufferIndex * 3 + 1], 
@@ -214,7 +211,7 @@ half ReadMask_{DataInterfaceName}(uint VertexIndex)
 #if ENABLE_DEFORMER_FLESH
 	if(VertexIndex < {DataInterfaceName}_NumMaskBuffer)
 	{
-		uint BufferIndex = VertexIndex + {DataInterfaceName}_BaseVertexIndex;
+		uint BufferIndex = VertexIndex;
 		return {DataInterfaceName}_MaskBuffer[BufferIndex];
 	}
 #endif
diff --git a/Engine/Plugins/Experimental/ChaosFlesh/Source/ChaosFleshEngine/Private/DataInterfaces/DIFleshDeformer.cpp b/Engine/Plugins/Experimental/ChaosFlesh/Source/ChaosFleshEngine/Private/DataInterfaces/DIFleshDeformer.cpp
--- a/Engine/Plugins/Experimental/ChaosFlesh/Source/ChaosFleshEngine/Private/DataInterfaces/DIFleshDeformer.cpp
+++ b/Engine/Plugins/Experimental/ChaosFlesh/Source/ChaosFleshEngine/Private/DataInterfaces/DIFleshDeformer.cpp
@@ -93,7 +93,6 @@ UDIFleshDeformer::GetSupportedInputs(TArray<FShaderFunctionDefinition>& OutFunct
 
 BEGIN_SHADER_PARAMETER_STRUCT(FFleshDeformerDataInterfaceParameters, )
 	SHADER_PARAMETER(uint32, NumVertices)
-	SHADER_PARAMETER(uint32, BaseVertexIndex)
 	SHADER_PARAMETER(uint32, NumTetRestVertexBuffer)
 	SHADER_PARAMETER(uint32, NumTetVertexBuffer)
 
@@ -368,12 +367,9 @@ FDIFleshDeformerProviderProxy::GatherDispatchData(FDispatchData const& InDispatc
 
 	for (int32 InvocationIndex = 0; InvocationIndex < ParameterArray.Num(); ++InvocationIndex)
 	{
-		FSkelMeshRenderSection const& RenderSection = LodRenderData->RenderSections[InvocationIndex];
-		
 		FParameters& Parameters = ParameterArray[InvocationIndex];
 
-		Parameters.NumVertices = RenderSection.GetNumVertices();
-		Parameters.BaseVertexIndex = RenderSection.BaseVertexIndex;
+		Parameters.NumVertices = LodRenderData->GetNumVertices();
 
 		//
 		// Time varying
diff --git a/Engine/Plugins/Runtime/HairStrands/Shaders/Private/Deformers/DeformerGroomGuidesRead.ush b/Engine/Plugins/Runtime/HairStrands/Shaders/Private/Deformers/DeformerGroomGuidesRead.ush
--- a/Engine/Plugins/Runtime/HairStrands/Shaders/Private/Deformers/DeformerGroomGuidesRead.ush
+++ b/Engine/Plugins/Runtime/HairStrands/Shaders/Private/Deformers/DeformerGroomGuidesRead.ush
@@ -7,6 +7,10 @@
 
 uint {DataInterfaceName}_PointCount; 
 uint {DataInterfaceName}_CurveCount;
+uint {DataInterfaceName}_BasePointIndex; 
+uint {DataInterfaceName}_BaseCurveIndex;
+uint {DataInterfaceName}_TotalPointCount; 
+uint {DataInterfaceName}_TotalCurveCount;
 uint {DataInterfaceName}_ObjectIndex;
 float3 {DataInterfaceName}_RestPositionOffset;
 
@@ -46,12 +50,12 @@ float3x4 BuildSkinnedTransform(uint CurveIndex, Buffer<float4> TrianglePositions
 
 uint ReadNumPoints_{DataInterfaceName}()
 {
-	return {DataInterfaceName}_PointCount;
+	return {DataInterfaceName}_TotalPointCount;
 }
 
 uint ReadNumCurves_{DataInterfaceName}()
 {
-	return {DataInterfaceName}_CurveCount;
+	return {DataInterfaceName}_TotalCurveCount;
 }
 
 uint ReadObjectIndex_{DataInterfaceName}()
@@ -61,26 +65,31 @@ uint ReadObjectIndex_{DataInterfaceName}()
 
 float3 ReadPointRestPosition_{DataInterfaceName}(uint PointIndex)
 {
+	PointIndex -= {DataInterfaceName}_BasePointIndex;
 	return (PointIndex < {DataInterfaceName}_PointCount) ? ReadHairControlPoint({DataInterfaceName}_PointRestPositions, PointIndex, {DataInterfaceName}_RestPositionOffset).Position : float3(0,0,0);
 }
 
 float ReadCurvePointOffset_{DataInterfaceName}(uint CurveIndex)
 {
-	return (CurveIndex < {DataInterfaceName}_CurveCount) ? ReadHairCurve({DataInterfaceName}_CurvePointOffsets, CurveIndex).PointIndex : 0;
+	CurveIndex -= {DataInterfaceName}_BaseCurveIndex;
+	return (CurveIndex < {DataInterfaceName}_CurveCount) ? ReadHairCurve({DataInterfaceName}_CurvePointOffsets, CurveIndex).PointIndex + {DataInterfaceName}_BasePointIndex : 0 + {DataInterfaceName}_BasePointIndex;
 }
 
 float ReadCurveNumPoints_{DataInterfaceName}(uint CurveIndex)
 {
+	CurveIndex -= {DataInterfaceName}_BaseCurveIndex;
 	return (CurveIndex < {DataInterfaceName}_CurveCount) ? ReadHairCurve({DataInterfaceName}_CurvePointOffsets, CurveIndex).PointCount : 0;
 }
 
 uint ReadPointCurveIndex_{DataInterfaceName}(uint PointIndex)
 {
-	return (PointIndex < {DataInterfaceName}_PointCount) ? ReadHairPointToCurveIndex({DataInterfaceName}_PointCurveIndices, PointIndex) : 0;
+	PointIndex -= {DataInterfaceName}_BasePointIndex;
+	return (PointIndex < {DataInterfaceName}_PointCount) ? ReadHairPointToCurveIndex({DataInterfaceName}_PointCurveIndices, PointIndex) + {DataInterfaceName}_BaseCurveIndex : {DataInterfaceName}_BaseCurveIndex;
 }
 
 float3x4 ReadCurveRestTransform_{DataInterfaceName}(uint CurveIndex)
 {
+	CurveIndex -= {DataInterfaceName}_BaseCurveIndex;	
 #if ENABLE_SKINNED_TRANSFORM == 1
 		return BuildSkinnedTransform(CurveIndex, {DataInterfaceName}_TriangleRestPositions,
 			{DataInterfaceName}_CurveBarycentricCoordinates, {DataInterfaceName}_CurveTriangleIndices);
@@ -93,6 +102,7 @@ float3x4 ReadCurveRestTransform_{DataInterfaceName}(uint CurveIndex)
 
 float3x4 ReadCurveDeformedTransform_{DataInterfaceName}(uint CurveIndex)
 {
+	CurveIndex -= {DataInterfaceName}_BaseCurveIndex;
 #if ENABLE_SKINNED_TRANSFORM == 1
 		return BuildSkinnedTransform(CurveIndex, {DataInterfaceName}_TriangleDeformedPositions,
 			{DataInterfaceName}_CurveBarycentricCoordinates, {DataInterfaceName}_CurveTriangleIndices);
diff --git a/Engine/Plugins/Runtime/HairStrands/Shaders/Private/Deformers/DeformerGroomGuidesWrite.ush b/Engine/Plugins/Runtime/HairStrands/Shaders/Private/Deformers/DeformerGroomGuidesWrite.ush
--- a/Engine/Plugins/Runtime/HairStrands/Shaders/Private/Deformers/DeformerGroomGuidesWrite.ush
+++ b/Engine/Plugins/Runtime/HairStrands/Shaders/Private/Deformers/DeformerGroomGuidesWrite.ush
@@ -4,8 +4,9 @@
 #include "/Engine/Private/HairStrands/HairStrandsVertexFactoryCommon.ush"
 
 uint {DataInterfaceName}_RegisteredIndex; 
-uint {DataInterfaceName}_PointCount; 
-uint {DataInterfaceName}_OutputStreamStart;
+uint {DataInterfaceName}_PointCount;
+uint {DataInterfaceName}_TotalPointCount;
+uint {DataInterfaceName}_BasePointIndex; 
 
 StructuredBuffer<float4> {DataInterfaceName}_DeformedPositionOffset;
 ByteAddressBuffer {DataInterfaceName}_PointRestPositions;
@@ -22,12 +23,13 @@ float3 {DataInterfaceName}_GetDeformedPositionOffset()
 
 uint ReadNumPoints_{DataInterfaceName}()
 {
-	return {DataInterfaceName}_PointCount;
+	return {DataInterfaceName}_TotalPointCount;
 }
 
 void WritePointDeformedPosition_{DataInterfaceName}(uint PointIndex, float3 PointPosition)
 {
-	const uint WriteIndex = {DataInterfaceName}_OutputStreamStart + PointIndex;
+	PointIndex -= {DataInterfaceName}_BasePointIndex;
+	const uint WriteIndex = PointIndex;
 	if (WriteIndex < {DataInterfaceName}_PointCount)
 	{
 		const FPackedHairPosition PackedCP = ReadPackedHairPosition({DataInterfaceName}_PointRestPositions, WriteIndex);
diff --git a/Engine/Plugins/Runtime/HairStrands/Shaders/Private/Deformers/DeformerGroomStrandsRead.ush b/Engine/Plugins/Runtime/HairStrands/Shaders/Private/Deformers/DeformerGroomStrandsRead.ush
--- a/Engine/Plugins/Runtime/HairStrands/Shaders/Private/Deformers/DeformerGroomStrandsRead.ush
+++ b/Engine/Plugins/Runtime/HairStrands/Shaders/Private/Deformers/DeformerGroomStrandsRead.ush
@@ -6,6 +6,12 @@
 
 HAIR_STRANDS_INSTANCE_PARAMETERS({DataInterfaceName})
 
+uint {DataInterfaceName}_BasePointIndex;
+uint {DataInterfaceName}_BaseCurveIndex;
+uint {DataInterfaceName}_BaseGuidePointIndex;
+uint {DataInterfaceName}_TotalPointCount;
+uint {DataInterfaceName}_TotalCurveCount;
+
 #define HAIR_STRANDS_ATTRIBUTE_ACCESSORS(Name) {DataInterfaceName}_##Name
 #include "/Engine/Private/HairStrands/HairStrandsAttributeTemplate.ush"
 
@@ -77,81 +83,95 @@ FHairCurve ReadCurveData(uint CurveIndex)
 
 uint ReadNumPoints_{DataInterfaceName}()
 {
-	return {DataInterfaceName}_PointCount;
+	return {DataInterfaceName}_TotalPointCount;
 }
 
 uint ReadNumCurves_{DataInterfaceName}()
 {
-	return {DataInterfaceName}_CurveCount;
+	return {DataInterfaceName}_TotalCurveCount;
 }
 
 float3 ReadPointRestPosition_{DataInterfaceName}(uint PointIndex)
 {
+	PointIndex -= {DataInterfaceName}_BasePointIndex;
 	return ReadPointData(PointIndex).Position;
 }
 
 float ReadPointRadius_{DataInterfaceName}(uint PointIndex)
 {
+	PointIndex -= {DataInterfaceName}_BasePointIndex;
 	return ReadPointData(PointIndex).WorldRadius;
 }
 
 float ReadPointCoordU_{DataInterfaceName}(uint PointIndex)
 {
+	PointIndex -= {DataInterfaceName}_BasePointIndex;
 	return ReadPointData(PointIndex).UCoord;
 }
 
 float ReadPointLength_{DataInterfaceName}(uint PointIndex)
 {
+	PointIndex -= {DataInterfaceName}_BasePointIndex;
 	return GetHairStrandsDimensions(PointIndex, 0).x;
 }
 
 float2 ReadPointRootUV_{DataInterfaceName}(uint PointIndex)
 {
+	PointIndex -= {DataInterfaceName}_BasePointIndex;
 	return GetHairStrandsRootUV(PointIndex);
 }
 
 float ReadPointSeed_{DataInterfaceName}(uint PointIndex)
 {
+	PointIndex -= {DataInterfaceName}_BasePointIndex;
 	return GetHairStrandsSeed(PointIndex);
 }
 
 uint ReadPointClumpId_{DataInterfaceName}(uint PointIndex)
 {
+	PointIndex -= {DataInterfaceName}_BasePointIndex;
 	return GetHairStrandsClumpID(PointIndex).x;
 }
 
 float3 ReadPointColor_{DataInterfaceName}(uint PointIndex)
 {
+	PointIndex -= {DataInterfaceName}_BasePointIndex;
 	return GetHairStrandsColor(PointIndex, 0);
 }
 
 float ReadPointRoughness_{DataInterfaceName}(uint PointIndex)
 {
+	PointIndex -= {DataInterfaceName}_BasePointIndex;
 	return GetHairStrandsRoughness(PointIndex, 0);
 }
 
 float ReadPointAO_{DataInterfaceName}(uint PointIndex)
 {
+	PointIndex -= {DataInterfaceName}_BasePointIndex;
 	return GetHairStrandsAO(PointIndex, 0);
 }
 
 uint ReadCurvePointOffset_{DataInterfaceName}(uint CurveIndex)
 {
-	return ReadCurveData(CurveIndex).PointIndex;
+	CurveIndex -= {DataInterfaceName}_BaseCurveIndex;
+	return ReadCurveData(CurveIndex).PointIndex + {DataInterfaceName}_BasePointIndex;
 }
 
 uint ReadCurveNumPoints_{DataInterfaceName}(uint CurveIndex)
 {
+	CurveIndex -= {DataInterfaceName}_BaseCurveIndex;
 	return ReadCurveData(CurveIndex).PointCount;
 }
 
 uint ReadPointCurveIndex_{DataInterfaceName}(uint PointIndex)
 {
-	return ReadHairPointToCurveIndex({DataInterfaceName}_PointToCurveBuffer, PointIndex);;
+	PointIndex -= {DataInterfaceName}_BasePointIndex;
+	return ReadHairPointToCurveIndex({DataInterfaceName}_PointToCurveBuffer, PointIndex) + {DataInterfaceName}_BaseCurveIndex;
 }
 
 int2 ReadPointGuideIndices_{DataInterfaceName}(uint PointIndex)
 {
+	PointIndex -= {DataInterfaceName}_BasePointIndex;
 	// Find the rendering curve index
 	const uint CurveIndex = ReadHairPointToCurveIndex({DataInterfaceName}_PointToCurveBuffer, PointIndex);
 
@@ -161,11 +181,12 @@ int2 ReadPointGuideIndices_{DataInterfaceName}(uint PointIndex)
 
 	// Build the sim point index
 	const uint2 GuidePointIndices = CurveData.RootPointIndices + PointData.PointIndices;
-	return int2(GuidePointIndices.x, GuidePointIndices.y);
+	return int2(GuidePointIndices.x + {DataInterfaceName}_BaseGuidePointIndex, GuidePointIndices.y + {DataInterfaceName}_BaseGuidePointIndex);
 }
 
 float2 ReadPointGuideWeights_{DataInterfaceName}(uint PointIndex)
 {
+	PointIndex -= {DataInterfaceName}_BasePointIndex;
 	// Find the rendering curve index
 	const uint CurveIndex = ReadHairPointToCurveIndex({DataInterfaceName}_PointToCurveBuffer, PointIndex);
 
@@ -178,6 +199,7 @@ float2 ReadPointGuideWeights_{DataInterfaceName}(uint PointIndex)
 
 float3x4 ReadCurveRestTransform_{DataInterfaceName}(uint CurveIndex)
 {
+	CurveIndex -= {DataInterfaceName}_BaseCurveIndex;
 #if ENABLE_SKINNED_TRANSFORM == 1
 		return BuildSkinnedTransform(CurveIndex, {DataInterfaceName}_TriangleRestPositions,
 			{DataInterfaceName}_CurveBarycentricCoordinates, {DataInterfaceName}_CurveTriangleIndices);
@@ -190,6 +212,7 @@ float3x4 ReadCurveRestTransform_{DataInterfaceName}(uint CurveIndex)
 
 float3x4 ReadCurveDeformedTransform_{DataInterfaceName}(uint CurveIndex)
 {
+	CurveIndex -= {DataInterfaceName}_BaseCurveIndex;
 #if ENABLE_SKINNED_TRANSFORM == 1
 		return BuildSkinnedTransform(CurveIndex, {DataInterfaceName}_TriangleDeformedPositions,
 			{DataInterfaceName}_CurveBarycentricCoordinates, {DataInterfaceName}_CurveTriangleIndices);
diff --git a/Engine/Plugins/Runtime/HairStrands/Shaders/Private/Deformers/DeformerGroomStrandsWrite.ush b/Engine/Plugins/Runtime/HairStrands/Shaders/Private/Deformers/DeformerGroomStrandsWrite.ush
--- a/Engine/Plugins/Runtime/HairStrands/Shaders/Private/Deformers/DeformerGroomStrandsWrite.ush
+++ b/Engine/Plugins/Runtime/HairStrands/Shaders/Private/Deformers/DeformerGroomStrandsWrite.ush
@@ -5,9 +5,13 @@
 
 HAIR_STRANDS_INSTANCE_PARAMETERS({DataInterfaceName})
 
+uint {DataInterfaceName}_BasePointIndex;
+uint {DataInterfaceName}_BaseCurveIndex;
+uint {DataInterfaceName}_TotalPointCount;
+uint {DataInterfaceName}_TotalCurveCount;
+
 // Note: resources are explicitly bound here (instance of using HAIR_STRANDS_INSTANCE_PARAMETERS's resources) 
 // because we write out data which would be bound as SRV otherwise (position, attributes, ...)
-uint 				{DataInterfaceName}_OutputStreamStart;
 StructuredBuffer<float4> {DataInterfaceName}_DeformedPositionOffset;
 ByteAddressBuffer		{DataInterfaceName}_PointRestPositions;
 RWByteAddressBuffer		{DataInterfaceName}_PointDeformedPositions;
@@ -62,17 +66,19 @@ float3 {DataInterfaceName}_GetDeformedPositionOffset()
 
 uint ReadNumPoints_{DataInterfaceName}()
 {
-	return {DataInterfaceName}_PointCount;
+	return {DataInterfaceName}_TotalPointCount;
 }
 
 uint ReadNumCurves_{DataInterfaceName}()
 {
-	return {DataInterfaceName}_CurveCount;
+	return {DataInterfaceName}_TotalCurveCount;
 }
 
 void WritePointDeformedPosition_{DataInterfaceName}(uint PointIndex, float3 Position)
 {
-	const uint WriteIndex = {DataInterfaceName}_OutputStreamStart + PointIndex;
+	PointIndex -= {DataInterfaceName}_BasePointIndex;
+	
+	const uint WriteIndex = PointIndex;
 	if (WriteIndex < {DataInterfaceName}_PointCount)
 	{
 		const FPackedHairPosition PackedCP = ReadPackedHairPosition({DataInterfaceName}_PointRestPositions, WriteIndex);
@@ -82,7 +88,9 @@ void WritePointDeformedPosition_{DataInterfaceName}(uint PointIndex, float3 Posi
 
 void WritePointRadius_{DataInterfaceName}(uint PointIndex, float Radius)
 {
-	const uint WriteIndex = {DataInterfaceName}_OutputStreamStart + PointIndex;
+	PointIndex -= {DataInterfaceName}_BasePointIndex;
+	
+	const uint WriteIndex = PointIndex;
 	if (WriteIndex < {DataInterfaceName}_PointCount)
 	{
 		const FPackedHairPosition PackedCP = ReadPackedHairPosition({DataInterfaceName}_PointRestPositions, WriteIndex);
@@ -93,7 +101,9 @@ void WritePointRadius_{DataInterfaceName}(uint PointIndex, float Radius)
 
 void WritePointPositionRadius_{DataInterfaceName}(uint PointIndex, float4 PositionAndRadius)
 {
-	const uint WriteIndex = {DataInterfaceName}_OutputStreamStart + PointIndex;
+	PointIndex -= {DataInterfaceName}_BasePointIndex;
+	
+	const uint WriteIndex = PointIndex;
 	if (WriteIndex < {DataInterfaceName}_PointCount)
 	{
 		const FPackedHairPosition PackedCP = ReadPackedHairPosition({DataInterfaceName}_PointRestPositions, WriteIndex);
@@ -106,7 +116,9 @@ void WritePointPositionRadius_{DataInterfaceName}(uint PointIndex, float4 Positi
 
 void WriteCurveRootUV_{DataInterfaceName}(uint CurveIndex, float2 RootUV)
 {
-	const uint WriteIndex = {DataInterfaceName}_OutputStreamStart + CurveIndex;
+	CurveIndex -= {DataInterfaceName}_BaseCurveIndex;
+	
+	const uint WriteIndex = CurveIndex;
 	if (WriteIndex < {DataInterfaceName}_CurveCount && HasHairStrandsRootUV())
 	{
 		InternalSetHairStrandsAttributes32(CurveIndex, HAIR_CURVE_ATTRIBUTE_OFFSET_ROOTUV({DataInterfaceName}_CurveAttributeOffsets), HAIR_CURVE_ATTRIBUTE_STRIDE_ROOTUV, {DataInterfaceName}_PointAttributeBuffer, PackHairRootUV(RootUV), {DataInterfaceName}_CurveAttributeIndexToChunkDivAsShift, {DataInterfaceName}_CurveAttributeChunkElementCount), {DataInterfaceName}_CurveAttributeChunkStrideInBytes);
@@ -115,7 +127,9 @@ void WriteCurveRootUV_{DataInterfaceName}(uint CurveIndex, float2 RootUV)
 
 void WriteCurveSeed_{DataInterfaceName}(uint CurveIndex, float Seed)
 {
-	const uint WriteIndex = {DataInterfaceName}_OutputStreamStart + CurveIndex;
+	CurveIndex -= {DataInterfaceName}_BaseCurveIndex;
+	
+	const uint WriteIndex = CurveIndex;
 	if (WriteIndex < {DataInterfaceName}_CurveCount)
 	{
 		InternalSetHairStrandsAttributes(CurveIndex, HAIR_CURVE_ATTRIBUTE_OFFSET_SEED({DataInterfaceName}_CurveAttributeOffsets), HAIR_CURVE_ATTRIBUTE_STRIDE_SEED, {DataInterfaceName}_PointAttributeBuffer, PackHairSeed(Seed), {DataInterfaceName}_CurveAttributeIndexToChunkDivAsShift, {DataInterfaceName}_CurveAttributeChunkElementCount), {DataInterfaceName}_CurveAttributeChunkStrideInBytes);
@@ -124,7 +138,9 @@ void WriteCurveSeed_{DataInterfaceName}(uint CurveIndex, float Seed)
 
 void WriteCurveClumpId_{DataInterfaceName}(uint CurveIndex, uint ClumpId)
 {
-	const uint WriteIndex = {DataInterfaceName}_OutputStreamStart + CurveIndex;
+	CurveIndex -= {DataInterfaceName}_BaseCurveIndex;
+	
+	const uint WriteIndex = CurveIndex;
 	if (WriteIndex < {DataInterfaceName}_CurveCount && HasHairStrandsClumpId())
 	{
 		InternalSetHairStrandsAttributes(CurveIndex, HAIR_CURVE_ATTRIBUTE_OFFSET_CLUMPID({DataInterfaceName}_CurveAttributeOffsets), HAIR_CURVE_ATTRIBUTE_STRIDE_CLUMPID, {DataInterfaceName}_PointAttributeBuffer, PackHairClumpID(ClumpId), {DataInterfaceName}_CurveAttributeIndexToChunkDivAsShift, {DataInterfaceName}_CurveAttributeChunkElementCount), {DataInterfaceName}_CurveAttributeChunkStrideInBytes);
@@ -133,7 +149,9 @@ void WriteCurveClumpId_{DataInterfaceName}(uint CurveIndex, uint ClumpId)
 
 void WritePointColor_{DataInterfaceName}(uint PointIndex, float3 Color)
 {
-	const uint WriteIndex = {DataInterfaceName}_OutputStreamStart + PointIndex;
+	PointIndex -= {DataInterfaceName}_BasePointIndex;
+	
+	const uint WriteIndex = PointIndex;
 	if (WriteIndex < {DataInterfaceName}_PointCount && HasHairStrandsColor())
 	{
 		InternalSetHairStrandsAttributes32(PointIndex, HAIR_POINT_ATTRIBUTE_OFFSET_COLOR({DataInterfaceName}_PointAttributeOffsets), HAIR_POINT_ATTRIBUTE_STRIDE_COLOR, {DataInterfaceName}_PointAttributeBuffer, PackHairColor(Color), {DataInterfaceName}_PointAttributeIndexToChunkDivAsShift, {DataInterfaceName}_PointAttributeChunkElementCount), {DataInterfaceName}_PointAttributeChunkStrideInBytes);
@@ -142,7 +160,9 @@ void WritePointColor_{DataInterfaceName}(uint PointIndex, float3 Color)
 
 void WritePointRoughness_{DataInterfaceName}(uint PointIndex, float Roughness)
 {
-	const uint WriteIndex = {DataInterfaceName}_OutputStreamStart + PointIndex;
+	PointIndex -= {DataInterfaceName}_BasePointIndex;
+	
+	const uint WriteIndex = PointIndex;
 	if (WriteIndex < {DataInterfaceName}_PointCount && HasHairStrandsRoughness())
 	{
 		InternalSetHairStrandsAttributes(PointIndex, HAIR_POINT_ATTRIBUTE_OFFSET_ROUGHNESS({DataInterfaceName}_PointAttributeOffsets), HAIR_POINT_ATTRIBUTE_STRIDE_ROUGHNESS, {DataInterfaceName}_PointAttributeBuffer, PackHairRoughness(Roughness), {DataInterfaceName}_PointAttributeIndexToChunkDivAsShift, {DataInterfaceName}_PointAttributeChunkElementCount), {DataInterfaceName}_PointAttributeChunkStrideInBytes);
@@ -151,7 +171,9 @@ void WritePointRoughness_{DataInterfaceName}(uint PointIndex, float Roughness)
 
 void WritePointAO_{DataInterfaceName}(uint PointIndex, float Ao)
 {
-	const uint WriteIndex = {DataInterfaceName}_OutputStreamStart + PointIndex;
+	PointIndex -= {DataInterfaceName}_BasePointIndex;
+	
+	const uint WriteIndex = PointIndex;
 	if (WriteIndex < {DataInterfaceName}_PointCount && HasHairStrandsAO())
 	{
 		InternalSetHairStrandsAttributes(PointIndex, HAIR_POINT_ATTRIBUTE_OFFSET_AO({DataInterfaceName}_PointAttributeOffsets), HAIR_POINT_ATTRIBUTE_STRIDE_AO, {DataInterfaceName}_PointAttributeBuffer, PackHairAO(Ao), {DataInterfaceName}_PointAttributeIndexToChunkDivAsShift, {DataInterfaceName}_PointAttributeChunkElementCount), {DataInterfaceName}_PointAttributeChunkStrideInBytes);
diff --git a/Engine/Plugins/Runtime/HairStrands/Source/HairStrandsDeformer/Private/DeformerGroomGuidesRead.cpp b/Engine/Plugins/Runtime/HairStrands/Source/HairStrandsDeformer/Private/DeformerGroomGuidesRead.cpp
--- a/Engine/Plugins/Runtime/HairStrands/Source/HairStrandsDeformer/Private/DeformerGroomGuidesRead.cpp
+++ b/Engine/Plugins/Runtime/HairStrands/Source/HairStrandsDeformer/Private/DeformerGroomGuidesRead.cpp
@@ -113,6 +113,10 @@ void UOptimusGroomGuidesReadDataInterface::GetSupportedInputs(TArray<FShaderFunc
 BEGIN_SHADER_PARAMETER_STRUCT(FOptimusGroomGuidesReadParameters, )
 	SHADER_PARAMETER(uint32, PointCount)
 	SHADER_PARAMETER(uint32, CurveCount)
+	SHADER_PARAMETER(uint32, BasePointIndex)
+	SHADER_PARAMETER(uint32, BaseCurveIndex)
+	SHADER_PARAMETER(uint32, TotalPointCount)
+	SHADER_PARAMETER(uint32, TotalCurveCount)
 	SHADER_PARAMETER(uint32, ObjectIndex)
 	SHADER_PARAMETER(FVector3f, RestPositionOffset)
 	SHADER_PARAMETER_RDG_BUFFER_SRV(ByteAddressBuffer, PointRestPositions)
@@ -262,7 +266,9 @@ void FOptimusGroomGuidesReadDataProviderProxy::GatherDispatchData(FDispatchData
 {
 	check(InDispatchData.NumInvocations == GroupInstances.Num());
 	const TStridedView<FParameters> ParameterArray = MakeStridedParameterView<FParameters>(InDispatchData);
-	
+
+	uint32 BasePointIndex = 0;
+	uint32 BaseCurveIndex = 0;
 	for (int32 InvocationIndex = 0; InvocationIndex < ParameterArray.Num(); ++InvocationIndex)
 	{
 		if (const FHairGroupInstance* GroupInstance = GroupInstances[InvocationIndex])
@@ -271,6 +277,13 @@ void FOptimusGroomGuidesReadDataProviderProxy::GatherDispatchData(FDispatchData
 			
 			Parameters.PointCount = GroupInstance->Guides.RestResource->GetPointCount();
 			Parameters.CurveCount = GroupInstance->Guides.RestResource->GetCurveCount();
+			
+			// Used to get the local element indices for the current group, since index supplied by the compute kernel goes from 0 to NumElementsPerGroup * NumGroups
+			Parameters.BasePointIndex = BasePointIndex;
+			BasePointIndex += Parameters.PointCount;
+			Parameters.BaseCurveIndex = BaseCurveIndex;
+			BaseCurveIndex += Parameters.CurveCount;
+			
 			Parameters.ObjectIndex = InvocationIndex;
 			Parameters.RestPositionOffset = FVector3f(GroupInstance->Guides.RestResource->GetPositionOffset());
 
@@ -297,4 +310,14 @@ void FOptimusGroomGuidesReadDataProviderProxy::GatherDispatchData(FDispatchData
 			}
 		}
 	}
+
+	for (int32 InvocationIndex = 0; InvocationIndex < ParameterArray.Num(); ++InvocationIndex)
+	{
+		if (const FHairGroupInstance* GroupInstance = GroupInstances[InvocationIndex])
+		{
+			FParameters& Parameters = ParameterArray[InvocationIndex];
+			Parameters.TotalPointCount = BasePointIndex;
+			Parameters.TotalCurveCount = BaseCurveIndex;
+		}
+	}
 }
\ No newline at end of file
diff --git a/Engine/Plugins/Runtime/HairStrands/Source/HairStrandsDeformer/Private/DeformerGroomGuidesWrite.cpp b/Engine/Plugins/Runtime/HairStrands/Source/HairStrandsDeformer/Private/DeformerGroomGuidesWrite.cpp
--- a/Engine/Plugins/Runtime/HairStrands/Source/HairStrandsDeformer/Private/DeformerGroomGuidesWrite.cpp
+++ b/Engine/Plugins/Runtime/HairStrands/Source/HairStrandsDeformer/Private/DeformerGroomGuidesWrite.cpp
@@ -58,7 +58,8 @@ void UOptimusGroomGuidesWriteDataInterface::GetSupportedOutputs(TArray<FShaderFu
 BEGIN_SHADER_PARAMETER_STRUCT(FOptimusGroomGuidesWriteParameters, )
 	SHADER_PARAMETER(uint32, RegisteredIndex)
 	SHADER_PARAMETER(uint32, PointCount)
-	SHADER_PARAMETER(uint32, OutputStreamStart)
+	SHADER_PARAMETER(uint32, TotalPointCount)
+	SHADER_PARAMETER(uint32, BasePointIndex)
 	SHADER_PARAMETER_RDG_BUFFER_SRV(StructuredBuffer, DeformedPositionOffset)
 	SHADER_PARAMETER_RDG_BUFFER_SRV(ByteAddressBuffer, PointRestPositions)
 	SHADER_PARAMETER_RDG_BUFFER_UAV(RWByteAddressBuffer, PointDeformedPositions)
@@ -150,6 +151,8 @@ void FOptimusGroomGuidesWriteProviderProxy::AllocateResources(FRDGBuilder& Graph
 void FOptimusGroomGuidesWriteProviderProxy::GatherDispatchData(FDispatchData const& InDispatchData)
 {
 	const TStridedView<FParameters> ParameterArray = MakeStridedParameterView<FParameters>(InDispatchData);
+
+	uint32 BasePointIndex = 0;
 	for (int32 InvocationIndex = 0; InvocationIndex < ParameterArray.Num(); ++InvocationIndex)
 	{
 		if (const FHairGroupInstance* GroupInstance = GroupInstances[InvocationIndex])
@@ -158,7 +161,8 @@ void FOptimusGroomGuidesWriteProviderProxy::GatherDispatchData(FDispatchData con
 			FInstanceResources& Resource = InstanceResources[InvocationIndex];
 			
 			Parameters.PointCount = GroupInstance->Guides.RestResource->GetPointCount();
-			Parameters.OutputStreamStart = 0;
+			Parameters.BasePointIndex = BasePointIndex;
+			BasePointIndex += Parameters.PointCount;
 			Parameters.RegisteredIndex = GroupInstance->RegisteredIndex;
 			
 			Parameters.DeformedPositionOffset 	=  Resource.DeformedPositionOffset;
@@ -166,4 +170,13 @@ void FOptimusGroomGuidesWriteProviderProxy::GatherDispatchData(FDispatchData con
 			Parameters.PointDeformedPositions = Resource.PointDeformedPositions;
 		}
 	}
+	
+	for (int32 InvocationIndex = 0; InvocationIndex < ParameterArray.Num(); ++InvocationIndex)
+	{
+		if (const FHairGroupInstance* GroupInstance = GroupInstances[InvocationIndex])
+		{
+			FParameters& Parameters = ParameterArray[InvocationIndex];
+			Parameters.TotalPointCount = BasePointIndex;
+		}
+	}
 }
\ No newline at end of file
diff --git a/Engine/Plugins/Runtime/HairStrands/Source/HairStrandsDeformer/Private/DeformerGroomStrandsRead.cpp b/Engine/Plugins/Runtime/HairStrands/Source/HairStrandsDeformer/Private/DeformerGroomStrandsRead.cpp
--- a/Engine/Plugins/Runtime/HairStrands/Source/HairStrandsDeformer/Private/DeformerGroomStrandsRead.cpp
+++ b/Engine/Plugins/Runtime/HairStrands/Source/HairStrandsDeformer/Private/DeformerGroomStrandsRead.cpp
@@ -161,6 +161,11 @@ void UOptimusGroomStrandsReadDataInterface::GetSupportedInputs(TArray<FShaderFun
 
 BEGIN_SHADER_PARAMETER_STRUCT(FOptimusGroomStrandsReadParameters, )
 	SHADER_PARAMETER_STRUCT_INCLUDE(FHairStrandsInstanceCommonParameters, Common)
+	SHADER_PARAMETER(uint32, BasePointIndex)
+	SHADER_PARAMETER(uint32, BaseCurveIndex)
+	SHADER_PARAMETER(uint32, BaseGuidePointIndex)
+	SHADER_PARAMETER(uint32, TotalPointCount)
+	SHADER_PARAMETER(uint32, TotalCurveCount)
 	SHADER_PARAMETER_STRUCT_INCLUDE(FHairStrandsInstanceResourceParameters, Resources)
 	SHADER_PARAMETER_STRUCT_INCLUDE(FHairStrandsInstanceInterpolationParameters, Interpolation)
 	SHADER_PARAMETER(FMatrix44f, RigidRestTransform)
@@ -321,6 +326,9 @@ void FOptimusGroomStrandsReadProviderProxy::GatherDispatchData(FDispatchData con
 	check(InDispatchData.NumInvocations == GroupInstances.Num());
 	const TStridedView<FParameters> ParameterArray = MakeStridedParameterView<FParameters>(InDispatchData);
 
+	uint32 BasePointIndex = 0;
+	uint32 BaseCurveIndex = 0;
+	uint32 BaseGuidePointIndex = 0;
 	for (int32 InvocationIndex = 0; InvocationIndex < ParameterArray.Num(); ++InvocationIndex)
 	{
 		if (const FHairGroupInstance* GroupInstance = GroupInstances[InvocationIndex])
@@ -329,6 +337,16 @@ void FOptimusGroomStrandsReadProviderProxy::GatherDispatchData(FDispatchData con
 			
 			const FHairGroupPublicData::FVertexFactoryInput VFInput = ComputeHairStrandsVertexInputData(GroupInstance, EGroomViewMode::None);
 			Parameters.Common = VFInput.Strands.Common;
+
+			// Used to get the local element indices for the current group, since index supplied by the compute kernel goes from 0 to NumElementsPerGroup * NumGroups
+			Parameters.BasePointIndex = BasePointIndex;
+			BasePointIndex += Parameters.Common.PointCount;
+			Parameters.BaseCurveIndex = BaseCurveIndex;
+			BaseCurveIndex += Parameters.Common.CurveCount;
+
+			Parameters.BaseGuidePointIndex = BaseGuidePointIndex;
+			BaseGuidePointIndex += GroupInstance->Guides.RestResource->GetPointCount();
+			
 			Parameters.Resources = GroupResources[InvocationIndex];
 			
 			if (GroupInterpolations[InvocationIndex].CurveInterpolationBuffer != nullptr)
@@ -355,4 +373,14 @@ void FOptimusGroomStrandsReadProviderProxy::GatherDispatchData(FDispatchData con
 			}
 		}
 	}
+	
+	for (int32 InvocationIndex = 0; InvocationIndex < ParameterArray.Num(); ++InvocationIndex)
+	{
+		if (const FHairGroupInstance* GroupInstance = GroupInstances[InvocationIndex])
+		{
+			FParameters& Parameters = ParameterArray[InvocationIndex];
+			Parameters.TotalPointCount = BasePointIndex;	
+			Parameters.TotalCurveCount = BaseCurveIndex;	
+		}
+	}
 }
\ No newline at end of file
diff --git a/Engine/Plugins/Runtime/HairStrands/Source/HairStrandsDeformer/Private/DeformerGroomStrandsWrite.cpp b/Engine/Plugins/Runtime/HairStrands/Source/HairStrandsDeformer/Private/DeformerGroomStrandsWrite.cpp
--- a/Engine/Plugins/Runtime/HairStrands/Source/HairStrandsDeformer/Private/DeformerGroomStrandsWrite.cpp
+++ b/Engine/Plugins/Runtime/HairStrands/Source/HairStrandsDeformer/Private/DeformerGroomStrandsWrite.cpp
@@ -112,7 +112,10 @@ void UOptimusGroomStrandsWriteDataInterface::GetSupportedOutputs(TArray<FShaderF
 
 BEGIN_SHADER_PARAMETER_STRUCT(FOptimusGroomStrandsWriteParameters, )
 	SHADER_PARAMETER_STRUCT_INCLUDE(FHairStrandsInstanceCommonParameters, Common)
-	SHADER_PARAMETER(uint32, OutputStreamStart)
+	SHADER_PARAMETER(uint32, BasePointIndex)
+	SHADER_PARAMETER(uint32, BaseCurveIndex)
+	SHADER_PARAMETER(uint32, TotalPointCount)
+	SHADER_PARAMETER(uint32, TotalCurveCount)
 	SHADER_PARAMETER_RDG_BUFFER_SRV(StructuredBuffer, DeformedPositionOffset)
 	SHADER_PARAMETER_RDG_BUFFER_SRV(ByteAddressBuffer, PointRestPositions)
 	SHADER_PARAMETER_RDG_BUFFER_UAV(RWByteAddressBuffer, PointDeformedPositions)
@@ -250,6 +253,9 @@ void FOptimusGroomStrandsWriteProviderProxy::AllocateResources(FRDGBuilder& Grap
 void FOptimusGroomStrandsWriteProviderProxy::GatherDispatchData(FDispatchData const& InDispatchData)
 {
 	const TStridedView<FParameters> ParameterArray = MakeStridedParameterView<FParameters>(InDispatchData);
+	
+	uint32 BasePointIndex = 0;
+	uint32 BaseCurveIndex = 0;
 	for (int32 InvocationIndex = 0; InvocationIndex < ParameterArray.Num(); ++InvocationIndex)
 	{
 		if (const FHairGroupInstance* GroupInstance = GroupInstances[InvocationIndex])
@@ -260,8 +266,15 @@ void FOptimusGroomStrandsWriteProviderProxy::GatherDispatchData(FDispatchData co
 			const FHairGroupPublicData::FVertexFactoryInput VFInput = ComputeHairStrandsVertexInputData(GroupInstance, EGroomViewMode::None);		
 
 			FParameters& Parameters = ParameterArray[InvocationIndex];
+			
 			Parameters.Common 					= VFInput.Strands.Common;
-			Parameters.OutputStreamStart 		= 0;
+			
+			// Used to get the local element indices for the current group, since index supplied by the compute kernel goes from 0 to NumElementsPerGroup * NumGroups
+			Parameters.BasePointIndex = BasePointIndex;
+			BasePointIndex += Parameters.Common.PointCount;
+			Parameters.BaseCurveIndex = BaseCurveIndex;
+			BaseCurveIndex += Parameters.Common.CurveCount;
+			
 			Parameters.DeformedPositionOffset 	= bValid ? Resource.DeformedPositionOffset : FallbackPositionBufferSRV;
 			Parameters.PointRestPositions 		= bValid ? Resource.PointRestPositions : FallbackPositionBufferSRV;
 			Parameters.PointDeformedPositions 	= bValid ? Resource.PointDeformedPositions : FallbackPositionBufferUAV;
@@ -269,4 +282,14 @@ void FOptimusGroomStrandsWriteProviderProxy::GatherDispatchData(FDispatchData co
 			Parameters.CurveAttributeBuffer 	= bValid ? Resource.CurveAttributeBuffer : FallbackAttributeBufferUAV;
 		}
 	}
+
+	for (int32 InvocationIndex = 0; InvocationIndex < ParameterArray.Num(); ++InvocationIndex)
+	{
+		if (const FHairGroupInstance* GroupInstance = GroupInstances[InvocationIndex])
+		{
+			FParameters& Parameters = ParameterArray[InvocationIndex];
+			Parameters.TotalPointCount = BasePointIndex;	
+			Parameters.TotalCurveCount = BaseCurveIndex;	
+		}
+	}	
 }
\ No newline at end of file
