[Deformer Graph] temp disable compile on load, causing cook crash

diff --git a/Engine/Plugins/Animation/DeformerGraph/Source/OptimusCore/Private/OptimusDeformer.cpp b/Engine/Plugins/Animation/DeformerGraph/Source/OptimusCore/Private/OptimusDeformer.cpp
--- a/Engine/Plugins/Animation/DeformerGraph/Source/OptimusCore/Private/OptimusDeformer.cpp
+++ b/Engine/Plugins/Animation/DeformerGraph/Source/OptimusCore/Private/OptimusDeformer.cpp
@@ -3434,7 +3434,7 @@ void UOptimusDeformer::PostLoad()
 	if (Status != EOptimusDeformerStatus::HasErrors &&
 		ObjectVersion < FOptimusObjectVersion::LatestVersion)
 	{
-		Compile();
+		// Compile();
 	}
 }
 
