[Deformer Graph & ComputeFramework] Fixed crash/error when compiling old deformer graph during cook

-Changed to use a soft object ptr to reference external nodes for DG compilation error debug to avoid direct dependency
-only apply kernel index check for non unified dispatch, this is to allow for group sync primitives to be used in the unified dispatch case
-avoid creating compute graph render proxy when rendering is not allowed

diff --git a/Engine/Plugins/Animation/DeformerGraph/Shaders/Private/DataInterfaceCopyKernel.ush b/Engine/Plugins/Animation/DeformerGraph/Shaders/Private/DataInterfaceCopyKernel.ush
--- a/Engine/Plugins/Animation/DeformerGraph/Shaders/Private/DataInterfaceCopyKernel.ush
+++ b/Engine/Plugins/Animation/DeformerGraph/Shaders/Private/DataInterfaceCopyKernel.ush
@@ -1,20 +1,20 @@
 // Copyright Epic Games, Inc. All Rights Reserved.
 
 int3 {DataInterfaceName}_NumThreads;
-int {DataInterfaceName}_NumThreadsPerInvocation;
-int {DataInterfaceName}_ThreadIndexOffset;
+uint {DataInterfaceName}_NumThreadsPerInvocation;
+uint {DataInterfaceName}_ThreadIndexOffset;
 
 int3 ReadNumThreads_{DataInterfaceName}()
 {
 	return {DataInterfaceName}_NumThreads;
 }
 
-int ReadNumThreadsPerInvocation_{DataInterfaceName}()
+uint ReadNumThreadsPerInvocation_{DataInterfaceName}()
 {
 	return {DataInterfaceName}_NumThreadsPerInvocation;
 }
 
-int ReadThreadIndexOffset_{DataInterfaceName}()
+uint ReadThreadIndexOffset_{DataInterfaceName}()
 {
 	return {DataInterfaceName}_ThreadIndexOffset;
 }
\ No newline at end of file
diff --git a/Engine/Plugins/Animation/DeformerGraph/Shaders/Private/DataInterfaceCustomComputeKernel.ush b/Engine/Plugins/Animation/DeformerGraph/Shaders/Private/DataInterfaceCustomComputeKernel.ush
--- a/Engine/Plugins/Animation/DeformerGraph/Shaders/Private/DataInterfaceCustomComputeKernel.ush
+++ b/Engine/Plugins/Animation/DeformerGraph/Shaders/Private/DataInterfaceCustomComputeKernel.ush
@@ -1,20 +1,20 @@
 // Copyright Epic Games, Inc. All Rights Reserved.
 
 int3 {DataInterfaceName}_NumThreads;
-int {DataInterfaceName}_NumThreadsPerInvocation;
-int {DataInterfaceName}_ThreadIndexOffset;
+uint {DataInterfaceName}_NumThreadsPerInvocation;
+uint {DataInterfaceName}_ThreadIndexOffset;
 
 int3 ReadNumThreads_{DataInterfaceName}()
 {
 	return {DataInterfaceName}_NumThreads;
 }
 
-int ReadNumThreadsPerInvocation_{DataInterfaceName}()
+uint ReadNumThreadsPerInvocation_{DataInterfaceName}()
 {
 	return {DataInterfaceName}_NumThreadsPerInvocation;
 }
 
-int ReadThreadIndexOffset_{DataInterfaceName}()
+uint ReadThreadIndexOffset_{DataInterfaceName}()
 {
 	return {DataInterfaceName}_ThreadIndexOffset;
 }
\ No newline at end of file
diff --git a/Engine/Plugins/Animation/DeformerGraph/Source/OptimusCore/Private/DataInterfaces/OptimusDataInterfaceCopyKernel.cpp b/Engine/Plugins/Animation/DeformerGraph/Source/OptimusCore/Private/DataInterfaces/OptimusDataInterfaceCopyKernel.cpp
--- a/Engine/Plugins/Animation/DeformerGraph/Source/OptimusCore/Private/DataInterfaces/OptimusDataInterfaceCopyKernel.cpp
+++ b/Engine/Plugins/Animation/DeformerGraph/Source/OptimusCore/Private/DataInterfaces/OptimusDataInterfaceCopyKernel.cpp
@@ -29,17 +29,17 @@ void UOptimusCopyKernelDataInterface::GetSupportedInputs(TArray<FShaderFunctionD
 
 	OutFunctions.AddDefaulted_GetRef()
 		.SetName(ReadNumThreadsPerInvocationFunctionName)
-		.AddReturnType(FShaderValueType::Get(EShaderFundamentalType::Int));
+		.AddReturnType(FShaderValueType::Get(EShaderFundamentalType::Uint));
 	
 	OutFunctions.AddDefaulted_GetRef()
 		.SetName(ReadThreadIndexOffsetFunctionName)
-		.AddReturnType(FShaderValueType::Get(EShaderFundamentalType::Int));
+		.AddReturnType(FShaderValueType::Get(EShaderFundamentalType::Uint));
 }
 
 BEGIN_SHADER_PARAMETER_STRUCT(FCopyKernelDataInterfaceParameters, )
 	SHADER_PARAMETER(FIntVector3, NumThreads)
-	SHADER_PARAMETER(int32, NumThreadsPerInvocation)
-	SHADER_PARAMETER(int32, ThreadIndexOffset)
+	SHADER_PARAMETER(uint32, NumThreadsPerInvocation)
+	SHADER_PARAMETER(uint32, ThreadIndexOffset)
 END_SHADER_PARAMETER_STRUCT()
 
 
@@ -93,6 +93,11 @@ void UOptimusCopyKernelDataInterface::SetComponentBinding(const UOptimusComponen
 	ComponentSourceBinding = InBinding;
 }
 
+const FString& UOptimusCopyKernelDataInterface::GetExecutionDomain() const
+{
+	return NumThreadsExpression;
+}
+
 const TCHAR* UOptimusCopyKernelDataInterface::GetReadNumThreadsFunctionName() const
 {
 	return ReadNumThreadsFunctionName;
diff --git a/Engine/Plugins/Animation/DeformerGraph/Source/OptimusCore/Private/DataInterfaces/OptimusDataInterfaceCopyKernel.h b/Engine/Plugins/Animation/DeformerGraph/Source/OptimusCore/Private/DataInterfaces/OptimusDataInterfaceCopyKernel.h
--- a/Engine/Plugins/Animation/DeformerGraph/Source/OptimusCore/Private/DataInterfaces/OptimusDataInterfaceCopyKernel.h
+++ b/Engine/Plugins/Animation/DeformerGraph/Source/OptimusCore/Private/DataInterfaces/OptimusDataInterfaceCopyKernel.h
@@ -38,6 +38,7 @@ public:
 	//~ Begin IOptimusComputeKernelDataInterface Interface
 	void SetExecutionDomain(const FString& InExecutionDomain) override;
 	void SetComponentBinding(const UOptimusComponentSourceBinding* InBinding) override;
+	const FString& GetExecutionDomain() const override;
 	const TCHAR* GetReadNumThreadsFunctionName() const override;
 	const TCHAR* GetReadNumThreadsPerInvocationFunctionName() const override;
 	const TCHAR* GetReadThreadIndexOffsetFunctionName() const override;
diff --git a/Engine/Plugins/Animation/DeformerGraph/Source/OptimusCore/Private/DataInterfaces/OptimusDataInterfaceCustomComputeKernel.cpp b/Engine/Plugins/Animation/DeformerGraph/Source/OptimusCore/Private/DataInterfaces/OptimusDataInterfaceCustomComputeKernel.cpp
--- a/Engine/Plugins/Animation/DeformerGraph/Source/OptimusCore/Private/DataInterfaces/OptimusDataInterfaceCustomComputeKernel.cpp
+++ b/Engine/Plugins/Animation/DeformerGraph/Source/OptimusCore/Private/DataInterfaces/OptimusDataInterfaceCustomComputeKernel.cpp
@@ -30,17 +30,17 @@ void UOptimusCustomComputeKernelDataInterface::GetSupportedInputs(TArray<FShader
 	
 	OutFunctions.AddDefaulted_GetRef()
 		.SetName(ReadNumThreadsPerInvocationFunctionName)
-		.AddReturnType(FShaderValueType::Get(EShaderFundamentalType::Int));
+		.AddReturnType(FShaderValueType::Get(EShaderFundamentalType::Uint));
 	
 	OutFunctions.AddDefaulted_GetRef()
 		.SetName(ReadThreadIndexOffsetFunctionName)
-		.AddReturnType(FShaderValueType::Get(EShaderFundamentalType::Int));
+		.AddReturnType(FShaderValueType::Get(EShaderFundamentalType::Uint));
 }
 
 BEGIN_SHADER_PARAMETER_STRUCT(FCustomComputeKernelDataInterfaceParameters, )
 	SHADER_PARAMETER(FIntVector3, NumThreads)
-	SHADER_PARAMETER(int32, NumThreadsPerInvocation)
-	SHADER_PARAMETER(int32, ThreadIndexOffset)
+	SHADER_PARAMETER(uint32, NumThreadsPerInvocation)
+	SHADER_PARAMETER(uint32, ThreadIndexOffset)
 END_SHADER_PARAMETER_STRUCT()
 
 void UOptimusCustomComputeKernelDataInterface::GetShaderParameters(TCHAR const* UID,
@@ -93,6 +93,11 @@ void UOptimusCustomComputeKernelDataInterface::SetComponentBinding(const UOptimu
 	ComponentSourceBinding = InBinding;
 }
 
+const FString& UOptimusCustomComputeKernelDataInterface::GetExecutionDomain() const
+{
+	return NumThreadsExpression;
+}
+
 const TCHAR* UOptimusCustomComputeKernelDataInterface::GetReadNumThreadsFunctionName() const
 {
 	return ReadNumThreadsFunctionName;
diff --git a/Engine/Plugins/Animation/DeformerGraph/Source/OptimusCore/Private/DataInterfaces/OptimusDataInterfaceCustomComputeKernel.h b/Engine/Plugins/Animation/DeformerGraph/Source/OptimusCore/Private/DataInterfaces/OptimusDataInterfaceCustomComputeKernel.h
--- a/Engine/Plugins/Animation/DeformerGraph/Source/OptimusCore/Private/DataInterfaces/OptimusDataInterfaceCustomComputeKernel.h
+++ b/Engine/Plugins/Animation/DeformerGraph/Source/OptimusCore/Private/DataInterfaces/OptimusDataInterfaceCustomComputeKernel.h
@@ -41,6 +41,7 @@ public:
 	//~ Begin IOptimusComputeKernelDataInterface Interface
 	void SetExecutionDomain(const FString& InExecutionDomain) override;
 	void SetComponentBinding(const UOptimusComponentSourceBinding* InBinding) override;
+	const FString& GetExecutionDomain() const override;
 	const TCHAR* GetReadNumThreadsFunctionName() const override;
 	const TCHAR* GetReadNumThreadsPerInvocationFunctionName() const override;
 	const TCHAR* GetReadThreadIndexOffsetFunctionName() const override;
diff --git a/Engine/Plugins/Animation/DeformerGraph/Source/OptimusCore/Private/Nodes/OptimusNode_ComputeKernelBase.cpp b/Engine/Plugins/Animation/DeformerGraph/Source/OptimusCore/Private/Nodes/OptimusNode_ComputeKernelBase.cpp
--- a/Engine/Plugins/Animation/DeformerGraph/Source/OptimusCore/Private/Nodes/OptimusNode_ComputeKernelBase.cpp
+++ b/Engine/Plugins/Animation/DeformerGraph/Source/OptimusCore/Private/Nodes/OptimusNode_ComputeKernelBase.cpp
@@ -145,26 +145,6 @@ FOptimus_ComputeKernelResult UOptimusNode_ComputeKernelBase::CreateComputeKernel
 	}
 	
 	UOptimusComponentSourceBinding* PrimaryBinding = *(PrimaryBindings.CreateIterator());
-	bool bUnifiedDispatchOnly = true;
-	{
-		using namespace Optimus::Expression;
-		TVariant<FExpressionObject, FParseError> ParseResult = Optimus::ParseExecutionDomainExpression(GetExecutionDomain().AsExpression(), PrimaryBinding->GetComponentSource());
-		if (ParseResult.IsType<FParseError>())
-		{
-			return ReturnError(
-				FText::Format(
-					LOCTEXT("ExecutionDomainParsingError", "Error while parsing Execution Domain: {0}"),
-					FText::FromString(ParseResult.Get<FParseError>().Message)));
-		}
-
-		FExpressionObject ExpressionObject = ParseResult.Get<FExpressionObject>();
-
-		// Only fixed domain supports both unified and non-unified dispatch
-		if (TOptional<FName> FixedDomain = ExpressionObject.GetAsSingleConstant())
-		{
-			bUnifiedDispatchOnly = false;
-		}
-	}
 	
 	UOptimusKernelSource* KernelSource = NewObject<UOptimusKernelSource>(InKernelSourceOuter);
 
@@ -172,6 +152,28 @@ FOptimus_ComputeKernelResult UOptimusNode_ComputeKernelBase::CreateComputeKernel
 	// (for reads) or do nothing (for writes).
 	TArray<FString> GeneratedFunctions;
 
+	bool bSupportUnifiedDispatchOnly = false;
+	
+	check(InOutKernelDataInterface)
+	{
+		TOptional<FText> Result = BindKernelDataInterfaceForComputeKernel(
+			PrimaryBinding,
+			InOutKernelDataInterface,
+			KernelSource,
+			OutInputDataBindings,
+			GeneratedFunctions,
+			bSupportUnifiedDispatchOnly
+			);
+	
+		if (Result.IsSet())
+		{
+			return ReturnError(Result.GetValue());	
+		}	
+	}
+	
+
+	bool bIsUnifiedDispatch = true;
+	
 	int32 GroupIndex = 0;
 	for (const UOptimusNodePin* Pin: GetPins())
 	{
@@ -180,10 +182,10 @@ FOptimus_ComputeKernelResult UOptimusNode_ComputeKernelBase::CreateComputeKernel
 			for (const UOptimusNodePin* SubPin: Pin->GetSubPins())
 			{
 				TOptional<FText> Result = ProcessInputPinForComputeKernel(
-					bUnifiedDispatchOnly,
+					bSupportUnifiedDispatchOnly,
 					InTraversalContext, SubPin,
 					GroupIndex == 0 ? FString() : Pin->GetName(),
-					InKernelInputs, KernelSource, GeneratedFunctions, OutInputDataBindings
+					InKernelInputs, KernelSource,  GeneratedFunctions, OutInputDataBindings, bIsUnifiedDispatch
 					);
 				if (Result.IsSet())
 				{
@@ -196,10 +198,10 @@ FOptimus_ComputeKernelResult UOptimusNode_ComputeKernelBase::CreateComputeKernel
 		else if (Pin->GetDirection() == EOptimusNodePinDirection::Output)
 		{
 			TOptional<FText> Result = ProcessOutputPinForComputeKernel(
-				bUnifiedDispatchOnly,
+				bSupportUnifiedDispatchOnly,
 				InTraversalContext, Pin,
 				InKernelOutputs,
-				KernelSource, GeneratedFunctions, OutInputDataBindings, OutOutputDataBindings
+				KernelSource, GeneratedFunctions, OutInputDataBindings, OutOutputDataBindings, bIsUnifiedDispatch
 				);
 				
 			if (Result.IsSet())
@@ -209,21 +211,12 @@ FOptimus_ComputeKernelResult UOptimusNode_ComputeKernelBase::CreateComputeKernel
 		}
 	}
 
-	if (ensure(InOutKernelDataInterface))
-	{
-		BindKernelDataInterfaceForComputeKernel(
-			PrimaryBindings.Array()[0],
-			InOutKernelDataInterface,
-			KernelSource,
-			OutInputDataBindings,
-			GeneratedFunctions
-			);
-	}
+
 
 	FString CookedSource;
 	CookedSource += FString::Join(GeneratedFunctions, TEXT("\n"));
 	CookedSource += "\n\n";
-	CookedSource += GetKernelSourceText();
+	CookedSource += GetKernelSourceText(bIsUnifiedDispatch);
 	
 	KernelSource->SetSource(CookedSource);
 	KernelSource->EntryPoint = GetKernelHlslName();
@@ -400,14 +393,15 @@ const TCHAR* UOptimusNode_ComputeKernelBase::GetNumThreadsReservedName()
 
 
 TOptional<FText> UOptimusNode_ComputeKernelBase::ProcessInputPinForComputeKernel(
-	bool bUnifiedDispatchOnly,
+	bool bSupportUnifiedDispatchOnly,
 	const FOptimusPinTraversalContext& InTraversalContext,
 	const UOptimusNodePin* InInputPin,
 	const FString& InGroupName,
 	const FOptimus_KernelInputMap& InKernelInputs,
 	UOptimusKernelSource* InKernelSource,
 	TArray<FString>& OutGeneratedFunctions,
-	FOptimus_InterfaceBindingMap& OutInputDataBindings
+	FOptimus_InterfaceBindingMap& OutInputDataBindings,
+	bool& bOutIsUnifiedDispatch
 	) const
 {
 	
@@ -429,10 +423,12 @@ TOptional<FText> UOptimusNode_ComputeKernelBase::ProcessInputPinForComputeKernel
 				return FText::Format(LOCTEXT("SecondaryGroupOnlySupportsDataInterfaceUnifiedDispatch", "Cannot connect Secondary group input binding pin '{0}' to a data interface that does not support unified dispatch, consider inserting a kernel inbetween"), FText::FromName(InInputPin->GetUniqueName()));
 			}
 
-			if (bUnifiedDispatchOnly)
+			if (bSupportUnifiedDispatchOnly)
 			{
 				return FText::Format(LOCTEXT("InvalidInput_KernelOnlySupportsDataInterfaceUnifiedDispatch", "Cannot connect input binding pin '{0}' to a data interface that does not support unified dispatch, consider using a simple kernel execution domain or inserting a kernel inbetween"), FText::FromName(InInputPin->GetUniqueName()));
 			}
+
+			bOutIsUnifiedDispatch = false;
 		}
 
 		int32 DataInterfaceFuncIndex = INDEX_NONE;
@@ -549,14 +545,15 @@ TOptional<FText> UOptimusNode_ComputeKernelBase::ProcessInputPinForComputeKernel
 
 
 TOptional<FText> UOptimusNode_ComputeKernelBase::ProcessOutputPinForComputeKernel(
-	bool bUnifiedDispatchOnly,
+	bool bSupportUnifiedDispatchOnly,
 	const FOptimusPinTraversalContext& InTraversalContext,
 	const UOptimusNodePin* InOutputPin,
 	const FOptimus_KernelOutputMap& InKernelOutputs,
 	UOptimusKernelSource* InKernelSource,
 	TArray<FString>& OutGeneratedFunctions,
 	FOptimus_InterfaceBindingMap& OutInputDataBindings,
-	FOptimus_InterfaceBindingMap& OutOutputDataBindings
+	FOptimus_InterfaceBindingMap& OutOutputDataBindings,
+	bool& bOutIsUnifiedDispatch
 	) const
 {
 	FShaderParamTypeDefinition IndexParamDef;
@@ -593,10 +590,12 @@ TOptional<FText> UOptimusNode_ComputeKernelBase::ProcessOutputPinForComputeKerne
 
 			if (!DataInterface->CanSupportUnifiedDispatch())
 			{
-				if (bUnifiedDispatchOnly)
+				if (bSupportUnifiedDispatchOnly)
 				{
 					return FText::Format(LOCTEXT("InvalidOutput_KernelOnlySupportsDataInterfaceUnifiedDispatch", "Cannot connect output binding pin '{0}' to a data interface that does not support unified dispatch, consider using to a simple kernel execution domain or inserting a kernel inbetween"), FText::FromName(InOutputPin->GetUniqueName()));
 				}
+
+				bOutIsUnifiedDispatch = false;
 			}
 			
 			if (!ConnectedNode)
@@ -759,12 +758,13 @@ TOptional<FText> UOptimusNode_ComputeKernelBase::ProcessOutputPinForComputeKerne
 	return {};
 }
 
-void UOptimusNode_ComputeKernelBase::BindKernelDataInterfaceForComputeKernel(
+TOptional<FText> UOptimusNode_ComputeKernelBase::BindKernelDataInterfaceForComputeKernel(
 	const UOptimusComponentSourceBinding* InKernelPrimaryComponentSourceBinding,
 	UComputeDataInterface* InOutKernelDataInterface,
 	UOptimusKernelSource* InKernelSource,
 	FOptimus_InterfaceBindingMap& OutInputDataBindings,
-	TArray<FString>& OutGeneratedFunctions
+	TArray<FString>& OutGeneratedFunctions,
+	bool& bOutUnifiedDispatchOnly
 	) const
 {
 	check(InOutKernelDataInterface);
@@ -772,10 +772,19 @@ void UOptimusNode_ComputeKernelBase::BindKernelDataInterfaceForComputeKernel(
 	
 	IOptimusComputeKernelDataInterface* KernelDataInterface = Cast<IOptimusComputeKernelDataInterface>(InOutKernelDataInterface);
 	check(KernelDataInterface);
-	
-	KernelDataInterface->SetExecutionDomain(GetExecutionDomain().AsExpression());
+	FString NumThreadsExpression = GetExecutionDomain().AsExpression();
+	KernelDataInterface->SetExecutionDomain(NumThreadsExpression);
 	KernelDataInterface->SetComponentBinding(InKernelPrimaryComponentSourceBinding);
 	
+	TVariant<bool, FText> UnifiedDispatchResult = Optimus::IsExecutionDomainUnifiedDispatchOnly(NumThreadsExpression,InKernelPrimaryComponentSourceBinding->GetComponentSource());
+
+	if (UnifiedDispatchResult.IsType<FText>())
+	{
+		return UnifiedDispatchResult.Get<FText>();
+	}
+
+	bOutUnifiedDispatchOnly = UnifiedDispatchResult.Get<bool>();
+	
 	TArray<FShaderFunctionDefinition> ReadFunctions;
 	InOutKernelDataInterface->GetSupportedInputs(ReadFunctions);	
 
@@ -804,6 +813,8 @@ void UOptimusNode_ComputeKernelBase::BindKernelDataInterfaceForComputeKernel(
 	OutGeneratedFunctions.Add(FString::Printf(
 			TEXT("int3 Read%s() { return %s::%s(); }"),
 			GetNumThreadsReservedName(), Optimus::GetKernelInternalNamespaceName(), KernelDataInterface->GetReadNumThreadsFunctionName() ));
+
+	return {};
 }
 
 #undef LOCTEXT_NAMESPACE
diff --git a/Engine/Plugins/Animation/DeformerGraph/Source/OptimusCore/Private/Nodes/OptimusNode_ComputeKernelBase.h b/Engine/Plugins/Animation/DeformerGraph/Source/OptimusCore/Private/Nodes/OptimusNode_ComputeKernelBase.h
--- a/Engine/Plugins/Animation/DeformerGraph/Source/OptimusCore/Private/Nodes/OptimusNode_ComputeKernelBase.h
+++ b/Engine/Plugins/Animation/DeformerGraph/Source/OptimusCore/Private/Nodes/OptimusNode_ComputeKernelBase.h
@@ -30,7 +30,7 @@ public:
 	virtual FIntVector GetGroupSize() const PURE_VIRTUAL(UOptimusNode_ComputeKernelBase::GetGroupSize, return FIntVector();)
 
 	/** Implement this to return the complete HLSL code for this kernel */
-	virtual FString GetKernelSourceText() const PURE_VIRTUAL(UOptimusNode_ComputeKernelBase::GetKernelSourceText, return FString();)
+	virtual FString GetKernelSourceText(bool bInIsUnifiedDispatch) const PURE_VIRTUAL(UOptimusNode_ComputeKernelBase::GetKernelSourceText, return FString();)
 
 	/** Implement this to return additional source code that the kernel requires. */
 	virtual TArray<TObjectPtr<UComputeSource>> GetAdditionalSources() const { return {}; }
@@ -97,34 +97,38 @@ protected:
 	static const TCHAR* GetNumThreadsReservedName();
 	
 private:
+	TOptional<FText> BindKernelDataInterfaceForComputeKernel(
+		const UOptimusComponentSourceBinding* InKernelPrimaryComponentSourceBinding,
+		UComputeDataInterface* InOutKernelDataInterface,
+		UOptimusKernelSource* InKernelSource,
+		FOptimus_InterfaceBindingMap& OutInputDataBindings,
+		TArray<FString>& OutGeneratedFunctions,
+		bool& bOutUnifiedDispatchOnly
+		) const;
+	
 	TOptional<FText> ProcessInputPinForComputeKernel(
-		bool bUnifiedDispatchOnly,
+		bool bSupportUnifiedDispatchOnly,
 		const FOptimusPinTraversalContext& InTraversalContext,
 		const UOptimusNodePin* InInputPin,
 		const FString& InGroupName,
 		const FOptimus_KernelInputMap& InKernelInputs,
 		UOptimusKernelSource* InKernelSource,
 		TArray<FString>& OutGeneratedFunctions,
-		FOptimus_InterfaceBindingMap& OutInputDataBindings
+		FOptimus_InterfaceBindingMap& OutInputDataBindings,
+		bool& bOutIsUnifiedDispatch
 		) const;
 
 	TOptional<FText> ProcessOutputPinForComputeKernel(
-		bool bUnifiedDispatchOnly,
+	bool bSupportUnifiedDispatchOnly,
 		const FOptimusPinTraversalContext& InTraversalContext,
 		const UOptimusNodePin* InOutputPin,
 		const FOptimus_KernelOutputMap& InKernelOutputs,
 		UOptimusKernelSource* InKernelSource,
 		TArray<FString>& OutGeneratedFunctions,
 		FOptimus_InterfaceBindingMap& OutInputDataBindings,
-		FOptimus_InterfaceBindingMap& OutOutputDataBindings
+		FOptimus_InterfaceBindingMap& OutOutputDataBindings,
+		bool& bOutIsUnifiedDispatch
 		) const;
 
-	void BindKernelDataInterfaceForComputeKernel(
-		const UOptimusComponentSourceBinding* InKernelPrimaryComponentSourceBinding,
-		UComputeDataInterface* InOutKernelDataInterface,
-		UOptimusKernelSource* InKernelSource,
-		FOptimus_InterfaceBindingMap& OutInputDataBindings,
-		TArray<FString>& OutGeneratedFunctions
-		) const;
 };
 
diff --git a/Engine/Plugins/Animation/DeformerGraph/Source/OptimusCore/Private/Nodes/OptimusNode_ComputeKernelFunction.cpp b/Engine/Plugins/Animation/DeformerGraph/Source/OptimusCore/Private/Nodes/OptimusNode_ComputeKernelFunction.cpp
--- a/Engine/Plugins/Animation/DeformerGraph/Source/OptimusCore/Private/Nodes/OptimusNode_ComputeKernelFunction.cpp
+++ b/Engine/Plugins/Animation/DeformerGraph/Source/OptimusCore/Private/Nodes/OptimusNode_ComputeKernelFunction.cpp
@@ -179,7 +179,7 @@ FIntVector UOptimusNode_ComputeKernelFunction::GetGroupSize() const
 }
 
 
-FString UOptimusNode_ComputeKernelFunction::GetKernelSourceText() const
+FString UOptimusNode_ComputeKernelFunction::GetKernelSourceText(bool bInIsUnifiedDispatch) const
 {
 	checkNoEntry();
 	return {};
diff --git a/Engine/Plugins/Animation/DeformerGraph/Source/OptimusCore/Private/Nodes/OptimusNode_ComputeKernelFunction.h b/Engine/Plugins/Animation/DeformerGraph/Source/OptimusCore/Private/Nodes/OptimusNode_ComputeKernelFunction.h
--- a/Engine/Plugins/Animation/DeformerGraph/Source/OptimusCore/Private/Nodes/OptimusNode_ComputeKernelFunction.h
+++ b/Engine/Plugins/Animation/DeformerGraph/Source/OptimusCore/Private/Nodes/OptimusNode_ComputeKernelFunction.h
@@ -71,7 +71,7 @@ public:
 	// UOptimusNode_ComputeKernelBase overrides
 	FString GetKernelHlslName() const override;
 	FIntVector GetGroupSize() const override;
-	FString GetKernelSourceText() const override;
+	FString GetKernelSourceText(bool bInIsUnifiedDispatch) const override;
 
 	void ConstructNode() override;
 	
diff --git a/Engine/Plugins/Animation/DeformerGraph/Source/OptimusCore/Private/Nodes/OptimusNode_CustomComputeKernel.cpp b/Engine/Plugins/Animation/DeformerGraph/Source/OptimusCore/Private/Nodes/OptimusNode_CustomComputeKernel.cpp
--- a/Engine/Plugins/Animation/DeformerGraph/Source/OptimusCore/Private/Nodes/OptimusNode_CustomComputeKernel.cpp
+++ b/Engine/Plugins/Animation/DeformerGraph/Source/OptimusCore/Private/Nodes/OptimusNode_CustomComputeKernel.cpp
@@ -100,7 +100,7 @@ UOptimusNode_CustomComputeKernel::UOptimusNode_CustomComputeKernel()
 }
 
 
-FString UOptimusNode_CustomComputeKernel::GetKernelSourceText() const
+FString UOptimusNode_CustomComputeKernel::GetKernelSourceText(bool bInIsUnifiedDispatch) const
 {
 	using UKernelDI = UOptimusCustomComputeKernelDataInterface;
 	
@@ -110,7 +110,8 @@ FString UOptimusNode_CustomComputeKernel::GetKernelSourceText() const
 		GetKernelHlslName(),
 		GroupSize,
 		UKernelDI::StaticClass()->GetDefaultObject<UKernelDI>()->GetReadNumThreadsPerInvocationFunctionName(),
-		UKernelDI::StaticClass()->GetDefaultObject<UKernelDI>()->GetReadThreadIndexOffsetFunctionName()
+		UKernelDI::StaticClass()->GetDefaultObject<UKernelDI>()->GetReadThreadIndexOffsetFunctionName(),
+		bInIsUnifiedDispatch
 		);
 }
 
diff --git a/Engine/Plugins/Animation/DeformerGraph/Source/OptimusCore/Private/Nodes/OptimusNode_CustomComputeKernel.h b/Engine/Plugins/Animation/DeformerGraph/Source/OptimusCore/Private/Nodes/OptimusNode_CustomComputeKernel.h
--- a/Engine/Plugins/Animation/DeformerGraph/Source/OptimusCore/Private/Nodes/OptimusNode_CustomComputeKernel.h
+++ b/Engine/Plugins/Animation/DeformerGraph/Source/OptimusCore/Private/Nodes/OptimusNode_CustomComputeKernel.h
@@ -56,7 +56,7 @@ public:
 	// UOptimusNode_ComputeKernelBase overrides
 	FString GetKernelHlslName() const override;
 	FIntVector GetGroupSize() const override { return GroupSize; }
-	FString GetKernelSourceText() const override;
+	FString GetKernelSourceText(bool bInIsUnifiedDispatch) const override;
 	TArray<TObjectPtr<UComputeSource>> GetAdditionalSources() const override { return AdditionalSources; }
 
 	// IOptimusComputeKernelProvider
diff --git a/Engine/Plugins/Animation/DeformerGraph/Source/OptimusCore/Private/OptimusComputeGraph.cpp b/Engine/Plugins/Animation/DeformerGraph/Source/OptimusCore/Private/OptimusComputeGraph.cpp
--- a/Engine/Plugins/Animation/DeformerGraph/Source/OptimusCore/Private/OptimusComputeGraph.cpp
+++ b/Engine/Plugins/Animation/DeformerGraph/Source/OptimusCore/Private/OptimusComputeGraph.cpp
@@ -117,7 +117,7 @@ void UOptimusComputeGraph::OnKernelCompilationComplete(int32 InKernelIndex, FCom
 		UOptimusDeformer* Owner = Cast<UOptimusDeformer>(GetOuter());
 		
 		// Make sure the node hasn't been GC'd.
-		if (UOptimusNode* Node = const_cast<UOptimusNode*>(KernelToNode[InKernelIndex].Get()))
+		if (UOptimusNode* Node = const_cast<UOptimusNode*>(KernelToNode[InKernelIndex].LoadSynchronous()))
 		{
 			FString FriendlyName = Owner->GetName() / GetFName().GetPlainNameString() / Node->GetDisplayName().ToString();
 			EOptimusDiagnosticLevel DiagnosticLevel = EOptimusDiagnosticLevel::None;
diff --git a/Engine/Plugins/Animation/DeformerGraph/Source/OptimusCore/Private/OptimusComputeGraph.h b/Engine/Plugins/Animation/DeformerGraph/Source/OptimusCore/Private/OptimusComputeGraph.h
--- a/Engine/Plugins/Animation/DeformerGraph/Source/OptimusCore/Private/OptimusComputeGraph.h
+++ b/Engine/Plugins/Animation/DeformerGraph/Source/OptimusCore/Private/OptimusComputeGraph.h
@@ -28,9 +28,9 @@ public:
 
 	EMeshDeformerOutputBuffer GetOutputBuffers() const;
 protected:
-	// Lookup into Graphs array from the UComputeGraph kernel index. 
+	// Lookup into Graphs array from the UComputeGraph kernel index.
 	UPROPERTY()
-	TArray<TWeakObjectPtr<const UOptimusNode>> KernelToNode;
+	TArray<TSoftObjectPtr<const UOptimusNode>> KernelToNode;
 
 	UOptimusGraphDataInterface* GetGraphDataInterfaceForPostLoadFixUp();
 	bool HasHalfEdgeDataInterface() const;
diff --git a/Engine/Plugins/Animation/DeformerGraph/Source/OptimusCore/Private/OptimusDeformer.cpp b/Engine/Plugins/Animation/DeformerGraph/Source/OptimusCore/Private/OptimusDeformer.cpp
--- a/Engine/Plugins/Animation/DeformerGraph/Source/OptimusCore/Private/OptimusDeformer.cpp
+++ b/Engine/Plugins/Animation/DeformerGraph/Source/OptimusCore/Private/OptimusDeformer.cpp
@@ -2823,6 +2823,7 @@ FOptimusNodeGraphCompilationResult UOptimusDeformer::CompileNodeGraphToComputeGr
 		{
 			TArray<FKernelWithDataBindings> BoundCopyKernels;
 			TArray<UComputeKernel*> InsertAfterComputeKernelLookUpArray;
+			bool bIsUnifiedDispatch = true;
 			for (TPair<FOptimusInstancedPin, TArray<FOptimusInstancedPin>> OutputLink : LinksToInsertCopyKernel)
 			{
 				// Create a copy kernel per source pin, that copies from 1 source pin to multiple target pins
@@ -2846,9 +2847,15 @@ FOptimusNodeGraphCompilationResult UOptimusDeformer::CompileNodeGraphToComputeGr
 				FOptimus_InterfaceBindingMap& OutputDataBindings = BoundCopyKernel.OutputDataBindings;
 
 				UOptimusKernelSource* KernelSource = NewObject<UOptimusKernelSource>(BoundCopyKernel.Kernel);
+				UComputeDataInterface* CopyKernelDataInterface = CopyKernelDataInterfaceMap[SourceInstancedPin];
+				
 				FString SourceText;
+				SourceText = FString::Printf(
+					TEXT("if (Index >= %s::%s().x) return;\n"),
+					Optimus::GetKernelInternalNamespaceName(),
+					CastChecked<IOptimusComputeKernelDataInterface>(CopyKernelDataInterface)->GetReadNumThreadsFunctionName());
+				
 				{
-					UComputeDataInterface* CopyKernelDataInterface = CopyKernelDataInterfaceMap[SourceInstancedPin];
 					TArray<FShaderFunctionDefinition> Functions;
 					CopyKernelDataInterface->GetSupportedInputs(Functions);
 					// Simply grab everything the kernel data interface has to offer
@@ -2876,7 +2883,12 @@ FOptimusNodeGraphCompilationResult UOptimusDeformer::CompileNodeGraphToComputeGr
 
 				{
 					const FDataInterfaceFunctionBinding& CopyFromBinding = CopyFromDataInterfaceMap[SourceInstancedPin];
-				
+
+					if (!CopyFromBinding.DataInterface->CanSupportUnifiedDispatch())
+					{
+						bIsUnifiedDispatch = false;
+					}
+					
 					TArray<FShaderFunctionDefinition> Functions;
 					CopyFromBinding.DataInterface->GetSupportedInputs(Functions);
 					FShaderFunctionDefinition FuncDef = Functions[CopyFromBinding.FunctionIndex];
@@ -2908,6 +2920,11 @@ FOptimusNodeGraphCompilationResult UOptimusDeformer::CompileNodeGraphToComputeGr
 				{
 					const FOptimusInstancedPin& TargetInstancedPin = OutputLink.Value[TargetIndex];
 					const FDataInterfaceFunctionBinding& CopyToBinding = CopyToDataInterfaceMap[TargetInstancedPin];
+
+					if (!CopyToBinding.DataInterface->CanSupportUnifiedDispatch())
+					{
+						bIsUnifiedDispatch = false;
+					}
 					
 					TArray<FShaderFunctionDefinition> Functions;
 					CopyToBinding.DataInterface->GetSupportedOutputs(Functions);
@@ -2940,7 +2957,8 @@ FOptimusNodeGraphCompilationResult UOptimusDeformer::CompileNodeGraphToComputeGr
 					CopyKernelName,
 					GroupSize,
 					UKernelDI::StaticClass()->GetDefaultObject<UKernelDI>()->GetReadNumThreadsPerInvocationFunctionName(),
-					UKernelDI::StaticClass()->GetDefaultObject<UKernelDI>()->GetReadThreadIndexOffsetFunctionName()	
+					UKernelDI::StaticClass()->GetDefaultObject<UKernelDI>()->GetReadThreadIndexOffsetFunctionName(),
+					bIsUnifiedDispatch
 					);
 				KernelSource->SetSource(CookedSource);
 				KernelSource->EntryPoint = CopyKernelName;
@@ -3434,7 +3452,7 @@ void UOptimusDeformer::PostLoad()
 	if (Status != EOptimusDeformerStatus::HasErrors &&
 		ObjectVersion < FOptimusObjectVersion::LatestVersion)
 	{
-		// Compile();
+		Compile();
 	}
 }
 
diff --git a/Engine/Plugins/Animation/DeformerGraph/Source/OptimusCore/Private/OptimusHelpers.cpp b/Engine/Plugins/Animation/DeformerGraph/Source/OptimusCore/Private/OptimusHelpers.cpp
--- a/Engine/Plugins/Animation/DeformerGraph/Source/OptimusCore/Private/OptimusHelpers.cpp
+++ b/Engine/Plugins/Animation/DeformerGraph/Source/OptimusCore/Private/OptimusHelpers.cpp
@@ -11,6 +11,8 @@
 #include "ShaderParameterMetadata.h"
 #include "StructUtils/UserDefinedStruct.h"
 
+#define LOCTEXT_NAMESPACE "OptimusHelpers"
+
 FName Optimus::GetUniqueNameForScope(UObject* InScopeObj, FName InName)
 {
 	// If there's already an object with this name, then attempt to make the name unique.
@@ -226,7 +228,8 @@ FString Optimus::GetCookedKernelSource(
 	const FString& InKernelName,
 	FIntVector InGroupSize,
 	const TCHAR* InReadNumThreadsPerInvocationFunctionName,
-	const TCHAR* InReadThreadIndexOffsetFunctionName
+	const TCHAR* InReadThreadIndexOffsetFunctionName,
+	bool bInIsUnifiedDispatch
 	)
 {
 	// FIXME: Create source range mappings so that we can go from error location to
@@ -256,17 +259,25 @@ FString Optimus::GetCookedKernelSource(
 
 	const FString ThreadIndexOffsetForKernel = FString::Printf(TEXT(""));
 
-	const FString IndexCheckAndApplyOffset= FString::Printf(
-		TEXT(
-		"uint IndexForInvocation = %s;\n"
-		"if (IndexForInvocation >= %s::%s()) return;\n"
-		"uint Index = IndexForInvocation + %s::%s();\n"),
-		*TheadIndexForInvocation,
-		GetKernelInternalNamespaceName(),
-		InReadNumThreadsPerInvocationFunctionName,
-		GetKernelInternalNamespaceName(),
-		InReadThreadIndexOffsetFunctionName
-		);
+	// Avoid thread index check for unified dispatch to allow for the usage of Group Sync primitives
+	// Shader compiler does not like early return when those primitives are used
+	const FString IndexCheckAndApplyOffset= bInIsUnifiedDispatch ?
+		FString::Printf(
+			TEXT(
+			"uint Index = %s;\n" ),
+			*TheadIndexForInvocation
+			) :
+		FString::Printf(
+			TEXT(
+			"uint IndexForInvocation = %s;\n"
+			"if (IndexForInvocation >= %s::%s()) return;\n"
+			"uint Index = IndexForInvocation + %s::%s();\n"),
+			*TheadIndexForInvocation,
+			GetKernelInternalNamespaceName(),
+			InReadNumThreadsPerInvocationFunctionName,
+			GetKernelInternalNamespaceName(),
+			InReadThreadIndexOffsetFunctionName
+			);
 	
 	if (bHasKernelKeyword)
 	{
@@ -330,6 +341,33 @@ Optimus::Expression::FParseResult Optimus::ParseExecutionDomainExpression(
 	});
 }
 
+TVariant<bool, FText> Optimus::IsExecutionDomainUnifiedDispatchOnly(const FString& InExpression,
+	TWeakObjectPtr<const UOptimusComponentSource> InComponentSource)
+{
+	bool bUnifiedDispatchOnly = true;
+	{
+		using namespace Optimus::Expression;
+		FParseResult ParseResult = ParseExecutionDomainExpression(InExpression, InComponentSource);
+		if (ParseResult.IsType<FParseError>())
+		{
+			return TVariant<bool, FText>(TInPlaceType<FText>(), FText::Format(
+					LOCTEXT("ExecutionDomainParsingError", "Error while parsing Execution Domain: {0}"),
+					FText::FromString(ParseResult.Get<FParseError>().Message)));
+		}
+
+		FExpressionObject ExpressionObject = ParseResult.Get<FExpressionObject>();
+
+		// Only fixed domain supports both unified and non-unified dispatch
+		if (TOptional<FName> FixedDomain = ExpressionObject.GetAsSingleConstant())
+		{
+			bUnifiedDispatchOnly = false;
+		}
+	}
+
+	return TVariant<bool, FText>(TInPlaceType<bool>(), bUnifiedDispatchOnly);
+}
+
+
 bool Optimus::EvaluateExecutionDomainExpressionParseResult(
 	const TVariant<Expression::FExpressionObject, Expression::FParseError>& InParseResult,
 	TWeakObjectPtr<const UOptimusComponentSource> InComponentSource,
@@ -520,3 +558,5 @@ FString Optimus::ExtractSourceValueName(const FString& InUniqueValueName)
 
 	return InUniqueValueName;
 }
+
+#undef LOCTEXT_NAMESPACE
\ No newline at end of file
diff --git a/Engine/Plugins/Animation/DeformerGraph/Source/OptimusCore/Public/IOptimusComputeKernelDataInterface.h b/Engine/Plugins/Animation/DeformerGraph/Source/OptimusCore/Public/IOptimusComputeKernelDataInterface.h
--- a/Engine/Plugins/Animation/DeformerGraph/Source/OptimusCore/Public/IOptimusComputeKernelDataInterface.h
+++ b/Engine/Plugins/Animation/DeformerGraph/Source/OptimusCore/Public/IOptimusComputeKernelDataInterface.h
@@ -22,8 +22,10 @@ class OPTIMUSCORE_API IOptimusComputeKernelDataInterface
 {
 public:
 	GENERATED_BODY()
+	
 	virtual void SetExecutionDomain(const FString& InExecutionDomain) = 0;
 	virtual void SetComponentBinding(const UOptimusComponentSourceBinding* InBinding) = 0;
+	virtual const FString& GetExecutionDomain() const = 0;
 	virtual const TCHAR* GetReadNumThreadsFunctionName() const = 0;
 	virtual const TCHAR* GetReadNumThreadsPerInvocationFunctionName() const = 0;
 	virtual const TCHAR* GetReadThreadIndexOffsetFunctionName() const = 0;
diff --git a/Engine/Plugins/Animation/DeformerGraph/Source/OptimusCore/Public/OptimusHelpers.h b/Engine/Plugins/Animation/DeformerGraph/Source/OptimusCore/Public/OptimusHelpers.h
--- a/Engine/Plugins/Animation/DeformerGraph/Source/OptimusCore/Public/OptimusHelpers.h
+++ b/Engine/Plugins/Animation/DeformerGraph/Source/OptimusCore/Public/OptimusHelpers.h
@@ -134,7 +134,8 @@ namespace Optimus
 		const FString& InKernelName,
 		FIntVector InGroupSize,
 		const TCHAR* InReadNumThreadsPerInvocationFunctionName,
-		const TCHAR* InReadThreadIndexOffsetFunctionName
+		const TCHAR* InReadThreadIndexOffsetFunctionName,
+		bool bInIsUnifiedDispatch
 	);
 
 	const TCHAR* GetKernelInternalNamespaceName();
@@ -144,6 +145,11 @@ namespace Optimus
 		TWeakObjectPtr<const UOptimusComponentSource> InComponentSource
 		);
 
+	TVariant<bool, FText> IsExecutionDomainUnifiedDispatchOnly(
+		const FString& InExpression,
+		TWeakObjectPtr<const UOptimusComponentSource> InComponentSource
+		);
+	
 	bool EvaluateExecutionDomainExpressionParseResult(
 		const TVariant<Expression::FExpressionObject, Expression::FParseError>& InParseResult,
 		TWeakObjectPtr<const UOptimusComponentSource> InComponentSource,
diff --git a/Engine/Plugins/Runtime/ComputeFramework/Source/ComputeFramework/Private/ComputeGraph.cpp b/Engine/Plugins/Runtime/ComputeFramework/Source/ComputeFramework/Private/ComputeGraph.cpp
--- a/Engine/Plugins/Runtime/ComputeFramework/Source/ComputeFramework/Private/ComputeGraph.cpp
+++ b/Engine/Plugins/Runtime/ComputeFramework/Source/ComputeFramework/Private/ComputeGraph.cpp
@@ -346,17 +346,29 @@ enum class EComputeKernelCompilationFlags
 
 void UComputeGraph::UpdateResources()
 {
+	bool bAllowRendering = true;
+	
 #if WITH_EDITOR
 	uint32 CompilationFlags = (uint32)EComputeKernelCompilationFlags::ApplyCompletedShaderMapForRendering;
 
-	bool bIsRenderCommandlet = IsRunningCommandlet() && IsAllowCommandletRendering();
-	CompilationFlags |= bIsRenderCommandlet ? (uint32)EComputeKernelCompilationFlags::Synchronous : (uint32)EComputeKernelCompilationFlags::None;
+	// Rendering is likely not allowed during processes like cooking
+	if (IsRunningCommandlet() && !IsAllowCommandletRendering())
+	{
+		bAllowRendering = false;
+	}
+	
+	CompilationFlags |= bAllowRendering ? (uint32)EComputeKernelCompilationFlags::Synchronous : (uint32)EComputeKernelCompilationFlags::None;
 
 	CacheResourceShadersForRendering(CompilationFlags);
 #endif
 
 	ReleaseRenderProxy(RenderProxy);
-	RenderProxy = CreateRenderProxy();
+	
+	// Rendering is disabled, so no need to create the render proxy
+	if (bAllowRendering)
+	{
+		RenderProxy = CreateRenderProxy();
+	}
 }
 
 bool UComputeGraph::HasKernelResourcesPendingShaderCompilation() const
