Allow a component to be bound to mulitple optimus deformer bindings
#tests Used a custom component to bind to two different sources within the same graph

diff --git a/Engine/Plugins/Animation/DeformerGraph/Source/OptimusCore/Private/OptimusDeformerInstance.cpp b/Engine/Plugins/Animation/DeformerGraph/Source/OptimusCore/Private/OptimusDeformerInstance.cpp
--- a/Engine/Plugins/Animation/DeformerGraph/Source/OptimusCore/Private/OptimusDeformerInstance.cpp
+++ b/Engine/Plugins/Animation/DeformerGraph/Source/OptimusCore/Private/OptimusDeformerInstance.cpp
@@ -274,10 +274,7 @@ void UOptimusDeformerInstanceSettings::GetComponentBindings(
 
 					for (UActorComponent* Component: Components)
 					{
-						if (!ComponentsUsed.Contains(Component))
-						{
-							TaggedComponents.Add(Component);
-						}
+						TaggedComponents.Add(Component);
 					}
 				}
 				TArray<UActorComponent*> RankedTaggedComponents = TaggedComponents.Array();
@@ -302,14 +299,9 @@ void UOptimusDeformerInstanceSettings::GetComponentBindings(
 			{
 				TArray<UActorComponent*> Components;
 				Actor->GetComponents(Binding->GetComponentSource()->GetComponentClass(), Components);
-
-				for (UActorComponent* Component: Components)
+				if (!Components.IsEmpty())
 				{
-					if (!ComponentsUsed.Contains(Component))
-					{
-						BoundComponent = Component;
-						break;
-					}
+					BoundComponent = Components[0];
 				}
 			}
 		}
