[Deformer Graph] temporarily remove an ensure while investigating the root cause

diff --git a/Engine/Plugins/Animation/DeformerGraph/Source/OptimusCore/Private/Nodes/OptimusNode_FunctionReference.cpp b/Engine/Plugins/Animation/DeformerGraph/Source/OptimusCore/Private/Nodes/OptimusNode_FunctionReference.cpp
--- a/Engine/Plugins/Animation/DeformerGraph/Source/OptimusCore/Private/Nodes/OptimusNode_FunctionReference.cpp
+++ b/Engine/Plugins/Animation/DeformerGraph/Source/OptimusCore/Private/Nodes/OptimusNode_FunctionReference.cpp
@@ -78,7 +78,7 @@ FOptimusRoutedNodePin UOptimusNode_FunctionReference::GetPinCounterpart(
 		CounterpartNode = FunctionGraph->GetTerminalNode(EOptimusTerminalType::Return);
 	}
 
-	if (!ensure(CounterpartNode))
+	if (!(CounterpartNode))
 	{
 		return {};
 	}
@@ -160,5 +160,8 @@ void UOptimusNode_FunctionReference::InitializeTransientData()
 {
 	Super::InitializeTransientData();
 	// Making sure the pointer is alive such that we avoid loading on demand everywhere
-	FunctionGraph.LoadSynchronous();
+	if (UOptimusFunctionNodeGraph* LoadedFunctionGraph = FunctionGraph.LoadSynchronous())
+	{
+		LoadedFunctionGraph->ConditionalPostLoad();
+	}
 }
