[Deformer Graph] fix common shader errors/warnings people may run into easily
-read skinned mesh now uses SRV instead UAV to avoid error on dx11 around unsupported typed UAV loads
-changing ReadNumThreads shader API to use uint3 instead of int3 to reduce warning whenever an error is found in the shader

diff --git a/Engine/Plugins/Animation/DeformerGraph/Shaders/Private/DataInterfaceCopyKernel.ush b/Engine/Plugins/Animation/DeformerGraph/Shaders/Private/DataInterfaceCopyKernel.ush
--- a/Engine/Plugins/Animation/DeformerGraph/Shaders/Private/DataInterfaceCopyKernel.ush
+++ b/Engine/Plugins/Animation/DeformerGraph/Shaders/Private/DataInterfaceCopyKernel.ush
@@ -1,10 +1,10 @@
 // Copyright Epic Games, Inc. All Rights Reserved.
 
-int3 {DataInterfaceName}_NumThreads;
+uint3 {DataInterfaceName}_NumThreads;
 uint {DataInterfaceName}_NumThreadsPerInvocation;
 uint {DataInterfaceName}_ThreadIndexOffset;
 
-int3 ReadNumThreads_{DataInterfaceName}()
+uint3 ReadNumThreads_{DataInterfaceName}()
 {
 	return {DataInterfaceName}_NumThreads;
 }
diff --git a/Engine/Plugins/Animation/DeformerGraph/Shaders/Private/DataInterfaceCustomComputeKernel.ush b/Engine/Plugins/Animation/DeformerGraph/Shaders/Private/DataInterfaceCustomComputeKernel.ush
--- a/Engine/Plugins/Animation/DeformerGraph/Shaders/Private/DataInterfaceCustomComputeKernel.ush
+++ b/Engine/Plugins/Animation/DeformerGraph/Shaders/Private/DataInterfaceCustomComputeKernel.ush
@@ -1,10 +1,10 @@
 // Copyright Epic Games, Inc. All Rights Reserved.
 
-int3 {DataInterfaceName}_NumThreads;
+uint3 {DataInterfaceName}_NumThreads;
 uint {DataInterfaceName}_NumThreadsPerInvocation;
 uint {DataInterfaceName}_ThreadIndexOffset;
 
-int3 ReadNumThreads_{DataInterfaceName}()
+uint3 ReadNumThreads_{DataInterfaceName}()
 {
 	return {DataInterfaceName}_NumThreads;
 }
diff --git a/Engine/Plugins/Animation/DeformerGraph/Shaders/Private/DataInterfaceSkinnedMeshRead.ush b/Engine/Plugins/Animation/DeformerGraph/Shaders/Private/DataInterfaceSkinnedMeshRead.ush
--- a/Engine/Plugins/Animation/DeformerGraph/Shaders/Private/DataInterfaceSkinnedMeshRead.ush
+++ b/Engine/Plugins/Animation/DeformerGraph/Shaders/Private/DataInterfaceSkinnedMeshRead.ush
@@ -5,9 +5,9 @@
 
 uint {DataInterfaceName}_NumVertices;
 
-RWBuffer<float> {DataInterfaceName}_PositionBufferUAV;
-RWBuffer<TANGENT_RWBUFFER_FORMAT> {DataInterfaceName}_TangentBufferUAV;
-RWBuffer<UNORM float4> {DataInterfaceName}_ColorBufferUAV;
+Buffer<float> {DataInterfaceName}_PositionBufferSRV;
+Buffer<TANGENT_RWBUFFER_FORMAT> {DataInterfaceName}_TangentBufferSRV;
+Buffer<UNORM float4> {DataInterfaceName}_ColorBufferSRV;
 
 Buffer<float> {DataInterfaceName}_PositionStaticBuffer;
 Buffer<TANGENT_RWBUFFER_FORMAT> {DataInterfaceName}_TangentStaticBuffer;
@@ -25,9 +25,9 @@ float3 ReadPosition_{DataInterfaceName}(uint VertexIndex)
 
 #if READABLE_OUTPUT_BUFFERS & OPTIMUS_SKINNED_MESH_READ_POSITION
 	return float3(
-		{DataInterfaceName}_PositionBufferUAV[BufferIndex + 0], 
-		{DataInterfaceName}_PositionBufferUAV[BufferIndex + 1],
-		{DataInterfaceName}_PositionBufferUAV[BufferIndex + 2]);
+		{DataInterfaceName}_PositionBufferSRV[BufferIndex + 0], 
+		{DataInterfaceName}_PositionBufferSRV[BufferIndex + 1],
+		{DataInterfaceName}_PositionBufferSRV[BufferIndex + 2]);
 #else
 	return float3(
 		{DataInterfaceName}_PositionStaticBuffer[BufferIndex + 0], 
@@ -41,7 +41,7 @@ float4 ReadTangentX_{DataInterfaceName}(uint VertexIndex)
 	uint BufferIndex = (VertexIndex) * 2;
 	
 #if READABLE_OUTPUT_BUFFERS & OPTIMUS_SKINNED_MESH_READ_TANGENTS
-	return TangentBias_SkinCache({DataInterfaceName}_TangentBufferUAV[BufferIndex]) ;
+	return TangentBias_SkinCache({DataInterfaceName}_TangentBufferSRV[BufferIndex]) ;
 #else	
 	return TangentBias_SkinCache({DataInterfaceName}_TangentStaticBuffer[BufferIndex]) ;
 #endif
@@ -51,7 +51,7 @@ float4 ReadTangentZ_{DataInterfaceName}(uint VertexIndex)
 {
 	uint BufferIndex = (VertexIndex) * 2;
 #if READABLE_OUTPUT_BUFFERS & OPTIMUS_SKINNED_MESH_READ_TANGENTS
-	return TangentBias_SkinCache({DataInterfaceName}_TangentBufferUAV[BufferIndex + 1]) ;
+	return TangentBias_SkinCache({DataInterfaceName}_TangentBufferSRV[BufferIndex + 1]) ;
 #else	
 	return TangentBias_SkinCache({DataInterfaceName}_TangentStaticBuffer[BufferIndex + 1]) ;
 #endif
@@ -62,7 +62,7 @@ float4 ReadColor_{DataInterfaceName}(uint VertexIndex)
 	// Index Mask is used here in case ColorStaticBuffer is bound to the GWhiteVertexBufferWithSRV that is only a few bytes
 	uint BufferIndex = (VertexIndex) & ({DataInterfaceName}_ColorIndexMask);
 #if READABLE_OUTPUT_BUFFERS & OPTIMUS_SKINNED_MESH_READ_COLOR
-	return {DataInterfaceName}_ColorBufferUAV[BufferIndex] FMANUALFETCH_COLOR_COMPONENT_SWIZZLE;
+	return {DataInterfaceName}_ColorBufferSRV[BufferIndex] FMANUALFETCH_COLOR_COMPONENT_SWIZZLE;
 #else	
 	return {DataInterfaceName}_ColorStaticBuffer[BufferIndex] FMANUALFETCH_COLOR_COMPONENT_SWIZZLE;
 #endif	
diff --git a/Engine/Plugins/Animation/DeformerGraph/Source/OptimusCore/Private/DataInterfaces/OptimusDataInterfaceCopyKernel.cpp b/Engine/Plugins/Animation/DeformerGraph/Source/OptimusCore/Private/DataInterfaces/OptimusDataInterfaceCopyKernel.cpp
--- a/Engine/Plugins/Animation/DeformerGraph/Source/OptimusCore/Private/DataInterfaces/OptimusDataInterfaceCopyKernel.cpp
+++ b/Engine/Plugins/Animation/DeformerGraph/Source/OptimusCore/Private/DataInterfaces/OptimusDataInterfaceCopyKernel.cpp
@@ -25,7 +25,7 @@ void UOptimusCopyKernelDataInterface::GetSupportedInputs(TArray<FShaderFunctionD
 {
 	OutFunctions.AddDefaulted_GetRef()
 		.SetName(ReadNumThreadsFunctionName)
-		.AddReturnType(FShaderValueType::Get(EShaderFundamentalType::Int, 3));
+		.AddReturnType(FShaderValueType::Get(EShaderFundamentalType::Uint, 3));
 
 	OutFunctions.AddDefaulted_GetRef()
 		.SetName(ReadNumThreadsPerInvocationFunctionName)
@@ -37,7 +37,7 @@ void UOptimusCopyKernelDataInterface::GetSupportedInputs(TArray<FShaderFunctionD
 }
 
 BEGIN_SHADER_PARAMETER_STRUCT(FCopyKernelDataInterfaceParameters, )
-	SHADER_PARAMETER(FIntVector3, NumThreads)
+	SHADER_PARAMETER(FUintVector3, NumThreads)
 	SHADER_PARAMETER(uint32, NumThreadsPerInvocation)
 	SHADER_PARAMETER(uint32, ThreadIndexOffset)
 END_SHADER_PARAMETER_STRUCT()
@@ -169,7 +169,7 @@ void FOptimusCopyKernelDataProviderProxy::GatherDispatchData(FDispatchData const
 	for (int32 InvocationIndex = 0; InvocationIndex < ParameterArray.Num(); ++InvocationIndex)
 	{
 		FParameters& Parameters = ParameterArray[InvocationIndex];
-		Parameters.NumThreads = FIntVector3(TotalThreadCount, 1, 1);
+		Parameters.NumThreads = FUintVector3(TotalThreadCount, 1, 1);
 		
 		int32 NumThreadsPerInvocation = InDispatchData.bUnifiedDispatch ? TotalThreadCount : InvocationThreadCounts[InvocationIndex] ;
 		Parameters.NumThreadsPerInvocation = NumThreadsPerInvocation;
diff --git a/Engine/Plugins/Animation/DeformerGraph/Source/OptimusCore/Private/DataInterfaces/OptimusDataInterfaceCustomComputeKernel.cpp b/Engine/Plugins/Animation/DeformerGraph/Source/OptimusCore/Private/DataInterfaces/OptimusDataInterfaceCustomComputeKernel.cpp
--- a/Engine/Plugins/Animation/DeformerGraph/Source/OptimusCore/Private/DataInterfaces/OptimusDataInterfaceCustomComputeKernel.cpp
+++ b/Engine/Plugins/Animation/DeformerGraph/Source/OptimusCore/Private/DataInterfaces/OptimusDataInterfaceCustomComputeKernel.cpp
@@ -26,7 +26,7 @@ void UOptimusCustomComputeKernelDataInterface::GetSupportedInputs(TArray<FShader
 {
 	OutFunctions.AddDefaulted_GetRef()
 		.SetName(ReadNumThreadsFunctionName)
-		.AddReturnType(FShaderValueType::Get(EShaderFundamentalType::Int, 3));
+		.AddReturnType(FShaderValueType::Get(EShaderFundamentalType::Uint, 3));
 	
 	OutFunctions.AddDefaulted_GetRef()
 		.SetName(ReadNumThreadsPerInvocationFunctionName)
@@ -38,7 +38,7 @@ void UOptimusCustomComputeKernelDataInterface::GetSupportedInputs(TArray<FShader
 }
 
 BEGIN_SHADER_PARAMETER_STRUCT(FCustomComputeKernelDataInterfaceParameters, )
-	SHADER_PARAMETER(FIntVector3, NumThreads)
+	SHADER_PARAMETER(FUintVector3, NumThreads)
 	SHADER_PARAMETER(uint32, NumThreadsPerInvocation)
 	SHADER_PARAMETER(uint32, ThreadIndexOffset)
 END_SHADER_PARAMETER_STRUCT()
@@ -169,7 +169,7 @@ void FOptimusCustomComputeKernelDataProviderProxy::GatherDispatchData(FDispatchD
 	for (int32 InvocationIndex = 0; InvocationIndex < ParameterArray.Num(); ++InvocationIndex)
 	{
 		FParameters& Parameters = ParameterArray[InvocationIndex];
-		Parameters.NumThreads = FIntVector3(TotalThreadCount, 1, 1);
+		Parameters.NumThreads = FUintVector3(TotalThreadCount, 1, 1);
 		
 		int32 NumThreadsPerInvocation = InDispatchData.bUnifiedDispatch ? TotalThreadCount : InvocationThreadCounts[InvocationIndex] ;
 		Parameters.NumThreadsPerInvocation = NumThreadsPerInvocation;
diff --git a/Engine/Plugins/Animation/DeformerGraph/Source/OptimusCore/Private/DataInterfaces/OptimusDataInterfaceSkinnedMeshRead.cpp b/Engine/Plugins/Animation/DeformerGraph/Source/OptimusCore/Private/DataInterfaces/OptimusDataInterfaceSkinnedMeshRead.cpp
--- a/Engine/Plugins/Animation/DeformerGraph/Source/OptimusCore/Private/DataInterfaces/OptimusDataInterfaceSkinnedMeshRead.cpp
+++ b/Engine/Plugins/Animation/DeformerGraph/Source/OptimusCore/Private/DataInterfaces/OptimusDataInterfaceSkinnedMeshRead.cpp
@@ -88,9 +88,9 @@ void UOptimusSkinnedMeshReadDataInterface::GetSupportedInputs(TArray<FShaderFunc
 
 BEGIN_SHADER_PARAMETER_STRUCT(FSkinnedMeshReadDataInterfaceParameters, )
 	SHADER_PARAMETER(uint32, NumVertices)
-	SHADER_PARAMETER_RDG_BUFFER_UAV(RWBuffer<float>, PositionBufferUAV)
-	SHADER_PARAMETER_RDG_BUFFER_UAV(RWBuffer<SNORM float4>, TangentBufferUAV)
-	SHADER_PARAMETER_RDG_BUFFER_UAV(RWBuffer<UNORM float4>, ColorBufferUAV)
+	SHADER_PARAMETER_RDG_BUFFER_SRV(Buffer<float>, PositionBufferSRV)
+	SHADER_PARAMETER_RDG_BUFFER_SRV(Buffer<SNORM float4>, TangentBufferSRV)
+	SHADER_PARAMETER_RDG_BUFFER_SRV(Buffer<UNORM float4>, ColorBufferSRV)
 	SHADER_PARAMETER_SRV(Buffer<float>, PositionStaticBuffer)
 	SHADER_PARAMETER_SRV(Buffer<SNORM float4>, TangentStaticBuffer)
 	SHADER_PARAMETER_SRV(Buffer<float4>, ColorStaticBuffer)
@@ -219,14 +219,16 @@ void FOptimusSkinnedMeshReadDataProviderProxy::AllocateResources(FRDGBuilder& Gr
 		*LastLodIndexPtr = LodIndex;
 	}
 	
+
 	if (InputMask & static_cast<uint64>(ESkinnedMeshReadDataInterfaceInputSelectorMask::Position) )
 	{
 		FRDGBuffer* PositionBuffer = FSkeletalMeshDeformerHelpers::AllocateVertexFactoryPositionBuffer(GraphBuilder, InAllocationData.ExternalAccessQueue, SkeletalMeshObject, LodIndex, TEXT("OptimusSkinnedMeshPosition"));
-		PositionBufferUAV = GraphBuilder.CreateUAV(PositionBuffer, PF_R32_FLOAT, ERDGUnorderedAccessViewFlags::SkipBarrier);
+		PositionBufferSRV = GraphBuilder.CreateSRV(PositionBuffer, PF_R32_FLOAT); 
 	}
 	else
 	{
-		PositionBufferUAV = GraphBuilder.CreateUAV(GraphBuilder.RegisterExternalBuffer(GWhiteVertexBufferWithRDG->Buffer), PF_R32_FLOAT);
+		FRDGBuffer* White = GraphBuilder.RegisterExternalBuffer(GWhiteVertexBufferWithRDG->Buffer);	
+		PositionBufferSRV = GraphBuilder.CreateSRV(White, PF_R32_FLOAT); 
 	}	
 
 	// OpenGL ES does not support writing to RGBA16_SNORM images, instead pack data into SINT in the shader
@@ -236,22 +238,24 @@ void FOptimusSkinnedMeshReadDataProviderProxy::AllocateResources(FRDGBuilder& Gr
 		(InputMask & static_cast<uint64>(ESkinnedMeshReadDataInterfaceInputSelectorMask::TangentZ)))
 	{
 		FRDGBuffer* TangentBuffer = FSkeletalMeshDeformerHelpers::AllocateVertexFactoryTangentBuffer(GraphBuilder, InAllocationData.ExternalAccessQueue, SkeletalMeshObject, LodIndex, TEXT("OptimusSkinnedMeshTangent"));
-		TangentBufferUAV = GraphBuilder.CreateUAV(TangentBuffer, TangentsFormat, ERDGUnorderedAccessViewFlags::SkipBarrier);
+		TangentBufferSRV = GraphBuilder.CreateSRV(TangentBuffer, TangentsFormat); 
 	}
 	else
 	{
-		TangentBufferUAV = GraphBuilder.CreateUAV(GraphBuilder.RegisterExternalBuffer(GWhiteVertexBufferWithRDG->Buffer), TangentsFormat);
+		FRDGBuffer* White = GraphBuilder.RegisterExternalBuffer(GWhiteVertexBufferWithRDG->Buffer);	
+		TangentBufferSRV = GraphBuilder.CreateSRV(White, TangentsFormat); 
 	}
 
 	if (InputMask & static_cast<uint64>(ESkinnedMeshReadDataInterfaceInputSelectorMask::Color))
 	{
 		FRDGBuffer* ColorBuffer = FSkeletalMeshDeformerHelpers::AllocateVertexFactoryColorBuffer(GraphBuilder, InAllocationData.ExternalAccessQueue, SkeletalMeshObject, LodIndex, TEXT("OptimusSkinnedMeshColor"));
 		// using RGBA here and do a manual fetch swizzle in shader instead of BGRA directly because some Mac does not support it. See GMetalBufferFormats[PF_B8G8R8A8] 
-		ColorBufferUAV = GraphBuilder.CreateUAV(ColorBuffer, PF_R8G8B8A8, ERDGUnorderedAccessViewFlags::SkipBarrier);
+		ColorBufferSRV = GraphBuilder.CreateSRV(ColorBuffer, PF_R8G8B8A8); 
 	}
 	else
 	{
-		ColorBufferUAV = GraphBuilder.CreateUAV(GraphBuilder.RegisterExternalBuffer(GWhiteVertexBufferWithRDG->Buffer), PF_A32B32G32R32F);
+		FRDGBuffer* White = GraphBuilder.RegisterExternalBuffer(GWhiteVertexBufferWithRDG->Buffer);		
+		ColorBufferSRV = GraphBuilder.CreateSRV(White, TangentsFormat); 
 	}
 
 	FSkeletalMeshDeformerHelpers::UpdateVertexFactoryBufferOverrides(GraphBuilder, SkeletalMeshObject, LodIndex, bInvalidatePreviousPosition);
@@ -298,9 +302,9 @@ void FOptimusSkinnedMeshReadDataProviderProxy::GatherDispatchData(FDispatchData
 		FParameters& Parameters = ParameterArray[InvocationIndex];
 		Parameters.NumVertices = LodRenderData->GetNumVertices();
 
-		Parameters.PositionBufferUAV = PositionBufferUAV;
-		Parameters.TangentBufferUAV = TangentBufferUAV;
-		Parameters.ColorBufferUAV = ColorBufferUAV;
+		Parameters.PositionBufferSRV = PositionBufferSRV;
+		Parameters.TangentBufferSRV = TangentBufferSRV;
+		Parameters.ColorBufferSRV = ColorBufferSRV;
 		
 		FRHIShaderResourceView* MeshVertexBufferSRV = LodRenderData->StaticVertexBuffers.PositionVertexBuffer.GetSRV();
 		FRHIShaderResourceView* MeshTangentBufferSRV = LodRenderData->StaticVertexBuffers.StaticMeshVertexBuffer.GetTangentsSRV();
diff --git a/Engine/Plugins/Animation/DeformerGraph/Source/OptimusCore/Private/DataInterfaces/OptimusDataInterfaceSkinnedMeshRead.h b/Engine/Plugins/Animation/DeformerGraph/Source/OptimusCore/Private/DataInterfaces/OptimusDataInterfaceSkinnedMeshRead.h
--- a/Engine/Plugins/Animation/DeformerGraph/Source/OptimusCore/Private/DataInterfaces/OptimusDataInterfaceSkinnedMeshRead.h
+++ b/Engine/Plugins/Animation/DeformerGraph/Source/OptimusCore/Private/DataInterfaces/OptimusDataInterfaceSkinnedMeshRead.h
@@ -98,9 +98,7 @@ private:
 	int32* LastLodIndexPtr = nullptr;
 	EMeshDeformerOutputBuffer OutputBuffersFromPreviousInstances;
 
-	// Using UAV here because we might be reading and writing to these buffers in the same kernel
-	// if we have a setup like Read -> Kernel -> Write 
-	FRDGBufferUAV* PositionBufferUAV = nullptr;
-	FRDGBufferUAV* TangentBufferUAV = nullptr;
-	FRDGBufferUAV* ColorBufferUAV = nullptr;
+	FRDGBufferSRV* PositionBufferSRV = nullptr;
+	FRDGBufferSRV* TangentBufferSRV = nullptr;
+	FRDGBufferSRV* ColorBufferSRV = nullptr;
 };
diff --git a/Engine/Plugins/Animation/DeformerGraph/Source/OptimusCore/Private/Nodes/OptimusNode_ComputeKernelBase.cpp b/Engine/Plugins/Animation/DeformerGraph/Source/OptimusCore/Private/Nodes/OptimusNode_ComputeKernelBase.cpp
--- a/Engine/Plugins/Animation/DeformerGraph/Source/OptimusCore/Private/Nodes/OptimusNode_ComputeKernelBase.cpp
+++ b/Engine/Plugins/Animation/DeformerGraph/Source/OptimusCore/Private/Nodes/OptimusNode_ComputeKernelBase.cpp
@@ -811,7 +811,7 @@ TOptional<FText> UOptimusNode_ComputeKernelBase::BindKernelDataInterfaceForCompu
 	}
 
 	OutGeneratedFunctions.Add(FString::Printf(
-			TEXT("int3 Read%s() { return %s::%s(); }"),
+			TEXT("uint3 Read%s() { return %s::%s(); }"),
 			GetNumThreadsReservedName(), Optimus::GetKernelInternalNamespaceName(), KernelDataInterface->GetReadNumThreadsFunctionName() ));
 
 	return {};
diff --git a/Engine/Plugins/Animation/DeformerGraph/Source/OptimusCore/Private/Nodes/OptimusNode_CustomComputeKernel.cpp b/Engine/Plugins/Animation/DeformerGraph/Source/OptimusCore/Private/Nodes/OptimusNode_CustomComputeKernel.cpp
--- a/Engine/Plugins/Animation/DeformerGraph/Source/OptimusCore/Private/Nodes/OptimusNode_CustomComputeKernel.cpp
+++ b/Engine/Plugins/Animation/DeformerGraph/Source/OptimusCore/Private/Nodes/OptimusNode_CustomComputeKernel.cpp
@@ -1768,7 +1768,7 @@ void UOptimusNode_CustomComputeKernel::UpdatePreamble()
 	}
 	ShaderSource.Declarations += "\n// Resource Indexing\n";
 	ShaderSource.Declarations += "uint Index;	// From SV_DispatchThreadID.x\n";
-	ShaderSource.Declarations += "int3 ReadNumThreads();\n";
+	ShaderSource.Declarations += "uint3 ReadNumThreads();\n";
 }
 
 UOptimusNodePin* UOptimusNode_CustomComputeKernel::GetPrimaryGroupPin_Internal() const
