[Deforemr Graph] Switching to use hard reference + guid for function reference node to refernce a function graph. Resaving most DG assets would be highly recommended, though not strictly requiired (there would be warnings.). Previously function nodes only soft references function graph causing two issues.
1. Loader cannot fully load function graph before the referencer graph's postload, preventing compile on postload to work
2. There is no redirector support for the function graph object, so renaming a function graph in the engine can invalidate existing references in user projects.

diff --git a/Engine/Plugins/Animation/DeformerGraph/Source/OptimusCore/Private/Nodes/OptimusNode_FunctionReference.cpp b/Engine/Plugins/Animation/DeformerGraph/Source/OptimusCore/Private/Nodes/OptimusNode_FunctionReference.cpp
--- a/Engine/Plugins/Animation/DeformerGraph/Source/OptimusCore/Private/Nodes/OptimusNode_FunctionReference.cpp
+++ b/Engine/Plugins/Animation/DeformerGraph/Source/OptimusCore/Private/Nodes/OptimusNode_FunctionReference.cpp
@@ -8,24 +8,34 @@
 #include "OptimusNodePin.h"
 #include "OptimusNode_GraphTerminal.h"
 #include "OptimusComponentSource.h"
+#include "OptimusCoreModule.h"
 #include "OptimusDeformer.h"
+#include "OptimusObjectVersion.h"
 
 #include UE_INLINE_GENERATED_CPP_BY_NAME(OptimusNode_FunctionReference)
 
 
+UOptimusFunctionNodeGraph* FOptimusFunctionGraphIdentifier::Resolve()
+{
+	return Asset->FindFunctionByGuid(Guid);
+}
+
 FName UOptimusNode_FunctionReference::GetNodeCategory() const
 {
-	return FunctionGraph->Category;
+	if (ResolvedFunctionGraph.IsValid())
+	{
+		return ResolvedFunctionGraph->Category;
+	}
 
+	return NAME_None;
 }
 
 FText UOptimusNode_FunctionReference::GetDisplayName() const
 {
-	if (FunctionGraph.IsValid())
+	if (ResolvedFunctionGraph.IsValid())
 	{
-		return FText::FromString(FunctionGraph->GetNodeName());
+		return FText::FromString(ResolvedFunctionGraph->GetNodeName());
 	}
-
 	
 	return FText::FromString("<graph missing>");
 }
@@ -33,7 +43,7 @@ FText UOptimusNode_FunctionReference::GetDisplayName() const
 
 void UOptimusNode_FunctionReference::ConstructNode()
 {
-	if (FunctionGraph.IsValid())
+	if (ResolvedFunctionGraph.IsValid())
 	{
 		const FOptimusDataTypeRegistry& TypeRegistry = FOptimusDataTypeRegistry::Get();
 		FOptimusDataTypeRef ComponentSourceType = TypeRegistry.FindType(*UOptimusComponentSourceBinding::StaticClass());
@@ -41,11 +51,11 @@ void UOptimusNode_FunctionReference::ConstructNode()
 		
 		// After a duplicate, the kernel node has no pins, so we need to reconstruct them from
 		// the bindings. We can assume that all naming clashes have already been dealt with.
-		for (const FOptimusParameterBinding& Binding: FunctionGraph->InputBindings)
+		for (const FOptimusParameterBinding& Binding: ResolvedFunctionGraph->InputBindings)
 		{
 			AddPinDirect(Binding, EOptimusNodePinDirection::Input);
 		}
-		for (const FOptimusParameterBinding& Binding: FunctionGraph->OutputBindings)
+		for (const FOptimusParameterBinding& Binding: ResolvedFunctionGraph->OutputBindings)
 		{
 			AddPinDirect(Binding, EOptimusNodePinDirection::Output);
 		}
@@ -63,7 +73,7 @@ FOptimusRoutedNodePin UOptimusNode_FunctionReference::GetPinCounterpart(
 		return {};
 	}
 
-	if (!FunctionGraph.IsValid())
+	if (!ResolvedFunctionGraph.IsValid())
 	{
 		return {};
 	}
@@ -71,14 +81,14 @@ FOptimusRoutedNodePin UOptimusNode_FunctionReference::GetPinCounterpart(
 	UOptimusNode_GraphTerminal* CounterpartNode = nullptr;
 	if (InNodePin->GetDirection() == EOptimusNodePinDirection::Input)
 	{
-		CounterpartNode = FunctionGraph->GetTerminalNode(EOptimusTerminalType::Entry);
+		CounterpartNode = ResolvedFunctionGraph->GetTerminalNode(EOptimusTerminalType::Entry);
 	}
 	else if (InNodePin->GetDirection() == EOptimusNodePinDirection::Output)
 	{
-		CounterpartNode = FunctionGraph->GetTerminalNode(EOptimusTerminalType::Return);
+		CounterpartNode = ResolvedFunctionGraph->GetTerminalNode(EOptimusTerminalType::Return);
 	}
 
-	if (!(CounterpartNode))
+	if (!ensure(CounterpartNode))
 	{
 		return {};
 	}
@@ -94,12 +104,12 @@ FOptimusRoutedNodePin UOptimusNode_FunctionReference::GetPinCounterpart(
 
 UOptimusNodeGraph* UOptimusNode_FunctionReference::GetNodeGraphToShow()
 {
-	return FunctionGraph.Get();
+	return ResolvedFunctionGraph.Get();
 }
 
 UOptimusNodeSubGraph* UOptimusNode_FunctionReference::GetReferencedSubGraph() const
 {
-	return FunctionGraph.Get();
+	return ResolvedFunctionGraph.Get();
 }
 
 UOptimusComponentSourceBinding* UOptimusNode_FunctionReference::GetDefaultComponentBinding(const FOptimusPinTraversalContext& InTraversalContext) const
@@ -136,32 +146,59 @@ UOptimusNodePin* UOptimusNode_FunctionReference::GetDefaultComponentBindingPin()
 	return DefaultComponentPin.Get();
 }
 
-FSoftObjectPath UOptimusNode_FunctionReference::GetSerializedGraphPath() const
+void UOptimusNode_FunctionReference::SetReferencedFunctionGraph(const FOptimusFunctionGraphIdentifier& InGraphIdentifier)
 {
-	return FunctionGraph.ToSoftObjectPath();
+	FunctionGraphIdentifier = InGraphIdentifier;
 }
 
-void UOptimusNode_FunctionReference::InitializeSerializedGraphPath(const FSoftObjectPath& InInitialGraphPath)
+const FOptimusFunctionGraphIdentifier& UOptimusNode_FunctionReference::GetReferencedFunctionGraphIdentifier()
 {
-	FunctionGraph = InInitialGraphPath;
+	return FunctionGraphIdentifier;
 }
 
-void UOptimusNode_FunctionReference::RefreshSerializedGraphPath(const FSoftObjectPath& InRenamedGraphPath)
+void UOptimusNode_FunctionReference::UpdateDisplayName()
 {
-	// Unlike subgraph reference where we can directly query the graph for its new name/path
-	// the function graph may not be loaded at the time of rename and thus
-	// the caller needs to provide the path to the renamed graph 
-	FunctionGraph = InRenamedGraphPath;
-	FunctionGraph.LoadSynchronous();
 	SetDisplayName(GetDisplayName());
 }
 
+void UOptimusNode_FunctionReference::PostLoadNodeSpecificData()
+{
+	Super::PostLoadNodeSpecificData();
+
+	FOptimusObjectVersion::Type ObjectVersion = static_cast<FOptimusObjectVersion::Type>(GetLinkerCustomVersion(FOptimusObjectVersion::GUID));
+	if (ObjectVersion < FOptimusObjectVersion::FunctionGraphUseGuid)
+	{
+PRAGMA_DISABLE_DEPRECATION_WARNINGS
+		if (!FunctionGraph_DEPRECATED.IsNull())
+		{
+			// Unfortnately we need to do a sync load here in post load to convert soft reference to hard reference.
+			// The load is not guaranteed to reach RF_LoadCompleted, hence the name "PartiallyLoadedGraph".
+			if (UOptimusFunctionNodeGraph* PartiallyLoadedGraph = FunctionGraph_DEPRECATED.LoadSynchronous())
+			{
+				FunctionGraphIdentifier.Asset = PartiallyLoadedGraph->GetTypedOuter<UOptimusDeformer>();
+				FunctionGraphIdentifier.Guid = UOptimusFunctionNodeGraph::GetGuidForGraphWithoutGuid(FunctionGraph_DEPRECATED);
+				
+				ResolvedFunctionGraph = PartiallyLoadedGraph;
+				
+				FunctionGraph_DEPRECATED.Reset();
+
+				// Suggest a resave to avoid going down this path all the time
+				UE_LOG(LogOptimusCore, Warning ,TEXT("Deformer Graph %s should be resaved to improve loading performance"), *GetPackage()->GetName());
+				Modify();
+			}
+			
+		}
+PRAGMA_ENABLE_DEPRECATION_WARNINGS
+	}
+}
+
 void UOptimusNode_FunctionReference::InitializeTransientData()
 {
 	Super::InitializeTransientData();
-	// Making sure the pointer is alive such that we avoid loading on demand everywhere
-	if (UOptimusFunctionNodeGraph* LoadedFunctionGraph = FunctionGraph.LoadSynchronous())
+	if (!ResolvedFunctionGraph.IsValid())
 	{
-		LoadedFunctionGraph->ConditionalPostLoad();
+		ResolvedFunctionGraph = FunctionGraphIdentifier.Resolve();
 	}
-}
+	
+	UpdateDisplayName();
+}
\ No newline at end of file
diff --git a/Engine/Plugins/Animation/DeformerGraph/Source/OptimusCore/Private/Nodes/OptimusNode_FunctionReference.h b/Engine/Plugins/Animation/DeformerGraph/Source/OptimusCore/Private/Nodes/OptimusNode_FunctionReference.h
--- a/Engine/Plugins/Animation/DeformerGraph/Source/OptimusCore/Private/Nodes/OptimusNode_FunctionReference.h
+++ b/Engine/Plugins/Animation/DeformerGraph/Source/OptimusCore/Private/Nodes/OptimusNode_FunctionReference.h
@@ -11,7 +11,6 @@
 #include "OptimusNode_FunctionReference.generated.h"
 
 
-class UOptimusFunctionNodeGraph;
 
 UCLASS(Hidden)
 class UOptimusNode_FunctionReference :
@@ -42,19 +41,27 @@ public:
 	UOptimusComponentSourceBinding* GetDefaultComponentBinding(const FOptimusPinTraversalContext& InTraversalContext) const override;
 	UOptimusNodePin* GetDefaultComponentBindingPin() const override;
 
-	FSoftObjectPath GetSerializedGraphPath() const;
-	void InitializeSerializedGraphPath(const FSoftObjectPath& InInitialGraphPath);
-	void RefreshSerializedGraphPath(const FSoftObjectPath& InRenamedGraphPath);
+	void SetReferencedFunctionGraph(const FOptimusFunctionGraphIdentifier& InGraphIdentifier);
+	const FOptimusFunctionGraphIdentifier& GetReferencedFunctionGraphIdentifier();
+	void UpdateDisplayName();
 	
 protected:
+	void PostLoadNodeSpecificData() override;
 	void InitializeTransientData() override;
 	
-	/** The graph that owns us. This contains all the necessary pin information to add on
-	 * the terminal node.
-	 */
 	UPROPERTY()
-	TSoftObjectPtr<UOptimusFunctionNodeGraph> FunctionGraph;
+	FOptimusFunctionGraphIdentifier FunctionGraphIdentifier;
 
 	UPROPERTY()
 	TWeakObjectPtr<UOptimusNodePin> DefaultComponentPin;
+
+	UPROPERTY(Transient, VisibleAnywhere, Category=FunctionReference, DisplayName = "Referenced Function Graph")
+	TWeakObjectPtr<UOptimusFunctionNodeGraph> ResolvedFunctionGraph;
+	
+private:
+	UE_DEPRECATED(5.6, "Use FunctionGraphIdentifier instead.")
+	UPROPERTY(Meta = (DeprecatedProperty))
+	TSoftObjectPtr<UOptimusFunctionNodeGraph> FunctionGraph_DEPRECATED;
+
+	bool bDelayResolvingGraph = false;
 };
diff --git a/Engine/Plugins/Animation/DeformerGraph/Source/OptimusCore/Private/OptimusCoreModule.cpp b/Engine/Plugins/Animation/DeformerGraph/Source/OptimusCore/Private/OptimusCoreModule.cpp
--- a/Engine/Plugins/Animation/DeformerGraph/Source/OptimusCore/Private/OptimusCoreModule.cpp
+++ b/Engine/Plugins/Animation/DeformerGraph/Source/OptimusCore/Private/OptimusCoreModule.cpp
@@ -62,54 +62,6 @@ bool FOptimusCoreModule::RegisterDataInterfaceClass(TSubclassOf<UOptimusComputeD
 	return false;
 }
 
-void FOptimusCoreModule::UpdateFunctionReferences(const FSoftObjectPath& InOldGraphPath, const FSoftObjectPath& InNewGraphPath)
-{
-	const FAssetRegistryModule& AssetRegistryModule = FModuleManager::LoadModuleChecked<FAssetRegistryModule>("AssetRegistry");
-	
-	// find all assets in the project
-	TArray<FAssetData> AssetDatas;
-	AssetRegistryModule.Get().GetAssetsByClass(UOptimusDeformer::StaticClass()->GetClassPathName(), AssetDatas);
-
-	FOptimusFunctionReferenceData FunctionReferenceData;
-	
-	// loop over all found assets
-	for(const FAssetData& AssetData : AssetDatas)
-	{
-		RegisterFunctionReferencesFromAsset(FunctionReferenceData, AssetData);
-	}
-	
-	if (FOptimusFunctionReferenceNodeSet* FunctionNodeArray = FunctionReferenceData.FunctionReferences.Find(InOldGraphPath))
-	{
-		for (TSoftObjectPtr<UOptimusNode_FunctionReference> FunctionNodePtr : FunctionNodeArray->Nodes)
-		{
-			UOptimusNode_FunctionReference* FunctionNode = FunctionNodePtr.LoadSynchronous();
-
-			if (FunctionNode)
-			{
-				FunctionNode->Modify();
-				FunctionNode->RefreshSerializedGraphPath(InNewGraphPath);
-				FunctionNode->MarkPackageDirty();
-			}
-		}
-	}
-}
-
-void FOptimusCoreModule::RegisterFunctionReferencesFromAsset(FOptimusFunctionReferenceData& InOutData, const FAssetData& AssetData)
-{
-	FString FunctionReferenceString = AssetData.GetTagValueRef<FString>(UOptimusDeformer::FunctionReferencesAssetTagName);
-
-	if (!FunctionReferenceString.IsEmpty())
-	{
-		FOptimusFunctionReferenceData AssetFunctionReferenceData;
-
-		FOptimusFunctionReferenceData::StaticStruct()->ImportText(*FunctionReferenceString, &AssetFunctionReferenceData, nullptr, PPF_None, nullptr, {});
-
-		for(const TPair<FSoftObjectPath, FOptimusFunctionReferenceNodeSet>& Pair : AssetFunctionReferenceData.FunctionReferences)
-		{
-			InOutData.FunctionReferences.FindOrAdd(Pair.Key).Nodes.Append(Pair.Value.Nodes);
-		}
-	}
-}
 
 IMPLEMENT_MODULE(FOptimusCoreModule, OptimusCore)
 
diff --git a/Engine/Plugins/Animation/DeformerGraph/Source/OptimusCore/Private/OptimusCoreModule.h b/Engine/Plugins/Animation/DeformerGraph/Source/OptimusCore/Private/OptimusCoreModule.h
--- a/Engine/Plugins/Animation/DeformerGraph/Source/OptimusCore/Private/OptimusCoreModule.h
+++ b/Engine/Plugins/Animation/DeformerGraph/Source/OptimusCore/Private/OptimusCoreModule.h
@@ -17,12 +17,6 @@ public:
 	/** IOptimusCoreModule implementation */
 	bool RegisterDataInterfaceClass(TSubclassOf<UOptimusComputeDataInterface> InDataInterfaceClass) override;
 
-	void UpdateFunctionReferences(const FSoftObjectPath& InOldGraphPath, const FSoftObjectPath& InNewGraphPath) override;
-	
-protected:
-	void RegisterFunctionReferencesFromAsset(FOptimusFunctionReferenceData& InOutData, const FAssetData& AssetData);
-	
-	
 };
 
 DECLARE_LOG_CATEGORY_EXTERN(LogOptimusCore, Log, All);
diff --git a/Engine/Plugins/Animation/DeformerGraph/Source/OptimusCore/Private/OptimusDeformer.cpp b/Engine/Plugins/Animation/DeformerGraph/Source/OptimusCore/Private/OptimusDeformer.cpp
--- a/Engine/Plugins/Animation/DeformerGraph/Source/OptimusCore/Private/OptimusDeformer.cpp
+++ b/Engine/Plugins/Animation/DeformerGraph/Source/OptimusCore/Private/OptimusDeformer.cpp
@@ -59,6 +59,8 @@
 
 #include "OptimusDeformerDynamicInstanceManager.h"
 #include "OptimusFunctionNodeGraphHeader.h"
+#include "OptimusFunctionNodeGraphHeaderWithGuid.h"
+#include "AssetRegistry/AssetRegistryModule.h"
 #include "Nodes/OptimusNode_FunctionReference.h"
 #include "Nodes/OptimusNode_SubGraphReference.h"
 
@@ -69,8 +71,8 @@
 static const FName DefaultResourceName("Resource");
 static const FName DefaultVariableName("Variable");
 
-const FName UOptimusDeformer::PublicFunctionsAssetTagName = TEXT("PublicFunctions");
-const FName UOptimusDeformer::FunctionReferencesAssetTagName = TEXT("FunctionReferences");
+const TCHAR* UOptimusDeformer::PublicFunctionsAssetTagName = TEXT("PublicFunctions");
+const TCHAR* UOptimusDeformer::PublicFunctionsWithGuidAssetTagName = TEXT("PublicFunctionsWithGuid");
 
 UOptimusDeformer::UOptimusDeformer()
 {
@@ -211,6 +213,18 @@ TArray<UOptimusFunctionNodeGraph*> UOptimusDeformer::GetFunctionGraphs(FName InA
 	return FunctionGraphs;
 }
 
+UOptimusFunctionNodeGraph* UOptimusDeformer::FindFunctionByGuid(FGuid InFunctionGraphGuid)
+{
+	for (UOptimusFunctionNodeGraph* FunctionNodeGraph : GetFunctionGraphs())
+	{
+		if (FunctionNodeGraph->GetGuid() == InFunctionGraphGuid)
+		{
+			return FunctionNodeGraph;
+		}
+	}
+
+	return nullptr;
+}
 
 UOptimusVariableDescription* UOptimusDeformer::AddVariable(
 	FOptimusDataTypeRef InDataTypeRef, 
@@ -1619,6 +1633,52 @@ TArray<UOptimusNode*> UOptimusDeformer::GetAllNodesOfClass(UClass* InNodeClass)
 	return NodesFound;
 }
 
+void UOptimusDeformer::OnGraphRenamedOrRemoved(UOptimusNodeGraph* InGraph) const
+{
+	if (UOptimusFunctionNodeGraph* FunctionNodeGraph = Cast<UOptimusFunctionNodeGraph>(InGraph))
+	{
+		const FAssetRegistryModule& AssetRegistryModule = FModuleManager::LoadModuleChecked<FAssetRegistryModule>("AssetRegistry");
+		IAssetRegistry&	AssetRegistry = AssetRegistryModule.Get();
+
+		TArray<FAssetIdentifier> Referencers;
+		FAssetIdentifier Source(InGraph->GetPackage()->GetFName(), NAME_None);
+		AssetRegistry.GetReferencers(Source, Referencers);
+
+		Referencers.Add(Source);
+		
+		for (const FAssetIdentifier& Referencer : Referencers)
+		{
+			TArray<FAssetData> Assets;
+			AssetRegistry.GetAssetsByPackageName(Referencer.PackageName, Assets);
+
+			for (const FAssetData& Asset : Assets)
+			{
+				if (Asset.IsInstanceOf(UOptimusDeformer::StaticClass()))
+				{
+					constexpr bool bLoad = false;
+					if (UOptimusDeformer* ReferencerDeformer = Cast<UOptimusDeformer>(Asset.FastGetAsset(bLoad)))
+					{
+						ReferencerDeformer->UpdateFunctionReferenceNodeDisplayName(FunctionNodeGraph->GetGraphIdentifier());
+					}
+				}
+			}
+		}	
+	}
+}
+
+void UOptimusDeformer::UpdateFunctionReferenceNodeDisplayName(const FOptimusFunctionGraphIdentifier& InRenamedFunction)
+{
+	TArray<UOptimusNode*> FunctionNodes = GetAllNodesOfClass(UOptimusNode_FunctionReference::StaticClass());
+	for (UOptimusNode* Node: FunctionNodes)
+	{
+		UOptimusNode_FunctionReference* ReferenceNode = CastChecked<UOptimusNode_FunctionReference>(Node);
+		if (ReferenceNode->GetReferencedFunctionGraphIdentifier() == InRenamedFunction)
+		{
+			ReferenceNode->UpdateDisplayName();
+		}
+	}
+}
+
 struct FOptimusInstancedNode
 {
 	FOptimusRoutedConstNode	RoutedNode;
@@ -3239,12 +3299,21 @@ void UOptimusDeformer::Notify(EOptimusGlobalNotifyType InNotifyType, UObject* In
 	switch (InNotifyType)
 	{
 	case EOptimusGlobalNotifyType::GraphAdded: 
-	case EOptimusGlobalNotifyType::GraphRemoved:
 	case EOptimusGlobalNotifyType::GraphIndexChanged:
-	case EOptimusGlobalNotifyType::GraphRenamed:
 		checkSlow(Cast<UOptimusNodeGraph>(InObject) != nullptr);
 		break;
+		
+	case EOptimusGlobalNotifyType::GraphRemoved:
+	case EOptimusGlobalNotifyType::GraphRenamed:
+		{
+			UOptimusNodeGraph* Graph = CastChecked<UOptimusNodeGraph>(InObject);
+			if (ensure(Graph))
+			{
+				OnGraphRenamedOrRemoved(Graph);
+			}
 
+			break;
+		}
 	case EOptimusGlobalNotifyType::ComponentBindingAdded:
 	case EOptimusGlobalNotifyType::ComponentBindingRemoved:
 	case EOptimusGlobalNotifyType::ComponentBindingIndexChanged:
@@ -3452,7 +3521,33 @@ void UOptimusDeformer::PostLoad()
 	if (Status != EOptimusDeformerStatus::HasErrors &&
 		ObjectVersion < FOptimusObjectVersion::LatestVersion)
 	{
-		Compile();
+		bool bDelayCompile = false;
+		if (ObjectVersion < FOptimusObjectVersion::FunctionGraphUseGuid)
+		{
+			TArray<UOptimusNode*> FunctionNodes = GetAllNodesOfClass(UOptimusNode_FunctionReference::StaticClass());
+			if (FunctionNodes.Num() > 0)
+			{
+				bDelayCompile = true;
+				// Before this version, external function graphs were only soft references 
+				// so in some case it is not possible make sure they are fully loaded during PostLoad.
+				// Delaying the compilation here such that it is performed after all dependencies are fully loaded (hopefully)
+				// It is best to resave the asset to avoid going down this path
+				FFunctionGraphTask::CreateAndDispatchWhenReady(
+					[WeakThis = TWeakObjectPtr(this)]()
+					{
+						if (WeakThis.IsValid())
+						{
+							WeakThis->Compile();
+						}
+					}
+				, TStatId(), NULL, ENamedThreads::GameThread);	
+			}
+		}
+
+		if (!bDelayCompile)
+		{
+			Compile();
+		}
 	}
 }
 
@@ -3823,37 +3918,19 @@ void UOptimusDeformer::GetAssetRegistryTags(FAssetRegistryTagsContext Context) c
 		Context.AddTag(FAssetRegistryTag(TEXT("PrimaryBindingClass"), *ClassPath.ToString(), FAssetRegistryTag::TT_Hidden));
 	}
 
-	// Expose all the public functions
+	// Expose all the public functions, we used to use FOptimusFunctionNodeGraphHeader, but has since switched to FOptimusFunctionNodeGraphHeaderWithGuid
 	{
-		FOptimusFunctionNodeGraphHeaderArray PublicFunctionHeadersArray;
+		FOptimusFunctionNodeGraphHeaderWithGuidArray PublicFunctionHeadersArray;
 		for (const UOptimusFunctionNodeGraph* FunctionNodeGraph: GetFunctionGraphs(UOptimusFunctionNodeGraph::AccessSpecifierPublicName))
 		{
-			PublicFunctionHeadersArray.Headers.Add(FunctionNodeGraph->GetHeader());
+			PublicFunctionHeadersArray.Headers.Add(FunctionNodeGraph->GetHeaderWithGuid());
 		}
-
 	
 		FString PublicFunctionString;
 	
-		FOptimusFunctionNodeGraphHeaderArray::StaticStruct()->ExportText(PublicFunctionString, &PublicFunctionHeadersArray, nullptr, nullptr, PPF_None, nullptr);
+		FOptimusFunctionNodeGraphHeaderWithGuidArray::StaticStruct()->ExportText(PublicFunctionString, &PublicFunctionHeadersArray, nullptr, nullptr, PPF_None, nullptr);
 
-		Context.AddTag(FAssetRegistryTag(PublicFunctionsAssetTagName, *PublicFunctionString, FAssetRegistryTag::TT_Hidden));	
-	}
-
-	// Record all nodes referencing a function, such that function nodes can react when the referencing function changes
-	{
-		FOptimusFunctionReferenceData FunctionReferenceData;
-		TArray<UOptimusNode*> FunctionNodes = GetAllNodesOfClass(UOptimusNode_FunctionReference::StaticClass());
-
-		for (UOptimusNode* Node : FunctionNodes)
-		{
-			UOptimusNode_FunctionReference* FunctionNode = CastChecked<UOptimusNode_FunctionReference>(Node);
-			FunctionReferenceData.FunctionReferences.FindOrAdd(FunctionNode->GetSerializedGraphPath()).Nodes.Add(FunctionNode);
-		}
-
-		FString FunctionReferenceString;
-		FOptimusFunctionReferenceData::StaticStruct()->ExportText(FunctionReferenceString, &FunctionReferenceData, nullptr, nullptr, PPF_None, nullptr);
-
-		Context.AddTag(FAssetRegistryTag(FunctionReferencesAssetTagName, *FunctionReferenceString, FAssetRegistryTag::TT_Hidden));	
+		Context.AddTag(FAssetRegistryTag(PublicFunctionsWithGuidAssetTagName, *PublicFunctionString, FAssetRegistryTag::TT_Hidden));
 	}
 
 	{
@@ -4223,20 +4300,8 @@ bool UOptimusDeformer::MoveGraphDirect(
 
 bool UOptimusDeformer::RenameGraphDirect(UOptimusNodeGraph* InGraph, const FString& InNewName)
 {
-	FSoftObjectPath PreviousFunctionGraphPath;
-	if (InGraph->GetGraphType() == EOptimusNodeGraphType::Function)
-	{
-		PreviousFunctionGraphPath = FSoftObjectPath(InGraph);
-	}	
-	
 	if (Optimus::RenameObject(InGraph, *InNewName, nullptr))
 	{
-		if (InGraph->GetGraphType() == EOptimusNodeGraphType::Function)
-		{
-			IOptimusCoreModule& OptimusCoreModule = IOptimusCoreModule::Get();
-			OptimusCoreModule.UpdateFunctionReferences(PreviousFunctionGraphPath, FSoftObjectPath(InGraph));
-		}
-
 		Notify(EOptimusGlobalNotifyType::GraphRenamed, InGraph);
 		return true;
 	}
diff --git a/Engine/Plugins/Animation/DeformerGraph/Source/OptimusCore/Private/OptimusFunctionNodeGraph.cpp b/Engine/Plugins/Animation/DeformerGraph/Source/OptimusCore/Private/OptimusFunctionNodeGraph.cpp
--- a/Engine/Plugins/Animation/DeformerGraph/Source/OptimusCore/Private/OptimusFunctionNodeGraph.cpp
+++ b/Engine/Plugins/Animation/DeformerGraph/Source/OptimusCore/Private/OptimusFunctionNodeGraph.cpp
@@ -1,7 +1,9 @@
 ﻿// Copyright Epic Games, Inc. All Rights Reserved.
 
 #include "OptimusFunctionNodeGraph.h"
-#include "OptimusFunctionNodeGraphHeader.h"
+#include "OptimusFunctionNodeGraphHeaderWithGuid.h"
+#include "OptimusDeformer.h"
+#include "OptimusObjectVersion.h"
 
 #include UE_INLINE_GENERATED_CPP_BY_NAME(OptimusFunctionNodeGraph)
 
@@ -13,25 +15,66 @@ UOptimusFunctionNodeGraph::UOptimusFunctionNodeGraph()
 {
 }
 
+void UOptimusFunctionNodeGraph::PostLoad()
+{
+	Super::PostLoad();
+
+	FOptimusObjectVersion::Type ObjectVersion = static_cast<FOptimusObjectVersion::Type>(GetLinkerCustomVersion(FOptimusObjectVersion::GUID));
+	if (ObjectVersion < FOptimusObjectVersion::FunctionGraphUseGuid)
+	{
+		check(!Guid.IsValid());
+		Guid = GetGuidForGraphWithoutGuid(TSoftObjectPtr(this));
+	}
+
+	check(Guid.IsValid());
+}
+
 FString UOptimusFunctionNodeGraph::GetNodeName() const
 {
 	return GetName();
 }
 
+void UOptimusFunctionNodeGraph::Init()
+{
+	Guid = FGuid::NewGuid();
+}
+
+FOptimusFunctionGraphIdentifier UOptimusFunctionNodeGraph::GetGraphIdentifier() const
+{
+	FOptimusFunctionGraphIdentifier GraphIdentifier;
+	GraphIdentifier.Asset = GetTypedOuter<UOptimusDeformer>();
+	
+	check(Guid.IsValid());
+	
+	GraphIdentifier.Guid = Guid;
+
+	return GraphIdentifier;
+}
+
 TArray<FName> UOptimusFunctionNodeGraph::GetAccessSpecifierOptions() const
 {
 	return {AccessSpecifierPublicName, AccessSpecifierPrivateName};
 }
 
-FOptimusFunctionNodeGraphHeader UOptimusFunctionNodeGraph::GetHeader() const
+FOptimusFunctionNodeGraphHeaderWithGuid UOptimusFunctionNodeGraph::GetHeaderWithGuid() const
 {
-	FOptimusFunctionNodeGraphHeader Header;
+	FOptimusFunctionNodeGraphHeaderWithGuid Header;
 
-	Header.GraphPath = const_cast<UOptimusFunctionNodeGraph*>(this);
+	Header.FunctionGraphGuid = Guid;
 	Header.FunctionName = GetFName();
 	Header.Category = Category;
 
-	return Header;
+	return Header;	
+}
+
+FGuid UOptimusFunctionNodeGraph::GetGuid() const
+{
+	return Guid;
+}
+
+FGuid UOptimusFunctionNodeGraph::GetGuidForGraphWithoutGuid(TSoftObjectPtr<UOptimusFunctionNodeGraph> InGraph)
+{
+	return FGuid::NewDeterministicGuid(InGraph.ToString());
 }
 
 
diff --git a/Engine/Plugins/Animation/DeformerGraph/Source/OptimusCore/Private/OptimusNodeGraph.cpp b/Engine/Plugins/Animation/DeformerGraph/Source/OptimusCore/Private/OptimusNodeGraph.cpp
--- a/Engine/Plugins/Animation/DeformerGraph/Source/OptimusCore/Private/OptimusNodeGraph.cpp
+++ b/Engine/Plugins/Animation/DeformerGraph/Source/OptimusCore/Private/OptimusNodeGraph.cpp
@@ -614,13 +614,13 @@ UOptimusNode* UOptimusNodeGraph::AddCommentNode(const FVector2D& InPosition, con
 		});	
 }
 
-UOptimusNode* UOptimusNodeGraph::AddFunctionReferenceNode(TSoftObjectPtr<UOptimusFunctionNodeGraph> InFunctionGraph, const FVector2D& InPosition)
+UOptimusNode* UOptimusNodeGraph::AddFunctionReferenceNode(UOptimusFunctionNodeGraph* InFunctionGraph, const FVector2D& InPosition)
 {
 	return AddNodeInternal(UOptimusNode_FunctionReference::StaticClass(), InPosition,
 	[InFunctionGraph](UOptimusNode *InNode)
 	{
 		UOptimusNode_FunctionReference* FunctionNode = CastChecked<UOptimusNode_FunctionReference>(InNode);
-		FunctionNode->InitializeSerializedGraphPath(InFunctionGraph.ToSoftObjectPath());
+		FunctionNode->SetReferencedFunctionGraph(InFunctionGraph->GetGraphIdentifier());
 	});
 }
 
@@ -1862,7 +1862,7 @@ bool UOptimusNodeGraph::ConvertToFunction(UOptimusNode* InSubGraphNode)
 				UOptimusFunctionNodeGraph* FunctionGraph = CastChecked<UOptimusFunctionNodeGraph>(InGraph);
 				FunctionGraph->InputBindings = SubGraph->InputBindings;
 				FunctionGraph->OutputBindings = SubGraph->OutputBindings;
-
+				FunctionGraph->Init();
 				return true;
 			});
 	
@@ -1922,7 +1922,7 @@ bool UOptimusNodeGraph::ConvertToFunction(UOptimusNode* InSubGraphNode)
 			{
 				UOptimusNode_FunctionReference* FunctionNode = Cast<UOptimusNode_FunctionReference>(InNode); 
 				UOptimusFunctionNodeGraph* FunctionGraph = Cast<UOptimusFunctionNodeGraph>(PathResolver->ResolveGraphPath(FunctionGraphPath));
-				FunctionNode->InitializeSerializedGraphPath(FunctionGraph);
+				FunctionNode->SetReferencedFunctionGraph(FunctionGraph->GetGraphIdentifier());
 				
 				return FunctionNode->SetGraphPositionDirect(ReferenceNodePosition);
 			});
diff --git a/Engine/Plugins/Animation/DeformerGraph/Source/OptimusCore/Public/IOptimusCoreModule.h b/Engine/Plugins/Animation/DeformerGraph/Source/OptimusCore/Public/IOptimusCoreModule.h
--- a/Engine/Plugins/Animation/DeformerGraph/Source/OptimusCore/Public/IOptimusCoreModule.h
+++ b/Engine/Plugins/Animation/DeformerGraph/Source/OptimusCore/Public/IOptimusCoreModule.h
@@ -28,8 +28,6 @@ public:
 		return RegisterDataInterfaceClass(T::StaticClass());
 	}
 	
-	virtual void UpdateFunctionReferences(const FSoftObjectPath& InOldGraphPath, const FSoftObjectPath& InNewGraphPath) = 0;
-	
 protected:
 	virtual bool RegisterDataInterfaceClass(TSubclassOf<UOptimusComputeDataInterface> InDataInterfaceClass) = 0;
 };
diff --git a/Engine/Plugins/Animation/DeformerGraph/Source/OptimusCore/Public/OptimusDeformer.h b/Engine/Plugins/Animation/DeformerGraph/Source/OptimusCore/Public/OptimusDeformer.h
--- a/Engine/Plugins/Animation/DeformerGraph/Source/OptimusCore/Public/OptimusDeformer.h
+++ b/Engine/Plugins/Animation/DeformerGraph/Source/OptimusCore/Public/OptimusDeformer.h
@@ -36,6 +36,7 @@ class UOptimusFunctionNodeGraph;
 enum class EOptimusDiagnosticLevel : uint8;
 struct FOptimusCompilerDiagnostic;
 struct FOptimusCompoundAction;
+struct FOptimusFunctionGraphIdentifier;
 
 DECLARE_MULTICAST_DELEGATE_OneParam(FOptimusCompileBegin, UOptimusDeformer *);
 DECLARE_MULTICAST_DELEGATE_OneParam(FOptimusCompileEnd, UOptimusDeformer *);
@@ -138,8 +139,8 @@ class OPTIMUSCORE_API UOptimusDeformer :
 
 public:
 
-	static const FName PublicFunctionsAssetTagName;
-	static const FName FunctionReferencesAssetTagName;
+	static const TCHAR* PublicFunctionsAssetTagName;
+	static const TCHAR* PublicFunctionsWithGuidAssetTagName;
 
 	
 	UOptimusDeformer();
@@ -184,6 +185,8 @@ public:
 
 	/// Returns all function graphs with the given access specifier. If InAccessSpecifier is None, it performs no filtering
 	TArray<UOptimusFunctionNodeGraph*> GetFunctionGraphs(FName InAccessSpecifier = NAME_None) const;
+
+	UOptimusFunctionNodeGraph* FindFunctionByGuid(FGuid InFunctionGraphGuid);
 	
 	// Variables
 	UOptimusVariableDescription* AddVariable(
@@ -534,6 +537,8 @@ private:
 	int32 GetUpdateGraphIndex() const;
 
 	TArray<UOptimusNode*> GetAllNodesOfClass(UClass* InNodeClass) const;
+	void OnGraphRenamedOrRemoved(UOptimusNodeGraph* InGraph) const;
+	void UpdateFunctionReferenceNodeDisplayName(const FOptimusFunctionGraphIdentifier& InRenamedFunction);
 	
 	/// Compile a node graph to a compute graph. Returns one or two complete compute graphs if compilation succeeded. 
 	FOptimusNodeGraphCompilationResult CompileNodeGraphToComputeGraphs(
diff --git a/Engine/Plugins/Animation/DeformerGraph/Source/OptimusCore/Public/OptimusFunctionNodeGraph.h b/Engine/Plugins/Animation/DeformerGraph/Source/OptimusCore/Public/OptimusFunctionNodeGraph.h
--- a/Engine/Plugins/Animation/DeformerGraph/Source/OptimusCore/Public/OptimusFunctionNodeGraph.h
+++ b/Engine/Plugins/Animation/DeformerGraph/Source/OptimusCore/Public/OptimusFunctionNodeGraph.h
@@ -9,7 +9,25 @@
 #include "OptimusFunctionNodeGraph.generated.h"
 
 
-struct FOptimusFunctionNodeGraphHeader;
+struct FOptimusFunctionNodeGraphHeaderWithGuid;
+
+USTRUCT()
+struct FOptimusFunctionGraphIdentifier
+{
+	GENERATED_BODY()
+
+	UPROPERTY()
+	TObjectPtr<UOptimusDeformer> Asset;
+
+	UPROPERTY()
+	FGuid Guid;
+
+	UOptimusFunctionNodeGraph* Resolve();
+	
+	bool operator==(const FOptimusFunctionGraphIdentifier& InOtherFunctionGraph) const = default;
+};
+
+
 /**
  * 
  */
@@ -24,10 +42,14 @@ public:
 	static FName AccessSpecifierPrivateName;
 
 	UOptimusFunctionNodeGraph();
+	void PostLoad() override;
 	
 	/** The name to give the node based off of this graph */
 	FString GetNodeName() const; 
 
+	void Init();
+	FOptimusFunctionGraphIdentifier GetGraphIdentifier() const;
+
 	/** The category of the node based of of this graph for listing purposes */ 
 	UPROPERTY(EditAnywhere, Category=Settings)
 	FName Category = UOptimusNode::CategoryName::Deformers;
@@ -38,12 +60,19 @@ public:
 	UFUNCTION()
 	TArray<FName> GetAccessSpecifierOptions() const;
 
-	FOptimusFunctionNodeGraphHeader GetHeader() const;
+	FOptimusFunctionNodeGraphHeaderWithGuid GetHeaderWithGuid() const;
+	FGuid GetGuid() const;
 
+	// Used for PostLoad fixup to ensure function references hold on to the same Guid that the graph
+	// would get when it calls PostLoad()
+	static FGuid GetGuidForGraphWithoutGuid(TSoftObjectPtr<UOptimusFunctionNodeGraph> InGraph);
 protected:
 	
 #if WITH_EDITOR
 	// UObject override
 	bool CanEditChange(const FProperty* InProperty) const override;
 #endif
+
+	UPROPERTY(VisibleAnywhere, Category=Debug, AdvancedDisplay)
+	FGuid Guid;
 };
diff --git a/Engine/Plugins/Animation/DeformerGraph/Source/OptimusCore/Public/OptimusFunctionNodeGraphHeaderWithGuid.h b/Engine/Plugins/Animation/DeformerGraph/Source/OptimusCore/Public/OptimusFunctionNodeGraphHeaderWithGuid.h
new file mode 100644
--- /dev/null
+++ b/Engine/Plugins/Animation/DeformerGraph/Source/OptimusCore/Public/OptimusFunctionNodeGraphHeaderWithGuid.h
@@ -0,0 +1,34 @@
+﻿// Copyright Epic Games, Inc. All Rights Reserved.
+
+#pragma once
+
+#include "CoreMinimal.h"
+
+#include "OptimusFunctionNodeGraphHeaderWithGuid.generated.h"
+
+class UOptimusNode_FunctionReference;
+
+USTRUCT()
+struct OPTIMUSCORE_API FOptimusFunctionNodeGraphHeaderWithGuid
+{
+	GENERATED_BODY()
+
+	UPROPERTY()
+	FGuid FunctionGraphGuid;
+
+	UPROPERTY()
+	FName FunctionName;
+
+	UPROPERTY()
+	FName Category;
+};
+
+USTRUCT()
+struct OPTIMUSCORE_API FOptimusFunctionNodeGraphHeaderWithGuidArray
+{
+	GENERATED_BODY()
+
+	UPROPERTY()
+	TArray<FOptimusFunctionNodeGraphHeaderWithGuid> Headers;
+};
+
diff --git a/Engine/Plugins/Animation/DeformerGraph/Source/OptimusCore/Public/OptimusNodeGraph.h b/Engine/Plugins/Animation/DeformerGraph/Source/OptimusCore/Public/OptimusNodeGraph.h
--- a/Engine/Plugins/Animation/DeformerGraph/Source/OptimusCore/Public/OptimusNodeGraph.h
+++ b/Engine/Plugins/Animation/DeformerGraph/Source/OptimusCore/Public/OptimusNodeGraph.h
@@ -169,7 +169,7 @@ public:
 
 	UFUNCTION(BlueprintCallable, Category = OptimusNodeGraph)
 	UOptimusNode* AddFunctionReferenceNode(
-		TSoftObjectPtr<UOptimusFunctionNodeGraph> InFunctionGraph,
+		UOptimusFunctionNodeGraph* InFunctionGraph,
 		const FVector2D& InPosition
 	);
 
diff --git a/Engine/Plugins/Animation/DeformerGraph/Source/OptimusCore/Public/OptimusObjectVersion.h b/Engine/Plugins/Animation/DeformerGraph/Source/OptimusCore/Public/OptimusObjectVersion.h
--- a/Engine/Plugins/Animation/DeformerGraph/Source/OptimusCore/Public/OptimusObjectVersion.h
+++ b/Engine/Plugins/Animation/DeformerGraph/Source/OptimusCore/Public/OptimusObjectVersion.h
@@ -25,6 +25,7 @@ struct OPTIMUSCORE_API FOptimusObjectVersion
 		SkinnedMeshWriteDIColorBufferManualFetchSwizzle,
 		PropertyPinSupport,
 		KernelThreadIndexUsingTotalThreadCount,
+		FunctionGraphUseGuid,
 
 		// -----<new versions can be added above this line>-------------------------------------------------
 		VersionPlusOne,
diff --git a/Engine/Plugins/Animation/DeformerGraph/Source/OptimusEditor/Private/OptimusEditorGraphSchema.cpp b/Engine/Plugins/Animation/DeformerGraph/Source/OptimusEditor/Private/OptimusEditorGraphSchema.cpp
--- a/Engine/Plugins/Animation/DeformerGraph/Source/OptimusEditor/Private/OptimusEditorGraphSchema.cpp
+++ b/Engine/Plugins/Animation/DeformerGraph/Source/OptimusEditor/Private/OptimusEditorGraphSchema.cpp
@@ -22,6 +22,7 @@
 #include "OptimusDeformer.h"
 #include "OptimusEditorGraphConnectionDrawingPolicy.h"
 #include "OptimusFunctionNodeGraph.h"
+#include "OptimusFunctionNodeGraphHeaderWithGuid.h"
 #include "AssetRegistry/AssetRegistryModule.h"
 #include "Styling/SlateIconFinder.h"
 
@@ -176,7 +177,7 @@ void UOptimusEditorGraphSchema::GetGraphActions(
 				{
 					for (UOptimusFunctionNodeGraph* FunctionNodeGraph : Deformer->GetFunctionGraphs(UOptimusFunctionNodeGraph::AccessSpecifierPrivateName))
 					{
-						FOptimusFunctionNodeGraphHeader Header = FunctionNodeGraph->GetHeader();
+						FOptimusFunctionNodeGraphHeaderWithGuid Header = FunctionNodeGraph->GetHeaderWithGuid();
 				
 						TSharedPtr< FOptimusGraphSchemaAction_NewFunctionReferenceNode> Action(
 									new FOptimusGraphSchemaAction_NewFunctionReferenceNode(
@@ -185,7 +186,8 @@ void UOptimusEditorGraphSchema::GetGraphActions(
 										/* Tooltip */{}, 0, /* Keywords */{}
 								));
 
-						Action->GraphPath = Header.GraphPath;
+						Action->AssetPath = Deformer;
+						Action->FunctionGraphGuid = Header.FunctionGraphGuid;
 
 						IoActionBuilder.AddAction(Action);
 					}
@@ -203,34 +205,66 @@ void UOptimusEditorGraphSchema::GetGraphActions(
 		// This triggers the gathering of asset tags across assets
 		AssetRegistryModule.Get().GetAssetsByClass(UOptimusDeformer::StaticClass()->GetClassPathName(), AssetDatas);
 
-		for (FAssetData& AssetData : AssetDatas)
+		auto AddPublicFunctionAction = [&IoActionBuilder](const FAssetData& AssetData, const FOptimusFunctionNodeGraphHeaderWithGuid& Header)
 		{
-			FString PublicFunctionString = AssetData.GetTagValueRef<FString>(UOptimusDeformer::PublicFunctionsAssetTagName);
+			const FText FunctionName = FText::FromName(Header.FunctionName);
 
-			if (!PublicFunctionString.IsEmpty())
-			{
-				FOptimusFunctionNodeGraphHeaderArray PublicFunctionHeaderArray;
+			const FText Category = FText::FromName(Header.Category);
 
-				FOptimusFunctionNodeGraphHeaderArray::StaticStruct()->ImportText(*PublicFunctionString, &PublicFunctionHeaderArray, nullptr, PPF_None, nullptr, {});
+			TSharedPtr< FOptimusGraphSchemaAction_NewFunctionReferenceNode> Action(
+				new FOptimusGraphSchemaAction_NewFunctionReferenceNode(
+					Category,
+					FunctionName,
+					/* Tooltip */{}, 0, /* Keywords */{}
+			));
 
-				for (const FOptimusFunctionNodeGraphHeader& Header : PublicFunctionHeaderArray.Headers)
-				{
-					const FText FunctionName = FText::FromName(Header.FunctionName);
+			Action->AssetPath = AssetData.GetSoftObjectPath();
+			Action->FunctionGraphGuid = Header.FunctionGraphGuid;
 
-					const FText Category = FText::FromName(Header.Category);
+			IoActionBuilder.AddAction(Action);	
+		};
+		
+		for (const FAssetData& AssetData : AssetDatas)
+		{
+			// Old public function exports
+			{
+				FString PublicFunctionString = AssetData.GetTagValueRef<FString>(UOptimusDeformer::PublicFunctionsAssetTagName);
 
-					TSharedPtr< FOptimusGraphSchemaAction_NewFunctionReferenceNode> Action(
-						new FOptimusGraphSchemaAction_NewFunctionReferenceNode(
-							Category,
-							FunctionName,
-							/* Tooltip */{}, 0, /* Keywords */{}
-					));
+				if (!PublicFunctionString.IsEmpty())
+				{
+					FOptimusFunctionNodeGraphHeaderArray PublicFunctionHeaderArray;
 
-					Action->GraphPath = Header.GraphPath;
+					FOptimusFunctionNodeGraphHeaderArray::StaticStruct()->ImportText(*PublicFunctionString, &PublicFunctionHeaderArray, nullptr, PPF_None, nullptr, {});
 
-					IoActionBuilder.AddAction(Action);	
-				}
+					for (const FOptimusFunctionNodeGraphHeader& Header : PublicFunctionHeaderArray.Headers)
+					{
+
+						// Unfortunately we need to do a load here because we need to grab the latest Guid from the function graph
+						if (UOptimusFunctionNodeGraph* FunctionNodeGraph = Header.GraphPath.LoadSynchronous())
+						{
+							FOptimusFunctionNodeGraphHeaderWithGuid HeaderWithGuid = FunctionNodeGraph->GetHeaderWithGuid();
+							AddPublicFunctionAction(AssetData, HeaderWithGuid);
+						}
+					}
+				}	
 			}
+
+			// New public function exports
+			{
+				FString PublicFunctionString = AssetData.GetTagValueRef<FString>(UOptimusDeformer::PublicFunctionsAssetTagName);
+
+				if (!PublicFunctionString.IsEmpty())
+				{
+					FOptimusFunctionNodeGraphHeaderWithGuidArray PublicFunctionHeaderArray;
+
+					FOptimusFunctionNodeGraphHeaderWithGuidArray::StaticStruct()->ImportText(*PublicFunctionString, &PublicFunctionHeaderArray, nullptr, PPF_None, nullptr, {});
+
+					for (const FOptimusFunctionNodeGraphHeaderWithGuid& Header : PublicFunctionHeaderArray.Headers)
+					{
+						AddPublicFunctionAction(AssetData, Header);	
+					}
+				}	
+			}	
 		}	
 	}
 	
diff --git a/Engine/Plugins/Animation/DeformerGraph/Source/OptimusEditor/Private/OptimusEditorGraphSchemaActions.cpp b/Engine/Plugins/Animation/DeformerGraph/Source/OptimusEditor/Private/OptimusEditorGraphSchemaActions.cpp
--- a/Engine/Plugins/Animation/DeformerGraph/Source/OptimusEditor/Private/OptimusEditorGraphSchemaActions.cpp
+++ b/Engine/Plugins/Animation/DeformerGraph/Source/OptimusEditor/Private/OptimusEditorGraphSchemaActions.cpp
@@ -3,6 +3,7 @@
 #include "OptimusEditorGraphSchemaActions.h"
 
 #include "OptimusComponentSource.h"
+#include "OptimusDeformer.h"
 #include "OptimusEditorGraph.h"
 #include "OptimusEditorGraphNode.h"
 #include "OptimusEditorHelpers.h"
@@ -145,16 +146,23 @@ UEdGraphNode* FOptimusGraphSchemaAction_NewFunctionReferenceNode::PerformAction(
 	
 	if (ensure(Graph != nullptr))
 	{
-		UOptimusNode* ModelNode = Graph->GetModelGraph()->AddFunctionReferenceNode(GraphPath, InLocation);
+		if (UOptimusDeformer* Deformer = AssetPath.LoadSynchronous())
+		{
+			if (UOptimusFunctionNodeGraph* FunctionGraph = Deformer->FindFunctionByGuid(FunctionGraphGuid))
+			{
+				UOptimusNode* ModelNode = Graph->GetModelGraph()->AddFunctionReferenceNode(FunctionGraph, InLocation);
 
-		// FIXME: Automatic connection from the given pin.
+				// FIXME: Automatic connection from the given pin.
 
-		UEdGraphNode* GraphNode = Graph->FindGraphNodeFromModelNode(ModelNode);
-		if (GraphNode && bInSelectNewNode)
-		{
-			Graph->SelectNodeSet({GraphNode});
+				UEdGraphNode* GraphNode = Graph->FindGraphNodeFromModelNode(ModelNode);
+				if (GraphNode && bInSelectNewNode)
+				{
+					Graph->SelectNodeSet({GraphNode});
+				}
+				return GraphNode;	
+			}
 		}
-		return GraphNode;
+	
 	}
 
 	return nullptr;
diff --git a/Engine/Plugins/Animation/DeformerGraph/Source/OptimusEditor/Private/OptimusEditorGraphSchemaActions.h b/Engine/Plugins/Animation/DeformerGraph/Source/OptimusEditor/Private/OptimusEditorGraphSchemaActions.h
--- a/Engine/Plugins/Animation/DeformerGraph/Source/OptimusEditor/Private/OptimusEditorGraphSchemaActions.h
+++ b/Engine/Plugins/Animation/DeformerGraph/Source/OptimusEditor/Private/OptimusEditorGraphSchemaActions.h
@@ -7,6 +7,7 @@
 
 #include "OptimusComputeDataInterface.h"
 #include "OptimusDataType.h"
+#include "OptimusFunctionNodeGraph.h"
 #include "OptimusFunctionNodeGraphHeader.h"
 
 #include "OptimusEditorGraphSchemaActions.generated.h"
@@ -128,7 +129,9 @@ struct FOptimusGraphSchemaAction_NewFunctionReferenceNode:
 	// Inherit the base class's constructors
 	using FEdGraphSchemaAction::FEdGraphSchemaAction;
 
-	TSoftObjectPtr<UOptimusFunctionNodeGraph> GraphPath;
+	TSoftObjectPtr<UOptimusDeformer> AssetPath;
+
+	FGuid FunctionGraphGuid;
 
 	static FName StaticGetTypeId() { static FName Type("FOptimusGraphSchemaAction_NewFunctionReferenceNode"); return Type; }
 	FName GetTypeId() const override { return StaticGetTypeId(); }
