[Deformer graph] fix minor compilation error

diff --git a/Engine/Plugins/Animation/DeformerGraph/Source/OptimusCore/Public/OptimusFunctionNodeGraph.h b/Engine/Plugins/Animation/DeformerGraph/Source/OptimusCore/Public/OptimusFunctionNodeGraph.h
--- a/Engine/Plugins/Animation/DeformerGraph/Source/OptimusCore/Public/OptimusFunctionNodeGraph.h
+++ b/Engine/Plugins/Animation/DeformerGraph/Source/OptimusCore/Public/OptimusFunctionNodeGraph.h
@@ -4,7 +4,7 @@
 
 #include "OptimusNode.h"
 #include "OptimusNodeSubGraph.h"
-
+#include "OptimusDeformer.h"
 
 #include "OptimusFunctionNodeGraph.generated.h"
 
