[Deformer Graph] fixed android specific compilation error

diff --git a/Engine/Plugins/Animation/DeformerGraph/Source/OptimusCore/Private/OptimusFunctionNodeGraph.cpp b/Engine/Plugins/Animation/DeformerGraph/Source/OptimusCore/Private/OptimusFunctionNodeGraph.cpp
--- a/Engine/Plugins/Animation/DeformerGraph/Source/OptimusCore/Private/OptimusFunctionNodeGraph.cpp
+++ b/Engine/Plugins/Animation/DeformerGraph/Source/OptimusCore/Private/OptimusFunctionNodeGraph.cpp
@@ -23,7 +23,7 @@ void UOptimusFunctionNodeGraph::PostLoad()
 	if (ObjectVersion < FOptimusObjectVersion::FunctionGraphUseGuid)
 	{
 		check(!Guid.IsValid());
-		Guid = GetGuidForGraphWithoutGuid(TSoftObjectPtr(this));
+		Guid = GetGuidForGraphWithoutGuid(TSoftObjectPtr<UOptimusFunctionNodeGraph>(this));
 	}
 
 	check(Guid.IsValid());
