Add an optional default value for the SkinnedMeshVertexAttribute node of the deformer graph to be used if no buffer is found.
#tests Confirmed existing nodes behave as before (using the default of 0) and updated a node and played with setting the default value on the node or by feeding in a constant

diff --git a/Engine/Plugins/Animation/DeformerGraph/Shaders/Private/DataInterfaceSkinnedMeshVertexAttribute.ush b/Engine/Plugins/Animation/DeformerGraph/Shaders/Private/DataInterfaceSkinnedMeshVertexAttribute.ush
--- a/Engine/Plugins/Animation/DeformerGraph/Shaders/Private/DataInterfaceSkinnedMeshVertexAttribute.ush
+++ b/Engine/Plugins/Animation/DeformerGraph/Shaders/Private/DataInterfaceSkinnedMeshVertexAttribute.ush
@@ -2,6 +2,7 @@
 
 uint {DataInterfaceName}_NumVertices;
 uint {DataInterfaceName}_bIsValid;
+float {DataInterfaceName}_DefaultValue;
 Buffer<float> {DataInterfaceName}_ValueBuffer;
 
 uint ReadNumVertices_{DataInterfaceName}()
@@ -18,6 +19,6 @@ float ReadValue_{DataInterfaceName}(uint VertexIndex)
 	}
 	else
 	{
-		return 0.0f;
+		return {DataInterfaceName}_DefaultValue;
 	}
 }
diff --git a/Engine/Plugins/Animation/DeformerGraph/Source/OptimusCore/Private/DataInterfaces/OptimusDataInterfaceSkinnedMeshVertexAttribute.cpp b/Engine/Plugins/Animation/DeformerGraph/Source/OptimusCore/Private/DataInterfaces/OptimusDataInterfaceSkinnedMeshVertexAttribute.cpp
--- a/Engine/Plugins/Animation/DeformerGraph/Source/OptimusCore/Private/DataInterfaces/OptimusDataInterfaceSkinnedMeshVertexAttribute.cpp
+++ b/Engine/Plugins/Animation/DeformerGraph/Source/OptimusCore/Private/DataInterfaces/OptimusDataInterfaceSkinnedMeshVertexAttribute.cpp
@@ -18,6 +18,11 @@ FName UOptimusSkinnedMeshVertexAttributeDataInterface::GetAttributeNamePropertyN
 	return GET_MEMBER_NAME_CHECKED(UOptimusSkinnedMeshVertexAttributeDataInterface, AttributeName);
 }
 
+FName UOptimusSkinnedMeshVertexAttributeDataInterface::GetDefaultValuePropertyName()
+{
+	return GET_MEMBER_NAME_CHECKED(UOptimusSkinnedMeshVertexAttributeDataInterface, DefaultValue);
+}
+
 FString UOptimusSkinnedMeshVertexAttributeDataInterface::GetDisplayName() const
 {
 	return TEXT("Skinned Mesh Vertex Attribute");
@@ -37,6 +42,7 @@ TArray<FOptimusCDIPropertyPinDefinition> UOptimusSkinnedMeshVertexAttributeDataI
 	TArray<FOptimusCDIPropertyPinDefinition> Defs;
 	const FOptimusDataTypeHandle NameType = FOptimusDataTypeRegistry::Get().FindType(*FNameProperty::StaticClass());
 	Defs.Add({GetAttributeNamePropertyName(), NameType});
+	Defs.Add({GetDefaultValuePropertyName(), FOptimusDataTypeRegistry::Get().FindType(*FFloatProperty::StaticClass())});
 
 	return Defs;
 }
@@ -63,6 +69,7 @@ void UOptimusSkinnedMeshVertexAttributeDataInterface::GetSupportedInputs(TArray<
 BEGIN_SHADER_PARAMETER_STRUCT(FSkinnedMeshVertexAttributeDataInterfaceParameters, )
 	SHADER_PARAMETER(uint32, NumVertices)
 	SHADER_PARAMETER(uint32, bIsValid)
+	SHADER_PARAMETER(float, DefaultValue)
 	SHADER_PARAMETER_SRV(Buffer<float>, ValueBuffer)
 END_SHADER_PARAMETER_STRUCT()
 
@@ -105,6 +112,7 @@ UComputeDataProvider* UOptimusSkinnedMeshVertexAttributeDataInterface::CreateDat
 	UOptimusSkinnedMeshVertexAttributeDataProvider* Provider = NewObject<UOptimusSkinnedMeshVertexAttributeDataProvider>();
 	Provider->SkinnedMeshComponent = Cast<USkinnedMeshComponent>(InBinding);
 	Provider->AttributeName = AttributeName;
+	Provider->DefaultValue = DefaultValue;
 	Provider->WeakDataInterface = this;
 
 	return Provider;
@@ -130,20 +138,36 @@ FComputeDataProviderRenderProxy* UOptimusSkinnedMeshVertexAttributeDataProvider:
 {
 	if (const UOptimusSkinnedMeshVertexAttributeDataInterface* DataInterface = WeakDataInterface.Get())
 	{
-		FOptimusValueContainerStruct ValueContainer =
-			DeformerInstance->GetDataInterfacePropertyOverride(
-				DataInterface,
-				UOptimusSkinnedMeshVertexAttributeDataInterface::GetAttributeNamePropertyName()
-				);
+		{
+			FOptimusValueContainerStruct ValueContainer =
+				DeformerInstance->GetDataInterfacePropertyOverride(
+					DataInterface,
+					UOptimusSkinnedMeshVertexAttributeDataInterface::GetAttributeNamePropertyName()
+					);
+
+			TValueOrError<FName, EPropertyBagResult> Value = ValueContainer.Value.GetValueName(FOptimusValueContainerStruct::ValuePropertyName);
+			if (Value.HasValue())
+			{
+				AttributeName = Value.GetValue();
+			}
+		}
 
-		TValueOrError<FName, EPropertyBagResult> Value = ValueContainer.Value.GetValueName(FOptimusValueContainerStruct::ValuePropertyName);
-		if (Value.HasValue())
 		{
-			AttributeName = Value.GetValue();
+			FOptimusValueContainerStruct ValueContainer =
+				DeformerInstance->GetDataInterfacePropertyOverride(
+					DataInterface,
+					UOptimusSkinnedMeshVertexAttributeDataInterface::GetDefaultValuePropertyName()
+				);
+
+			TValueOrError<float, EPropertyBagResult> Value = ValueContainer.Value.GetValueFloat(FOptimusValueContainerStruct::ValuePropertyName);
+			if (Value.HasValue())
+			{
+				DefaultValue = Value.GetValue();
+			}
 		}
 	}
 	
-	return new FOptimusSkinnedMeshVertexAttributeDataProviderProxy(SkinnedMeshComponent, AttributeName);
+	return new FOptimusSkinnedMeshVertexAttributeDataProviderProxy(SkinnedMeshComponent, AttributeName, DefaultValue);
 }
 
 void UOptimusSkinnedMeshVertexAttributeDataProvider::SetDeformerInstance(UOptimusDeformerInstance* InInstance)
@@ -154,11 +178,13 @@ void UOptimusSkinnedMeshVertexAttributeDataProvider::SetDeformerInstance(UOptimu
 
 FOptimusSkinnedMeshVertexAttributeDataProviderProxy::FOptimusSkinnedMeshVertexAttributeDataProviderProxy(
 	USkinnedMeshComponent* InSkinnedMeshComponent, 
-	FName InAttributeName
+	FName InAttributeName,
+	float InDefaultValue
 	)
 {
 	SkeletalMeshObject = InSkinnedMeshComponent->MeshObject;
 	AttributeName = InAttributeName;
+	DefaultValue = InDefaultValue;
 }
 
 
@@ -202,5 +228,6 @@ void FOptimusSkinnedMeshVertexAttributeDataProviderProxy::GatherDispatchData(
 		Parameters.NumVertices = LodRenderData->GetNumVertices();
 		Parameters.ValueBuffer = AttributeSRV;
 		Parameters.bIsValid = AttributeBuffer != nullptr;
+		Parameters.DefaultValue = DefaultValue;
 	}
 }
diff --git a/Engine/Plugins/Animation/DeformerGraph/Source/OptimusCore/Private/DataInterfaces/OptimusDataInterfaceSkinnedMeshVertexAttribute.h b/Engine/Plugins/Animation/DeformerGraph/Source/OptimusCore/Private/DataInterfaces/OptimusDataInterfaceSkinnedMeshVertexAttribute.h
--- a/Engine/Plugins/Animation/DeformerGraph/Source/OptimusCore/Private/DataInterfaces/OptimusDataInterfaceSkinnedMeshVertexAttribute.h
+++ b/Engine/Plugins/Animation/DeformerGraph/Source/OptimusCore/Private/DataInterfaces/OptimusDataInterfaceSkinnedMeshVertexAttribute.h
@@ -39,9 +39,13 @@ class UOptimusSkinnedMeshVertexAttributeDataInterface :
 
 	UPROPERTY(EditAnywhere, Category="Vertex Attribute")
 	FName AttributeName;
+
+	UPROPERTY(EditAnywhere, Category="Vertex Attribute")
+	float DefaultValue;
 private:
 	friend class UOptimusSkinnedMeshVertexAttributeDataProvider;
 	static FName GetAttributeNamePropertyName();
+	static FName GetDefaultValuePropertyName();
 };
 
 
@@ -60,6 +64,9 @@ public:
 	UPROPERTY(EditAnywhere, BlueprintReadWrite, Category = VertexAttribute)
 	FName AttributeName;
 
+	UPROPERTY(EditAnywhere, BlueprintReadWrite, Category = VertexAttribute)
+	float DefaultValue;
+
     TWeakObjectPtr<const UOptimusSkinnedMeshVertexAttributeDataInterface> WeakDataInterface;
 	
 	//~ Begin UComputeDataProvider Interface
@@ -80,7 +87,8 @@ class FOptimusSkinnedMeshVertexAttributeDataProviderProxy : public FComputeDataP
 public:
 	FOptimusSkinnedMeshVertexAttributeDataProviderProxy(
 		USkinnedMeshComponent* InSkinnedMeshComponent,
-		FName InAttributeName);
+		FName InAttributeName,
+		float InDefaultValue);
 
 	//~ Begin FComputeDataProviderRenderProxy Interface
 	bool IsValid(FValidationData const& InValidationData) const override;
@@ -92,4 +100,5 @@ private:
 	
 	FSkeletalMeshObject* SkeletalMeshObject;
 	FName AttributeName;
+	float DefaultValue;
 };
