[Deformer Graph] fixed typo causing public functions to not show up in right click menu

diff --git a/Engine/Plugins/Animation/DeformerGraph/Source/OptimusEditor/Private/OptimusEditorGraphSchema.cpp b/Engine/Plugins/Animation/DeformerGraph/Source/OptimusEditor/Private/OptimusEditorGraphSchema.cpp
--- a/Engine/Plugins/Animation/DeformerGraph/Source/OptimusEditor/Private/OptimusEditorGraphSchema.cpp
+++ b/Engine/Plugins/Animation/DeformerGraph/Source/OptimusEditor/Private/OptimusEditorGraphSchema.cpp
@@ -251,7 +251,7 @@ void UOptimusEditorGraphSchema::GetGraphActions(
 
 			// New public function exports
 			{
-				FString PublicFunctionString = AssetData.GetTagValueRef<FString>(UOptimusDeformer::PublicFunctionsAssetTagName);
+				FString PublicFunctionString = AssetData.GetTagValueRef<FString>(UOptimusDeformer::PublicFunctionsWithGuidAssetTagName);
 
 				if (!PublicFunctionString.IsEmpty())
 				{
