[Deformer Graph] fixed crash on adding a function reference node into the graph, making sure we don't send notification to slate during node construction

diff --git a/Engine/Plugins/Animation/DeformerGraph/Source/OptimusCore/Private/OptimusNode.cpp b/Engine/Plugins/Animation/DeformerGraph/Source/OptimusCore/Private/OptimusNode.cpp
--- a/Engine/Plugins/Animation/DeformerGraph/Source/OptimusCore/Private/OptimusNode.cpp
+++ b/Engine/Plugins/Animation/DeformerGraph/Source/OptimusCore/Private/OptimusNode.cpp
@@ -416,18 +416,24 @@ FName UOptimusNode::GetAvailablePinNameStable(const UObject* InNodeOrPin, FName
 
 void UOptimusNode::PostCreateNode()
 {
+	TGuardValue<bool> NodeConstructionGuard(bConstructingNode, true);
 	InitializeTransientData();
 	
 	CachedPinLookup.Empty();
+
 	Pins.Empty();
+	
+	ConstructNode();
 
-	{
-		TGuardValue<bool> NodeConstructionGuard(bConstructingNode, true);
-		ConstructNode();
+	// Update the pins display names
+	InitializePinsDisplayName();
+}
 
-		// Update the pins display names
-		InitializePinsDisplayName();
-	}
+void UOptimusNode::PostDeserializeNode()
+{
+	TGuardValue<bool> NodeConstructionGuard(bConstructingNode, true);
+	InitializeTransientData();
+	InitializePinsDisplayName();	
 }
 
 
@@ -446,11 +452,8 @@ void UOptimusNode::PostLoad()
 	// PostLoad fix up pass
 	PostLoadNodeSpecificData();
 
-	// Now that all serialized data is ready, let's generate any transient data
-	InitializeTransientData();
-
-	// Update the pins display names
-	InitializePinsDisplayName();
+	// Initialize transient data for nodes and pins
+	PostDeserializeNode();
 }
 
 void UOptimusNode::PostLoadNodeSpecificData()
diff --git a/Engine/Plugins/Animation/DeformerGraph/Source/OptimusEditor/Private/OptimusEditorClipboard.cpp b/Engine/Plugins/Animation/DeformerGraph/Source/OptimusEditor/Private/OptimusEditorClipboard.cpp
--- a/Engine/Plugins/Animation/DeformerGraph/Source/OptimusEditor/Private/OptimusEditorClipboard.cpp
+++ b/Engine/Plugins/Animation/DeformerGraph/Source/OptimusEditor/Private/OptimusEditorClipboard.cpp
@@ -420,8 +420,7 @@ bool FOptimusEditorClipboard::ProcessPostCreateObject(UObject* InRootOuter, UObj
 		{
 			UOptimusNode* Node = Cast<UOptimusNode>(InNewObject);
 
-			Node->InitializeTransientData();
-			Node->InitializePinsDisplayName();
+			Node->PostDeserializeNode();
 			
 			return Graph->AddNodeDirect(Node);
 		}
@@ -449,8 +448,7 @@ bool FOptimusEditorClipboard::ProcessPostCreateObject(UObject* InRootOuter, UObj
 			{
 				for (UOptimusNode* Node : WorkingSubGraph->GetAllNodes())
 				{
-					Node->InitializeTransientData();
-					Node->InitializePinsDisplayName();
+					Node->PostDeserializeNode();
 				}
 
 				for (UOptimusNodeGraph* SubGraphToAdd : WorkingSubGraph->GetGraphs())
