[Deformer Graph] added missing header file change

diff --git a/Engine/Plugins/Animation/DeformerGraph/Source/OptimusCore/Public/OptimusNode.h b/Engine/Plugins/Animation/DeformerGraph/Source/OptimusCore/Public/OptimusNode.h
--- a/Engine/Plugins/Animation/DeformerGraph/Source/OptimusCore/Public/OptimusNode.h
+++ b/Engine/Plugins/Animation/DeformerGraph/Source/OptimusCore/Public/OptimusNode.h
@@ -198,6 +198,12 @@ public:
 	 */
 	void PostCreateNode();
 
+	/**
+	 * Called just after the node is deserialized, either via loading or pasting.
+	 * It does not recreate the pins, instead it only initializes transient data.
+	 */
+	void PostDeserializeNode();
+
 	//== UObject overrides
 	void Serialize(FArchive& Ar) override;
 	
