[Deformer Graph] New Property pin may get added to data interfaces, so older data interface may not have the new pins created. No need to enforce that the pin has to be present, if the pin was not created, it simply cannot be overriden by the user in the graph.

diff --git a/Engine/Plugins/Animation/DeformerGraph/Source/OptimusCore/Private/Nodes/OptimusNode_DataInterface.cpp b/Engine/Plugins/Animation/DeformerGraph/Source/OptimusCore/Private/Nodes/OptimusNode_DataInterface.cpp
--- a/Engine/Plugins/Animation/DeformerGraph/Source/OptimusCore/Private/Nodes/OptimusNode_DataInterface.cpp
+++ b/Engine/Plugins/Animation/DeformerGraph/Source/OptimusCore/Private/Nodes/OptimusNode_DataInterface.cpp
@@ -392,7 +392,9 @@ TArray<UOptimusNodePin*> UOptimusNode_DataInterface::GetPropertyPins() const
 	for (const FOptimusCDIPropertyPinDefinition& Definition : PropertyPinDefinitions)
 	{
 		UOptimusNodePin* Pin = FindPinFromPath({Definition.PinName});
-		if (ensure(Pin) && ensure(Pin->GetDirection() == EOptimusNodePinDirection::Input))
+		// No need to ensure a pin is found here since older version of the data interface does not attempt to create
+		// property pins that were added later. 
+		if (Pin && Pin->GetDirection() == EOptimusNodePinDirection::Input)
 		{
 			Result.Add(Pin);
 		}
