[Deformer Graph] fixed duplicated functions in node graph right click menui

diff --git a/Engine/Plugins/Animation/DeformerGraph/Source/OptimusCore/Private/Nodes/OptimusNode_FunctionReference.cpp b/Engine/Plugins/Animation/DeformerGraph/Source/OptimusCore/Private/Nodes/OptimusNode_FunctionReference.cpp
--- a/Engine/Plugins/Animation/DeformerGraph/Source/OptimusCore/Private/Nodes/OptimusNode_FunctionReference.cpp
+++ b/Engine/Plugins/Animation/DeformerGraph/Source/OptimusCore/Private/Nodes/OptimusNode_FunctionReference.cpp
@@ -17,7 +17,12 @@
 
 UOptimusFunctionNodeGraph* FOptimusFunctionGraphIdentifier::Resolve()
 {
-	return Asset->FindFunctionByGuid(Guid);
+	if (Asset)
+	{
+		return Asset->FindFunctionByGuid(Guid);
+	}
+
+	return nullptr;
 }
 
 FName UOptimusNode_FunctionReference::GetNodeCategory() const
diff --git a/Engine/Plugins/Animation/DeformerGraph/Source/OptimusEditor/Private/OptimusEditorGraphSchema.cpp b/Engine/Plugins/Animation/DeformerGraph/Source/OptimusEditor/Private/OptimusEditorGraphSchema.cpp
--- a/Engine/Plugins/Animation/DeformerGraph/Source/OptimusEditor/Private/OptimusEditorGraphSchema.cpp
+++ b/Engine/Plugins/Animation/DeformerGraph/Source/OptimusEditor/Private/OptimusEditorGraphSchema.cpp
@@ -226,12 +226,33 @@ void UOptimusEditorGraphSchema::GetGraphActions(
 		
 		for (const FAssetData& AssetData : AssetDatas)
 		{
+			bool bPublicFunctionImported = false;
+			// New public function exports
+			{
+				FString PublicFunctionString = AssetData.GetTagValueRef<FString>(UOptimusDeformer::PublicFunctionsWithGuidAssetTagName);
+
+				if (!PublicFunctionString.IsEmpty())
+				{
+					bPublicFunctionImported = true;
+					FOptimusFunctionNodeGraphHeaderWithGuidArray PublicFunctionHeaderArray;
+
+					FOptimusFunctionNodeGraphHeaderWithGuidArray::StaticStruct()->ImportText(*PublicFunctionString, &PublicFunctionHeaderArray, nullptr, PPF_None, nullptr, {});
+
+					for (const FOptimusFunctionNodeGraphHeaderWithGuid& Header : PublicFunctionHeaderArray.Headers)
+					{
+						AddPublicFunctionAction(AssetData, Header);	
+					}
+				}	
+			}
+			
 			// Old public function exports
+			if (!bPublicFunctionImported)
 			{
 				FString PublicFunctionString = AssetData.GetTagValueRef<FString>(UOptimusDeformer::PublicFunctionsAssetTagName);
 
 				if (!PublicFunctionString.IsEmpty())
 				{
+					bPublicFunctionImported = true;
 					FOptimusFunctionNodeGraphHeaderArray PublicFunctionHeaderArray;
 
 					FOptimusFunctionNodeGraphHeaderArray::StaticStruct()->ImportText(*PublicFunctionString, &PublicFunctionHeaderArray, nullptr, PPF_None, nullptr, {});
@@ -249,22 +270,7 @@ void UOptimusEditorGraphSchema::GetGraphActions(
 				}	
 			}
 
-			// New public function exports
-			{
-				FString PublicFunctionString = AssetData.GetTagValueRef<FString>(UOptimusDeformer::PublicFunctionsWithGuidAssetTagName);
-
-				if (!PublicFunctionString.IsEmpty())
-				{
-					FOptimusFunctionNodeGraphHeaderWithGuidArray PublicFunctionHeaderArray;
-
-					FOptimusFunctionNodeGraphHeaderWithGuidArray::StaticStruct()->ImportText(*PublicFunctionString, &PublicFunctionHeaderArray, nullptr, PPF_None, nullptr, {});
-
-					for (const FOptimusFunctionNodeGraphHeaderWithGuid& Header : PublicFunctionHeaderArray.Headers)
-					{
-						AddPublicFunctionAction(AssetData, Header);	
-					}
-				}	
-			}	
+			
 		}	
 	}
 	
