[Deformer Graph] fixed bad kernel dispatch order when a copy is required from a set/get resource node to another data interface node like write skinned mesh

diff --git a/Engine/Plugins/Animation/DeformerGraph/Source/OptimusCore/Private/OptimusDeformer.cpp b/Engine/Plugins/Animation/DeformerGraph/Source/OptimusCore/Private/OptimusDeformer.cpp
--- a/Engine/Plugins/Animation/DeformerGraph/Source/OptimusCore/Private/OptimusDeformer.cpp
+++ b/Engine/Plugins/Animation/DeformerGraph/Source/OptimusCore/Private/OptimusDeformer.cpp
@@ -51,6 +51,7 @@
 
 #include "IOptimusDeprecatedExecutionDataInterface.h"
 #include "Nodes/OptimusNode_CustomComputeKernel.h"
+#include "Nodes/OptimusNode_Resource.h"
 
 
 #include UE_INLINE_GENERATED_CPP_BY_NAME(OptimusDeformer)
@@ -2194,11 +2195,13 @@ FOptimusNodeGraphCompilationResult UOptimusDeformer::CompileNodeGraphToComputeGr
 					if (const IOptimusPropertyPinProvider* PropertyPinProvider = Cast<const IOptimusPropertyPinProvider>(TargetNode))
 					{
 						// Property pins are CPU only, no need to involve a kernel
-						if (!PropertyPinProvider->GetPropertyPins().Contains(TargetPin.Pin))
+						if (PropertyPinProvider->GetPropertyPins().Contains(TargetPin.Pin))
 						{
-							LinksToInsertCopyKernel.FindOrAdd(SourcePin).Add(TargetPin);
+							continue;
 						}
 					}
+					
+					LinksToInsertCopyKernel.FindOrAdd(SourcePin).Add(TargetPin);
 				}
 			}
 		}
@@ -2370,6 +2373,12 @@ FOptimusNodeGraphCompilationResult UOptimusDeformer::CompileNodeGraphToComputeGr
 	TMap<FOptimusInstancedNode, UComputeDataInterface*> KernelDataInterfaceMap;
 	TMap<FOptimusInstancedNode, FOptimus_KernelInputMap> KernelInputMap;
 	TMap<FOptimusInstancedNode, FOptimus_KernelOutputMap> KernelOutputMap;
+	
+	// Kernel can either write to a buffer it creates on demand (transient, implicit-persistent raw buffer DIs), 
+	// or directly write to some other data interfaces, this TMap tracks the latter and is used to make sure
+	// copy kernels that copy data from these DIs only start after the writing kernel is done
+	TMap<FOptimusInstancedNode, TArray<FOptimusInstancedNode>> KernelToDirectlyWrittenDataInterfaceNodeMap;
+	
 	TMap<FOptimusInstancedPin, UOptimusComputeDataInterface*> KernelOutputDataInterfaceMap;
 
 	TMap<TWeakObjectPtr<const UComputeDataInterface>, FOptimusDataInterfacePropertyOverrideInfo> NodeGraphDataInterfacePropertyOverrideMap;
@@ -2415,6 +2424,7 @@ FOptimusNodeGraphCompilationResult UOptimusDeformer::CompileNodeGraphToComputeGr
 
 			KernelInputMap.Add(InstancedNode);
 			KernelOutputMap.Add(InstancedNode);
+			KernelToDirectlyWrittenDataInterfaceNodeMap.Add(InstancedNode);
 			
 			for (const UOptimusNodePin* Pin: Node->GetPinsByDirection(EOptimusNodePinDirection::Input, true))
 			{
@@ -2529,7 +2539,8 @@ FOptimusNodeGraphCompilationResult UOptimusDeformer::CompileNodeGraphToComputeGr
 
 					for (const FOptimusInstancedPin& TargetInstancedPin : TargetDataInterfacePins)
 					{
-						const FOptimusRoutedConstNode TargetRoutedNode = TargetInstancedPin.InstancedNode.RoutedNode;
+						const FOptimusInstancedNode& TargetInstancedNode = TargetInstancedPin.InstancedNode;
+						const FOptimusRoutedConstNode& TargetRoutedNode = TargetInstancedNode.RoutedNode;
 						const UOptimusNodePin* TargetPin = TargetInstancedPin.Pin;
 
 						if (bShouldCopyToDataInterface)
@@ -2542,6 +2553,7 @@ FOptimusNodeGraphCompilationResult UOptimusDeformer::CompileNodeGraphToComputeGr
 							if (UOptimusComputeDataInterface** NodeDataInterface = NodeDataInterfaceMap.Find(TargetRoutedNode))
 							{
 								KernelOutputMap[InstancedNode].FindOrAdd(Pin).Add({*NodeDataInterface, TargetPin});
+								KernelToDirectlyWrittenDataInterfaceNodeMap[InstancedNode].Add(TargetInstancedNode);
 							}
 						}	
 					}
@@ -2802,8 +2814,9 @@ FOptimusNodeGraphCompilationResult UOptimusDeformer::CompileNodeGraphToComputeGr
 
 		TArray<FKernelWithDataBindings> BoundKernels;
 		// Copy Kernel uses this map to look up the earliest point it can dispatch
-		// It needs to be dispatched after its input kernels but before all its output kernels
-		TMap<FOptimusInstancedNode, UComputeKernel*> InstancedNodeToComputeKernel;
+		// It needs to be dispatched after the kernels producing its inputs and before the kernels consuming its outputs
+		// This map tracks which kernel produces the input that the copy kernel is copying from
+		TMap<FOptimusInstancedNode, UComputeKernel*> ReadableNodeToProducerComputeKernel;
 		
 		for (const FOptimusInstancedNode& InstancedNode : InstancedNodes )
 		{
@@ -2832,7 +2845,18 @@ FOptimusNodeGraphCompilationResult UOptimusDeformer::CompileNodeGraphToComputeGr
 					BoundKernel.InputDataBindings, BoundKernel.OutputDataBindings
 				);
 
-				InstancedNodeToComputeKernel.Add(InstancedNode, BoundKernel.Kernel);
+				ReadableNodeToProducerComputeKernel.Add(InstancedNode, BoundKernel.Kernel);
+
+				const TArray<FOptimusInstancedNode>& WrittenDataInterfaceNodes = KernelToDirectlyWrittenDataInterfaceNodeMap[InstancedNode];
+				for (const FOptimusInstancedNode& WrittenNode : WrittenDataInterfaceNodes)
+				{
+					// This node is the only data interface provider that has both GPU input and GPU output, and only one of each
+					if (Cast<UOptimusNode_Resource>(WrittenNode.RoutedNode.Node))
+					{
+						check(!ReadableNodeToProducerComputeKernel.Contains(WrittenNode));
+						ReadableNodeToProducerComputeKernel.Add(WrittenNode, BoundKernel.Kernel);
+					}
+				}
 				
 				if (FText* ErrorMessage = KernelSourceResult.TryGet<FText>())
 				{
@@ -2882,20 +2906,14 @@ FOptimusNodeGraphCompilationResult UOptimusDeformer::CompileNodeGraphToComputeGr
 		if (InNodeGraph->GetGraphType() == GraphInfo.GraphType)
 		{
 			TArray<FKernelWithDataBindings> BoundCopyKernels;
-			TArray<UComputeKernel*> InsertAfterComputeKernelLookUpArray;
+			
 			bool bIsUnifiedDispatch = true;
 			for (TPair<FOptimusInstancedPin, TArray<FOptimusInstancedPin>> OutputLink : LinksToInsertCopyKernel)
 			{
 				// Create a copy kernel per source pin, that copies from 1 source pin to multiple target pins
 				const FOptimusInstancedPin& SourceInstancedPin = OutputLink.Key;
 
-				// Indicates that this copy kernel should run immediately after the found compute kernel
-				// Note: nullptr means that the copy kernel has no kernel node dependency and thus should run before everything else
-				UComputeKernel* InsertAfterComputeKernel = nullptr;
-				if (UComputeKernel** ComputeKnernel = InstancedNodeToComputeKernel.Find(SourceInstancedPin.InstancedNode))
-				{
-					InsertAfterComputeKernel = *ComputeKnernel;
-				}
+				
 				
 				const FShaderValueTypeHandle ValueType = SourceInstancedPin.Pin->GetDataType()->ShaderValueType;
 				
@@ -2979,6 +2997,14 @@ FOptimusNodeGraphCompilationResult UOptimusDeformer::CompileNodeGraphToComputeGr
 				for (int32 TargetIndex = 0 ; TargetIndex < OutputLink.Value.Num(); TargetIndex++)
 				{
 					const FOptimusInstancedPin& TargetInstancedPin = OutputLink.Value[TargetIndex];
+
+					// This node is the only data interface provider that has both GPU input and GPU output, and only one of each
+					if (Cast<UOptimusNode_Resource>(TargetInstancedPin.Pin->GetOwningNode()))
+					{
+						check(!ReadableNodeToProducerComputeKernel.Contains(TargetInstancedPin.InstancedNode));
+						ReadableNodeToProducerComputeKernel.Add(TargetInstancedPin.InstancedNode, BoundCopyKernel.Kernel);
+					}
+					
 					const FDataInterfaceFunctionBinding& CopyToBinding = CopyToDataInterfaceMap[TargetInstancedPin];
 
 					if (!CopyToBinding.DataInterface->CanSupportUnifiedDispatch())
@@ -3026,39 +3052,83 @@ FOptimusNodeGraphCompilationResult UOptimusDeformer::CompileNodeGraphToComputeGr
 				BoundCopyKernel.Kernel->KernelSource = KernelSource;
 
 				BoundCopyKernels.Add(BoundCopyKernel);
+				
+			}
+
+			TArray<UComputeKernel*> InsertAfterComputeKernelLookUpArray;
+			for (TPair<FOptimusInstancedPin, TArray<FOptimusInstancedPin>> OutputLink : LinksToInsertCopyKernel)
+			{
+				const FOptimusInstancedPin& SourceInstancedPin = OutputLink.Key;
+				
+				// Indicates that this copy kernel should run immediately after the found compute kernel, which can also be another copy kernel
+				// Note: nullptr means that the copy kernel has no kernel node dependency and thus should run before everything else
+				UComputeKernel* InsertAfterComputeKernel = nullptr;
+				if (UComputeKernel** ComputeKnernel = ReadableNodeToProducerComputeKernel.Find(SourceInstancedPin.InstancedNode))
+				{
+					InsertAfterComputeKernel = *ComputeKnernel;
+				}	
 				InsertAfterComputeKernelLookUpArray.Add(InsertAfterComputeKernel);
 			}
 
 			check(BoundCopyKernels.Num() == InsertAfterComputeKernelLookUpArray.Num());
-			
-			// Insert copy kernels immediately after the kernel that they are copying from
-			for (int32 CopyKernelIndex = BoundCopyKernels.Num()-1 ; CopyKernelIndex >= 0; CopyKernelIndex--)
-			{
-				const FKernelWithDataBindings& BoundCopyKernel = BoundCopyKernels[CopyKernelIndex];
-				UComputeKernel* InsertAfterComputeKernel = InsertAfterComputeKernelLookUpArray[CopyKernelIndex];
 
-				int32 BoundKernelIndex = BoundKernels.IndexOfByPredicate([InsertAfterComputeKernel]( const FKernelWithDataBindings& InBoundKernel)
-				{
-					return InBoundKernel.Kernel == InsertAfterComputeKernel;
-				});
 
-				int32 InsertIndex = INDEX_NONE;
-				if (BoundKernelIndex == INDEX_NONE)
-				{
-					// Insert to the beginning if there is no kernel dependency
-					InsertIndex = 0;
-				}
-				else
+			TOptional<int32> NumLastInserted;
+			TSet<int32> InsertedCopyKernels;
+			while(InsertedCopyKernels.Num() < BoundCopyKernels.Num())
+			{
+				// Avoid infinite loop, in case of unexpected edge case or bad data
+				if (NumLastInserted.IsSet() )
 				{
-					// By Default insert after the kernel that this copy kernel is copying from
-					InsertIndex = BoundKernelIndex + 1;
+					bool bMadeProgress = (InsertedCopyKernels.Num() - *NumLastInserted) > 0;
+					if (!ensure(bMadeProgress))
+					{
+						return {};	
+					}
 				}
-				
-				if (ensure(InsertIndex >= 0) && ensure(InsertIndex <= BoundKernels.Num()))
+				NumLastInserted = InsertedCopyKernels.Num();
+					
+				for (int32 Index = 0 ; Index < BoundCopyKernels.Num() ; Index++)
 				{
-					BoundKernels.Insert(BoundCopyKernel, InsertIndex);
-					ComputeGraph->KernelInvocations.Insert(BoundCopyKernel.Kernel, InsertIndex);
-					ComputeGraph->KernelToNode.Insert(nullptr, InsertIndex);
+					if (InsertedCopyKernels.Contains(Index))
+					{
+						continue;
+					}
+
+					const FKernelWithDataBindings& BoundCopyKernel = BoundCopyKernels[Index];
+					UComputeKernel* InsertAfterComputeKernel = InsertAfterComputeKernelLookUpArray[Index];
+
+					int32 InsertIndex = INDEX_NONE;
+					if (InsertAfterComputeKernel == nullptr)
+					{
+						InsertIndex = 0;
+					}
+					else
+					{
+						int32 BoundKernelIndex = BoundKernels.IndexOfByPredicate([InsertAfterComputeKernel]( const FKernelWithDataBindings& InBoundKernel)
+						{
+							return InBoundKernel.Kernel == InsertAfterComputeKernel;
+						});
+						
+						if (BoundKernelIndex == INDEX_NONE)
+						{
+							// Skip since this copy kernel's dependency hasn't been inserted
+							continue;
+						}
+
+						// By Default insert after the kernel producing the data that this copy kernel is copying from
+						InsertIndex = BoundKernelIndex + 1;
+					}
+					
+					
+					if (ensure(InsertIndex >= 0) && ensure(InsertIndex <= BoundKernels.Num()))
+					{
+						BoundKernels.Insert(BoundCopyKernel, InsertIndex);
+						ComputeGraph->KernelInvocations.Insert(BoundCopyKernel.Kernel, InsertIndex);
+						ComputeGraph->KernelToNode.Insert(nullptr, InsertIndex);
+					}
+
+					InsertedCopyKernels.Add(Index);	
 				}
 			}
 		}
