[Deformer Graph] mark skin weight as vertex mask DI as not supporting unified dispatch, fixed a typo causing the wrong LOD to be used when creating buffers

diff --git a/Engine/Plugins/Animation/DeformerGraph/Source/OptimusCore/Private/DataInterfaces/OptimusDataInterfaceSkinWeightsAsVertexMask.cpp b/Engine/Plugins/Animation/DeformerGraph/Source/OptimusCore/Private/DataInterfaces/OptimusDataInterfaceSkinWeightsAsVertexMask.cpp
--- a/Engine/Plugins/Animation/DeformerGraph/Source/OptimusCore/Private/DataInterfaces/OptimusDataInterfaceSkinWeightsAsVertexMask.cpp
+++ b/Engine/Plugins/Animation/DeformerGraph/Source/OptimusCore/Private/DataInterfaces/OptimusDataInterfaceSkinWeightsAsVertexMask.cpp
@@ -375,7 +375,7 @@ FComputeDataProviderRenderProxy* UOptimusSkinWeightsAsVertexMaskDataProvider::Ge
 				
 				for (int32 LODIndex = 0; LODIndex < SkeletalMeshRenderData.LODRenderData.Num(); LODIndex++)
 				{
-					FSkeletalMeshLODRenderData const& LocalLodRenderData = SkeletalMeshRenderData.LODRenderData[CurrentLodIndex];
+					FSkeletalMeshLODRenderData const& LocalLodRenderData = SkeletalMeshRenderData.LODRenderData[LODIndex];
 					
 					TArray<TArray<uint32>>& CachedBoneIsSelectedPerSectionRef = CachedBoneIsSelectedPerSectionPerLod.AddDefaulted_GetRef();
 					for (int32 SectionIndex = 0; SectionIndex < LocalLodRenderData.RenderSections.Num(); ++SectionIndex)
diff --git a/Engine/Plugins/Animation/DeformerGraph/Source/OptimusCore/Private/DataInterfaces/OptimusDataInterfaceSkinWeightsAsVertexMask.h b/Engine/Plugins/Animation/DeformerGraph/Source/OptimusCore/Private/DataInterfaces/OptimusDataInterfaceSkinWeightsAsVertexMask.h
--- a/Engine/Plugins/Animation/DeformerGraph/Source/OptimusCore/Private/DataInterfaces/OptimusDataInterfaceSkinWeightsAsVertexMask.h
+++ b/Engine/Plugins/Animation/DeformerGraph/Source/OptimusCore/Private/DataInterfaces/OptimusDataInterfaceSkinWeightsAsVertexMask.h
@@ -36,7 +36,7 @@ public:
 
 	//~ Begin UComputeDataInterface Interface
 	TCHAR const* GetClassName() const override { return TEXT("SkinWeightMask"); }
-	bool CanSupportUnifiedDispatch() const override {return true;};
+	bool CanSupportUnifiedDispatch() const override {return false;};
 	void GetSupportedInputs(TArray<FShaderFunctionDefinition>& OutFunctions) const override;
 	void GetShaderParameters(TCHAR const* UID, FShaderParametersMetadataBuilder& InOutBuilder, FShaderParametersMetadataAllocations& InOutAllocations) const override;
 	void GetPermutations(FComputeKernelPermutationVector& OutPermutationVector) const override;
