[Deformer Graph & Control Rig] Added support for pulling variable default value from control rig

diff --git a/Engine/Plugins/Animation/DeformerGraph/Source/OptimusCore/Private/ControlRig/RigUnit_Optimus.cpp b/Engine/Plugins/Animation/DeformerGraph/Source/OptimusCore/Private/ControlRig/RigUnit_Optimus.cpp
--- a/Engine/Plugins/Animation/DeformerGraph/Source/OptimusCore/Private/ControlRig/RigUnit_Optimus.cpp
+++ b/Engine/Plugins/Animation/DeformerGraph/Source/OptimusCore/Private/ControlRig/RigUnit_Optimus.cpp
@@ -208,6 +208,8 @@ TArray<FRigVMUserWorkflow> FRigUnit_AddOptimusDeformer::GetSupportedWorkflows(co
 						}
 						
 						Controller->SetPinExpansion(Node->FindTrait(TraitName)->GetPinPath(), true, true);
+						const FString& DefaultValueString = Variable->DefaultValueStruct.GetValueAsString();
+						Controller->SetPinDefaultValue(Node->FindTrait(TraitName, FRigVMTrait_OptimusVariableBase::GetValuePinName())->GetPinPath(), DefaultValueString);
 					}
 				}
 
@@ -379,6 +381,10 @@ FRigUnit_AddOptimusDeformer_Execute()
 	}
 }
 
+const TCHAR* FRigVMTrait_OptimusVariableBase::GetValuePinName()
+{
+	return TEXT("Value");
+}
 
 void FRigVMTrait_SetDeformerIntVariable::SetValue(UOptimusDeformerInstance* InInstance) const
 {
diff --git a/Engine/Plugins/Animation/DeformerGraph/Source/OptimusCore/Private/ControlRig/RigUnit_Optimus.h b/Engine/Plugins/Animation/DeformerGraph/Source/OptimusCore/Private/ControlRig/RigUnit_Optimus.h
--- a/Engine/Plugins/Animation/DeformerGraph/Source/OptimusCore/Private/ControlRig/RigUnit_Optimus.h
+++ b/Engine/Plugins/Animation/DeformerGraph/Source/OptimusCore/Private/ControlRig/RigUnit_Optimus.h
@@ -80,6 +80,7 @@ struct FRigVMTrait_OptimusVariableBase: public FRigVMTrait
 	GENERATED_BODY()
 
 	virtual void SetValue(UOptimusDeformerInstance* InInstance) const {};
+	static const TCHAR* GetValuePinName();
 };
 
 USTRUCT(BlueprintType)
