[Deformer Graph] Fixed crash on LOD change from LOD1 to LOD0 if LOD0 has more material section than LOD1

diff --git a/Engine/Plugins/Animation/DeformerGraph/Source/OptimusCore/Private/DataInterfaces/OptimusDataInterfaceAdvancedSkeleton.cpp b/Engine/Plugins/Animation/DeformerGraph/Source/OptimusCore/Private/DataInterfaces/OptimusDataInterfaceAdvancedSkeleton.cpp
--- a/Engine/Plugins/Animation/DeformerGraph/Source/OptimusCore/Private/DataInterfaces/OptimusDataInterfaceAdvancedSkeleton.cpp
+++ b/Engine/Plugins/Animation/DeformerGraph/Source/OptimusCore/Private/DataInterfaces/OptimusDataInterfaceAdvancedSkeleton.cpp
@@ -481,6 +481,10 @@ UComputeDataProvider* UOptimusAdvancedSkeletonDataInterface::CreateDataProvider(
 }
 void FOptimusBoneTransformBuffer::SetData(FSkeletalMeshLODRenderData const& InLodRenderData, const TArray<FTransform>& InBoneTransforms)
 {
+	if (InBoneTransforms.IsEmpty())
+	{
+		return;
+	}
 	const FName Matrix34TypeName = FOptimusDataTypeRegistry::Matrix34TypeName;
 	const FOptimusDataTypeHandle Matrix34TypeHandle = FOptimusDataTypeRegistry::Get().FindType(Matrix34TypeName);
 	int32 Matrix34ShaderSize = Matrix34TypeHandle->ShaderValueSize;
@@ -514,24 +518,26 @@ bool FOptimusBoneTransformBuffer::HasData() const
 
 void FOptimusBoneTransformBuffer::AllocateResources(FRDGBuilder& GraphBuilder)
 {
-	if (BufferRefPerSection.Num() == 0)
+	if (!HasData())
 	{
-		BufferRefPerSection.AddDefaulted(BufferData.Num());
-		BufferSRVPerSection.AddDefaulted(BufferData.Num());	
+		return;
+	}
 	
-		// If we are using a raw type alias for the buffer then we need to adjust stride and count.
-		for (int32 InvocationIndex = 0 ; InvocationIndex < BufferData.Num(); InvocationIndex++)
-		{
-			int32 Stride = BufferData[InvocationIndex].Num() / NumBones[InvocationIndex];
-			FRDGBufferRef Buffer = GraphBuilder.CreateBuffer(FRDGBufferDesc::CreateStructuredDesc(Stride,  NumBones[InvocationIndex]), TEXT("BoneTransformBuffer"), ERDGBufferFlags::None);
-			FRDGBufferSRVRef BufferSRV = GraphBuilder.CreateSRV(Buffer);
+	BufferRefPerSection.AddDefaulted(BufferData.Num());
+	BufferSRVPerSection.AddDefaulted(BufferData.Num());	
 
-			BufferRefPerSection[InvocationIndex] = Buffer;
-			BufferSRVPerSection[InvocationIndex] = BufferSRV;
+	// If we are using a raw type alias for the buffer then we need to adjust stride and count.
+	for (int32 InvocationIndex = 0 ; InvocationIndex < BufferData.Num(); InvocationIndex++)
+	{
+		int32 Stride = BufferData[InvocationIndex].Num() / NumBones[InvocationIndex];
+		FRDGBufferRef Buffer = GraphBuilder.CreateBuffer(FRDGBufferDesc::CreateStructuredDesc(Stride,  NumBones[InvocationIndex]), TEXT("BoneTransformBuffer"), ERDGBufferFlags::None);
+		FRDGBufferSRVRef BufferSRV = GraphBuilder.CreateSRV(Buffer);
 
-			GraphBuilder.QueueBufferUpload(Buffer, BufferData[InvocationIndex].GetData(), BufferData[InvocationIndex].Num(), ERDGInitialDataFlags::None);	
-		}	
-	}
+		BufferRefPerSection[InvocationIndex] = Buffer;
+		BufferSRVPerSection[InvocationIndex] = BufferSRV;
+
+		GraphBuilder.QueueBufferUpload(Buffer, BufferData[InvocationIndex].GetData(), BufferData[InvocationIndex].Num(), ERDGInitialDataFlags::None);	
+	}	
 }
 
 void UOptimusAdvancedSkeletonDataProvider::SetDeformerInstance(UOptimusDeformerInstance* InInstance)
@@ -593,7 +599,10 @@ void UOptimusAdvancedSkeletonDataProvider::Init(const UOptimusAdvancedSkeletonDa
 
 }
 
-void UOptimusAdvancedSkeletonDataProvider::ComputeBoneTransformsForLayeredSkinning(FOptimusBoneTransformBuffer& OutBoneBuffer, FSkeletalMeshLODRenderData const& InLodRenderData, const FReferenceSkeleton& InRefSkeleton)
+void UOptimusAdvancedSkeletonDataProvider::ComputeBoneTransformsForLayeredSkinning(
+	TArray<FTransform>& OutBoneTransforms,
+	FSkeletalMeshLODRenderData const& InLodRenderData, 
+	const FReferenceSkeleton& InRefSkeleton)
 {
 	if (!bIsLayeredSkinInitialized)
 	{
@@ -704,9 +713,8 @@ void UOptimusAdvancedSkeletonDataProvider::ComputeBoneTransformsForLayeredSkinni
 		LayeredBoneMatrix[BoneIndex] =
 			(CachedLayerSpaceInitialBoneTransform[BoneIndex] * CurrentComponentSpaceBoneMatrix[BoundaryBoneParentIndex]).Inverse() * CurrentComponentSpaceBoneMatrix[BoneIndex];
 	}
-	
-	
-	OutBoneBuffer.SetData(InLodRenderData, LayeredBoneMatrix);
+
+	OutBoneTransforms = MoveTemp(LayeredBoneMatrix);
 }
 
 FComputeDataProviderRenderProxy* UOptimusAdvancedSkeletonDataProvider::GetRenderProxy()
@@ -771,91 +779,76 @@ FComputeDataProviderRenderProxy* UOptimusAdvancedSkeletonDataProvider::GetRender
 		}
 		else
 		{
-			FOptimusBoneTransformBuffer LayeredBoneMatrixBuffer;
+			const FReferenceSkeleton& RefSkeleton = SkeletalMesh->GetSkinnedAsset()->GetRefSkeleton();
+			
+			TArray<FTransform> LayeredBoneMatrices;
 			
 			if (bEnableLayeredSkinning)
 			{
-				const FReferenceSkeleton& RefSkeleton = SkeletalMesh->GetSkinnedAsset()->GetRefSkeleton();
-				ComputeBoneTransformsForLayeredSkinning(LayeredBoneMatrixBuffer, LodRenderData, RefSkeleton);
+				ComputeBoneTransformsForLayeredSkinning(LayeredBoneMatrices, LodRenderData, RefSkeleton);
 			}
 			
 			// Per-Bone Animation Attributes
 			const UE::Anim::FMeshAttributeContainer& AttributeContainer = SkeletalMesh->GetCustomAttributes();
-			// Section, Attribute, BoneData
-			TArray<TArray<TArray<uint8>>> AttributeBuffers;
 
-			AttributeBuffers.AddDefaulted(LodRenderData.RenderSections.Num());
-			for (int32 SectionIndex = 0; SectionIndex < LodRenderData.RenderSections.Num(); ++SectionIndex)
+			TArray<TArray<uint8>> Attributes;
+			
+			Attributes.AddDefaulted(AttributeBufferRuntimeData.Num());
+			for (int32 AttributeIndex = 0; AttributeIndex < AttributeBufferRuntimeData.Num(); ++AttributeIndex)
 			{
-				FSkelMeshRenderSection const& RenderSection = LodRenderData.RenderSections[SectionIndex];
-				AttributeBuffers[SectionIndex].AddDefaulted(AttributeBufferRuntimeData.Num());
-				for (int32 AttributeIndex = 0; AttributeIndex < AttributeBufferRuntimeData.Num(); ++AttributeIndex)
-				{
-					const FOptimusAnimAttributeBufferRuntimeData& AttributeData = AttributeBufferRuntimeData[AttributeIndex];
-					const int32 NumBones = RenderSection.BoneMap.Num();
-					AttributeBuffers[SectionIndex][AttributeIndex].AddDefaulted(NumBones * AttributeData.Size);
-
-					for (int32 BoneIndex = 0; BoneIndex < NumBones; BoneIndex++)
-					{
-						uint8* BoneDataPtr = AttributeBuffers[SectionIndex][AttributeIndex].GetData() + BoneIndex * AttributeData.Size;
-						int32 FinalBoneIndex = RenderSection.BoneMap[BoneIndex];
+				const FOptimusAnimAttributeBufferRuntimeData& AttributeData = AttributeBufferRuntimeData[AttributeIndex];
 
-						const UE::Anim::FAttributeId Id = {AttributeData.Name, FCompactPoseBoneIndex(FinalBoneIndex)} ;
+				const int32 NumBones = RefSkeleton.GetNum();
+				Attributes[AttributeIndex].AddDefaulted(NumBones * AttributeData.Size);
+				for (int32 BoneIndex = 0; BoneIndex < NumBones; BoneIndex++)
+				{
+					uint8* BoneDataPtr = Attributes[AttributeIndex].GetData() + BoneIndex * AttributeData.Size;
+					const UE::Anim::FAttributeId Id = {AttributeData.Name, FCompactPoseBoneIndex(BoneIndex)} ;
 
-						bool bIsValueSet = false;
+					bool bIsValueSet = false;
 				
-						FOptimusDataTypeRegistry::PropertyValueConvertFuncT ConvertFunc = AttributeData.ConvertFunc;
+					FOptimusDataTypeRegistry::PropertyValueConvertFuncT ConvertFunc = AttributeData.ConvertFunc;
 
-						if (ConvertFunc)
-						{
-							UScriptStruct* AttributeType = AttributeData.AttributeType;
+					if (ConvertFunc)
+					{
+						UScriptStruct* AttributeType = AttributeData.AttributeType;
 				
-							if (const uint8* Attribute = AttributeContainer.Find(AttributeType, Id))
-							{
-								bIsValueSet = true;
+						if (const uint8* Attribute = AttributeContainer.Find(AttributeType, Id))
+						{
+							bIsValueSet = true;
 						
-								const uint8* ValuePtr = Attribute;
-								int32 ValueSize = AttributeType->GetStructureSize();
+							const uint8* ValuePtr = Attribute;
+							int32 ValueSize = AttributeType->GetStructureSize();
 
-								// TODO: use a specific function to extract the value from the attribute
-								// it works for now because even if the attribute type != actual value type
-								// it should only have a single property, whose type == actual property type
+							// TODO: use a specific function to extract the value from the attribute
+							// it works for now because even if the attribute type != actual value type
+							// it should only have a single property, whose type == actual property type
 						
-								ConvertFunc(
-									{ValuePtr, ValueSize},
-									{ {BoneDataPtr, AttributeData.Size}, {} });
-							}
+							ConvertFunc(
+								{ValuePtr, ValueSize},
+								{ {BoneDataPtr, AttributeData.Size}, {} });
+						}
 				
-							// Use the default value if the attribute was not found
-							if (!bIsValueSet)
-							{
-								const uint8* DefaultValuePtr = AttributeData.CachedDefaultValue.ShaderValue.GetData();
-								const uint32 DefaultValueSize = AttributeData.CachedDefaultValue.ShaderValue.Num();
-
-								FMemory::Memcpy(BoneDataPtr, DefaultValuePtr, DefaultValueSize);
-							}
+						// Use the default value if the attribute was not found
+						if (!bIsValueSet)
+						{
+							const uint8* DefaultValuePtr = AttributeData.CachedDefaultValue.ShaderValue.GetData();
+							const uint32 DefaultValueSize = AttributeData.CachedDefaultValue.ShaderValue.Num();
+
+							FMemory::Memcpy(BoneDataPtr, DefaultValuePtr, DefaultValueSize);
 						}
 					}
 				}
 			}
 			
-			TArray<int32> AttributeBufferOffsets;
-			AttributeBufferOffsets.AddUninitialized(AttributeBufferRuntimeData.Num());
-			
-			for (int32 AttributeIndex = 0; AttributeIndex < AttributeBufferRuntimeData.Num(); ++AttributeIndex)
-			{
-				AttributeBufferOffsets[AttributeIndex] = AttributeBufferRuntimeData[AttributeIndex].Offset;
-			}
 
 			// Pipe data into proxy
 			Proxy->SkeletalMeshObject = SkeletalMeshObject;
 			Proxy->SkinWeightProfile = SkinWeightProfile;
-			Proxy->BoneRevisionNumber = SkeletalMesh->GetBoneTransformRevisionNumber();
 			Proxy->ParameterBuffer.AddDefaulted(ParameterBufferSize);
-			Proxy->LayeredBoneMatrixBuffer = MoveTemp(LayeredBoneMatrixBuffer);
-			
-			Proxy->AttributeBufferOffsets = MoveTemp(AttributeBufferOffsets);
-			Proxy->AttributeBuffers = MoveTemp(AttributeBuffers);	
+			Proxy->LayeredBoneMatrices = MoveTemp(LayeredBoneMatrices);
+			Proxy->AttributeBufferRuntimeData = AttributeBufferRuntimeData;
+			Proxy->Attributes = MoveTemp(Attributes);
 		}
 	}
 	
@@ -904,22 +897,52 @@ void FOptimusAdvancedSkeletonDataProviderProxy::AllocateResources(FRDGBuilder& G
 {
 	const int32 LodIndex = SkeletalMeshObject->GetLOD();
 	FSkeletalMeshRenderData const& SkeletalMeshRenderData = SkeletalMeshObject->GetSkeletalMeshRenderData();
-	FSkeletalMeshLODRenderData const* LodRenderData = &SkeletalMeshRenderData.LODRenderData[LodIndex];
+	FSkeletalMeshLODRenderData const& LodRenderData = SkeletalMeshRenderData.LODRenderData[LodIndex];
 	
 	FallbackSRV = GraphBuilder.CreateSRV(FRDGBufferSRVDesc(GSystemTextures.GetDefaultStructuredBuffer(GraphBuilder, sizeof(int32))));
+
+	LayeredBoneMatrixBuffer.SetData(LodRenderData, LayeredBoneMatrices);
 	
 	LayeredBoneMatrixBuffer.AllocateResources(GraphBuilder);
 	
 	// Per-Bone Animation Attribute 
-	BuffersPerAttributePerSection.AddDefaulted(LodRenderData->RenderSections.Num());
-	BufferSRVsPerAttributePerSection.AddDefaulted(LodRenderData->RenderSections.Num());	
+	BuffersPerAttributePerSection.AddDefaulted(LodRenderData.RenderSections.Num());
+	BufferSRVsPerAttributePerSection.AddDefaulted(LodRenderData.RenderSections.Num());	
+
+	AttributeBuffers.AddDefaulted(LodRenderData.RenderSections.Num());
+	for (int32 SectionIndex = 0; SectionIndex < LodRenderData.RenderSections.Num(); ++SectionIndex)
+	{
+		FSkelMeshRenderSection const& RenderSection = LodRenderData.RenderSections[SectionIndex];
+		AttributeBuffers[SectionIndex].AddDefaulted(AttributeBufferRuntimeData.Num());
+		for (int32 AttributeIndex = 0; AttributeIndex < AttributeBufferRuntimeData.Num(); ++AttributeIndex)
+		{
+			const FOptimusAnimAttributeBufferRuntimeData& AttributeData = AttributeBufferRuntimeData[AttributeIndex];
+			const int32 NumBones = RenderSection.BoneMap.Num();
+			AttributeBuffers[SectionIndex][AttributeIndex].AddDefaulted(NumBones * AttributeData.Size);
+
+			for (int32 BoneIndex = 0; BoneIndex < NumBones; BoneIndex++)
+			{
+				uint8* TargetBoneDataPtr = AttributeBuffers[SectionIndex][AttributeIndex].GetData() + BoneIndex * AttributeData.Size;
+				int32 FinalBoneIndex = RenderSection.BoneMap[BoneIndex];
+
+				int32 SourceDataOffset = FinalBoneIndex * AttributeData.Size;
+				if (ensure((SourceDataOffset + AttributeData.Size) <= Attributes[AttributeIndex].Num()))
+				{
+					uint8* SourceBoneDataPtr = Attributes[AttributeIndex].GetData() + SourceDataOffset;
+					FMemory::Memcpy(TargetBoneDataPtr, SourceBoneDataPtr,  AttributeData.Size);	
+				}
+			}
+		}
+		
+	}
 	
-	for (int32 InvocationIndex = 0 ; InvocationIndex < LodRenderData->RenderSections.Num(); InvocationIndex++)
+	for (int32 InvocationIndex = 0 ; InvocationIndex < LodRenderData.RenderSections.Num(); InvocationIndex++)
 	{
 		BuffersPerAttributePerSection[InvocationIndex].AddDefaulted(AttributeBuffers[InvocationIndex].Num());
 		BufferSRVsPerAttributePerSection[InvocationIndex].AddDefaulted(AttributeBuffers[InvocationIndex].Num());
 		
-		int32 NumBones = LodRenderData->RenderSections[InvocationIndex].BoneMap.Num();
+		FSkelMeshRenderSection const& RenderSection = LodRenderData.RenderSections[InvocationIndex];	
+		int32 NumBones = RenderSection.BoneMap.Num();
 		
 		for (int32 AttributeIndex = 0; AttributeIndex < AttributeBuffers[InvocationIndex].Num(); ++AttributeIndex)
 		{
@@ -1046,9 +1069,9 @@ void FOptimusAdvancedSkeletonDataProviderProxy::GatherDispatchData(FDispatchData
 			LayeredBoneMatrixBuffer.HasData() && LayeredBoneMatrixBuffer.BufferSRVPerSection[InvocationIndex] != nullptr ?
 				LayeredBoneMatrixBuffer.BufferSRVPerSection[InvocationIndex] : FallbackSRV;
 
-		for(int32 AttributeIndex = 0 ; AttributeIndex < AttributeBufferOffsets.Num(); AttributeIndex++)
+		for(int32 AttributeIndex = 0 ; AttributeIndex < AttributeBufferRuntimeData.Num(); AttributeIndex++)
 		{
-			*((FRDGBufferSRV**)(ParameterBufferPtr + AttributeBufferOffsets[AttributeIndex])) = BufferSRVsPerAttributePerSection[InvocationIndex][AttributeIndex];
+			*((FRDGBufferSRV**)(ParameterBufferPtr + AttributeBufferRuntimeData[AttributeIndex].Offset)) = BufferSRVsPerAttributePerSection[InvocationIndex][AttributeIndex];
 		}
 	}
 }
diff --git a/Engine/Plugins/Animation/DeformerGraph/Source/OptimusCore/Private/DataInterfaces/OptimusDataInterfaceAdvancedSkeleton.h b/Engine/Plugins/Animation/DeformerGraph/Source/OptimusCore/Private/DataInterfaces/OptimusDataInterfaceAdvancedSkeleton.h
--- a/Engine/Plugins/Animation/DeformerGraph/Source/OptimusCore/Private/DataInterfaces/OptimusDataInterfaceAdvancedSkeleton.h
+++ b/Engine/Plugins/Animation/DeformerGraph/Source/OptimusCore/Private/DataInterfaces/OptimusDataInterfaceAdvancedSkeleton.h
@@ -242,7 +242,7 @@ public:
 		USkeletalMeshComponent* InSkeletalMesh
 	);
 
-	void ComputeBoneTransformsForLayeredSkinning(FOptimusBoneTransformBuffer& OutBoneBuffer, FSkeletalMeshLODRenderData const& InLodRenderData, const FReferenceSkeleton& InRefSkeleton);
+	void ComputeBoneTransformsForLayeredSkinning(TArray<FTransform>& OutBoneTransforms, FSkeletalMeshLODRenderData const& InLodRenderData, const FReferenceSkeleton& InRefSkeleton);
 private:
 	UPROPERTY()
 	TObjectPtr<UOptimusDeformerInstance> DeformerInstance;
@@ -264,16 +264,19 @@ public:
 
 	using FDefaultParameters = FAdvancedSkeletonDataInterfaceDefaultParameters;
 
+	// Data supplied by data provider
 	FSkeletalMeshObject* SkeletalMeshObject = nullptr;
 	FName SkinWeightProfile = NAME_None;
-	uint32 BoneRevisionNumber = 0;
 	
 	TArray<uint8> ParameterBuffer;
 	
+	TArray<FTransform> LayeredBoneMatrices;
+	TArray<FOptimusAnimAttributeBufferRuntimeData> AttributeBufferRuntimeData;
+	TArray<TArray<uint8>> Attributes;
+
+	// Proxy managed data to be uploaded to GPU
 	FOptimusBoneTransformBuffer LayeredBoneMatrixBuffer;
 	
-	TArray<int32> AttributeBufferOffsets;
-	
 	TArray<TArray<FRDGBufferRef>> BuffersPerAttributePerSection;
 	TArray<TArray<FRDGBufferSRVRef>> BufferSRVsPerAttributePerSection;
 	TArray<TArray<TArray<uint8>>> AttributeBuffers;
