[Deformer Graph Editor] Fixed copy paste not working in the shader text editor

diff --git a/Engine/Plugins/Animation/DeformerGraph/Source/OptimusEditor/Private/Widgets/SOptimusShaderTextDocumentTextBox.cpp b/Engine/Plugins/Animation/DeformerGraph/Source/OptimusEditor/Private/Widgets/SOptimusShaderTextDocumentTextBox.cpp
--- a/Engine/Plugins/Animation/DeformerGraph/Source/OptimusEditor/Private/Widgets/SOptimusShaderTextDocumentTextBox.cpp
+++ b/Engine/Plugins/Animation/DeformerGraph/Source/OptimusEditor/Private/Widgets/SOptimusShaderTextDocumentTextBox.cpp
@@ -127,6 +127,7 @@ void SOptimusShaderTextDocumentTextBox::Construct(const FArguments& InArgs)
 						{
 							InBuilder.PushCommandList(TextCommandList);
 							InBuilder.AddMenuEntry(Commands.ToggleComment);
+							InBuilder.PopCommandList();
 						}
 						
 					}))
