[DeformerGraph] Added a null check

diff --git a/Engine/Plugins/Animation/DeformerGraph/Source/OptimusCore/Private/OptimusDeformerDynamicInstanceManager.cpp b/Engine/Plugins/Animation/DeformerGraph/Source/OptimusCore/Private/OptimusDeformerDynamicInstanceManager.cpp
--- a/Engine/Plugins/Animation/DeformerGraph/Source/OptimusCore/Private/OptimusDeformerDynamicInstanceManager.cpp
+++ b/Engine/Plugins/Animation/DeformerGraph/Source/OptimusCore/Private/OptimusDeformerDynamicInstanceManager.cpp
@@ -177,7 +177,10 @@ void UOptimusDeformerDynamicInstanceManager::OnObjectBeginDestroy(IMeshDeformerP
 		{
 			if (TObjectPtr<UOptimusDeformerInstance>* InstancePtr = GuidToInstanceMap.Find(Guid))
 			{
-				(*InstancePtr)->ReleaseResources();
+				if (*InstancePtr)
+				{
+					(*InstancePtr)->ReleaseResources();
+				}
 			}
 
 			GuidToInstanceMap.Remove(Guid);
