[Deformer Graph & Control Rig] Added a check for empty default value, which is not accepted by the SetPinDefaultValue API

diff --git a/Engine/Plugins/Animation/DeformerGraph/Source/OptimusCore/Private/ControlRig/RigUnit_Optimus.cpp b/Engine/Plugins/Animation/DeformerGraph/Source/OptimusCore/Private/ControlRig/RigUnit_Optimus.cpp
--- a/Engine/Plugins/Animation/DeformerGraph/Source/OptimusCore/Private/ControlRig/RigUnit_Optimus.cpp
+++ b/Engine/Plugins/Animation/DeformerGraph/Source/OptimusCore/Private/ControlRig/RigUnit_Optimus.cpp
@@ -209,7 +209,12 @@ TArray<FRigVMUserWorkflow> FRigUnit_AddOptimusDeformer::GetSupportedWorkflows(co
 						
 						Controller->SetPinExpansion(Node->FindTrait(TraitName)->GetPinPath(), true, true);
 						const FString& DefaultValueString = Variable->DefaultValueStruct.GetValueAsString();
-						Controller->SetPinDefaultValue(Node->FindTrait(TraitName, FRigVMTrait_OptimusVariableBase::GetValuePinName())->GetPinPath(), DefaultValueString);
+						// Array pin's value string from DefaultValueStruct can be empty if it is empty, but control rig expects "()" instead of empty default value,
+						// better to simply use control rig supplied default value in this case
+						if (!DefaultValueString.IsEmpty())
+						{
+							Controller->SetPinDefaultValue(Node->FindTrait(TraitName, FRigVMTrait_OptimusVariableBase::GetValuePinName())->GetPinPath(), DefaultValueString);
+						}
 					}
 				}
 
