[Deformer Graph] Fixed bad tangents when using opengl_es31
original swarm https://p4-swarm.epicgames.net/reviews/43052121

diff --git a/Engine/Plugins/Animation/DeformerGraph/Shaders/Private/DataInterfaceSkinnedMeshRead.ush b/Engine/Plugins/Animation/DeformerGraph/Shaders/Private/DataInterfaceSkinnedMeshRead.ush
--- a/Engine/Plugins/Animation/DeformerGraph/Shaders/Private/DataInterfaceSkinnedMeshRead.ush
+++ b/Engine/Plugins/Animation/DeformerGraph/Shaders/Private/DataInterfaceSkinnedMeshRead.ush
@@ -3,6 +3,13 @@
 // Include required for platform TANGENT_RWBUFFER_FORMAT.
 #include "/Engine/Private/GpuSkinCommon.ush"
 
+#if OPENGL_PROFILE
+// On OpenGL, deformer/skin cache tangents are stored as SINT [-32767, 32767], instead of SNORM [-1.0,1.0], see TangentUnbias_SkinCache
+#define TangentRead_SkinCache(X)	(X / 32767.0)
+#else
+#define TangentRead_SkinCache(X)	(X)
+#endif // OPENGL_PROFILE
+
 uint {DataInterfaceName}_NumVertices;
 
 Buffer<float> {DataInterfaceName}_PositionBufferSRV;
@@ -10,8 +17,8 @@ Buffer<TANGENT_RWBUFFER_FORMAT> {DataInterfaceName}_TangentBufferSRV;
 Buffer<UNORM float4> {DataInterfaceName}_ColorBufferSRV;
 
 Buffer<float> {DataInterfaceName}_PositionStaticBuffer;
-Buffer<TANGENT_RWBUFFER_FORMAT> {DataInterfaceName}_TangentStaticBuffer;
-Buffer<UNORM float4> {DataInterfaceName}_ColorStaticBuffer;
+Buffer<SNORM float4> {DataInterfaceName}_TangentStaticBuffer;
+Buffer<float4> {DataInterfaceName}_ColorStaticBuffer;
 uint {DataInterfaceName}_ColorIndexMask;
 
 uint ReadNumVertices_{DataInterfaceName}()
@@ -41,7 +48,7 @@ float4 ReadTangentX_{DataInterfaceName}(uint VertexIndex)
 	uint BufferIndex = (VertexIndex) * 2;
 	
 #if READABLE_OUTPUT_BUFFERS & OPTIMUS_SKINNED_MESH_READ_TANGENTS
-	return TangentBias_SkinCache({DataInterfaceName}_TangentBufferSRV[BufferIndex]) ;
+	return TangentRead_SkinCache({DataInterfaceName}_TangentBufferSRV[BufferIndex]) ;
 #else	
 	return TangentBias_SkinCache({DataInterfaceName}_TangentStaticBuffer[BufferIndex]) ;
 #endif
@@ -51,7 +58,7 @@ float4 ReadTangentZ_{DataInterfaceName}(uint VertexIndex)
 {
 	uint BufferIndex = (VertexIndex) * 2;
 #if READABLE_OUTPUT_BUFFERS & OPTIMUS_SKINNED_MESH_READ_TANGENTS
-	return TangentBias_SkinCache({DataInterfaceName}_TangentBufferSRV[BufferIndex + 1]) ;
+	return TangentRead_SkinCache({DataInterfaceName}_TangentBufferSRV[BufferIndex + 1]) ;
 #else	
 	return TangentBias_SkinCache({DataInterfaceName}_TangentStaticBuffer[BufferIndex + 1]) ;
 #endif
