[Deformer Graph] Fixed SkinWeightAsVertexMask DI doesn't work with non-default skin weight profiles

diff --git a/Engine/Plugins/Animation/DeformerGraph/Source/OptimusCore/Private/DataInterfaces/OptimusDataInterfaceSkinWeightsAsVertexMask.cpp b/Engine/Plugins/Animation/DeformerGraph/Source/OptimusCore/Private/DataInterfaces/OptimusDataInterfaceSkinWeightsAsVertexMask.cpp
--- a/Engine/Plugins/Animation/DeformerGraph/Source/OptimusCore/Private/DataInterfaces/OptimusDataInterfaceSkinWeightsAsVertexMask.cpp
+++ b/Engine/Plugins/Animation/DeformerGraph/Source/OptimusCore/Private/DataInterfaces/OptimusDataInterfaceSkinWeightsAsVertexMask.cpp
@@ -269,7 +269,9 @@ FComputeDataProviderRenderProxy* UOptimusSkinWeightsAsVertexMaskDataProvider::Ge
 		ExpandTowardsLeaf = FMath::Max(0, ExpandTowardsLeaf);
 		
 		FSkeletalMeshObject* SkeletalMeshObject = SkeletalMesh->MeshObject;
-		const int32 CurrentLodIndex = SkeletalMeshObject->GetLOD();
+		// Not a good idea to get lod from mesh object, which is from last frame and not updated until later by render thread.
+		// So using predicted lod level here, which mirrors the way we get lod during SendRenderDynamicData_Concurrent
+		const int32 CurrentLodIndex = SkeletalMesh->GetPredictedLODLevel();
 		FSkeletalMeshRenderData const& SkeletalMeshRenderData = SkeletalMeshObject->GetSkeletalMeshRenderData();
 		FSkeletalMeshLODRenderData const& CurrentLodRenderData = SkeletalMeshRenderData.LODRenderData[CurrentLodIndex];
 
@@ -576,6 +578,11 @@ void FOptimusSkinWeightsAsVertexMaskDataProviderProxy::GatherDispatchData(FDispa
 	FRHIShaderResourceView* NullSRVBinding = GWhiteVertexBufferWithSRV->ShaderResourceViewRHI.GetReference();
 	
 	FSkinWeightVertexBuffer const* WeightBuffer = LodRenderData->GetSkinWeightVertexBuffer();
+	if (LodRenderData->SkinWeightProfilesData.ContainsProfile(SkinWeightProfile))
+	{
+		const FSkinWeightProfileStack ProfileStack{SkinWeightProfile};
+		WeightBuffer = LodRenderData->SkinWeightProfilesData.GetOverrideBuffer(ProfileStack);
+	}
 	check(WeightBuffer != nullptr);
 	FRHIShaderResourceView* SkinWeightBufferSRV = WeightBuffer->GetDataVertexBuffer()->GetSRV();
 	const bool bUnlimitedBoneInfluences = WeightBuffer->GetBoneInfluenceType() == GPUSkinBoneInfluenceType::UnlimitedBoneInfluence;
