[DeformerGraph] change to async load instead of sync load when dynamically adding deformer

diff --git a/Engine/Plugins/Animation/DeformerGraph/Source/OptimusCore/Private/ControlRig/RigUnit_Optimus.cpp b/Engine/Plugins/Animation/DeformerGraph/Source/OptimusCore/Private/ControlRig/RigUnit_Optimus.cpp
--- a/Engine/Plugins/Animation/DeformerGraph/Source/OptimusCore/Private/ControlRig/RigUnit_Optimus.cpp
+++ b/Engine/Plugins/Animation/DeformerGraph/Source/OptimusCore/Private/ControlRig/RigUnit_Optimus.cpp
@@ -314,54 +314,75 @@ FRigUnit_AddOptimusDeformer_Execute()
 			{
 				if (!DeformerTrait->DeformerGraph.IsNull())
 				{
-					FFunctionGraphTask::CreateAndDispatchWhenReady([
-						WeakMesh = TWeakObjectPtr(RigMeshComponent),
-						WeakRig = TWeakObjectPtr(ExecuteContext.ControlRig),
-						DeformerGraphAsset = DeformerTrait->DeformerGraph,
+					auto AddDeformerToComponentsIfNeeded = [
+						WeakMesh = TWeakObjectPtr<USkeletalMeshComponent>(RigMeshComponent),
+						WeakRig = TWeakObjectPtr<UControlRig>(ExecuteContext.ControlRig),
 						DeformerInstanceGuid,
 						DeformChildComponents = SettingsTraitPtr->DeformChildComponents,
-						ExcludeChildComponentsWithTag = SettingsTraitPtr->ExcludeChildComponentsWithTag]()
-					{
-						if (WeakMesh.IsValid() && WeakRig.IsValid())
+						ExcludeChildComponentsWithTag = SettingsTraitPtr->ExcludeChildComponentsWithTag
+						](UOptimusDeformer* DeformerGraph)
 						{
-							// See StaticFindObjectFastSafe()
-							if (!UE::IsSavingPackage(nullptr) && !IsGarbageCollectingAndLockingUObjectHashTables())
+							check(IsInGameThread());
+
+							if (WeakMesh.IsValid() && WeakRig.IsValid())
 							{
-								if (UOptimusDeformer* DeformerGraph = DeformerGraphAsset.LoadSynchronous())
-								{
-									USkeletalMeshComponent* RigMeshComponent = WeakMesh.Get();
-									
-									TArray<USkeletalMeshComponent*> ComponentsToProcess =
-										Local::GetComponentsToProcess(RigMeshComponent, DeformChildComponents, ExcludeChildComponentsWithTag);
+								USkeletalMeshComponent* RigMeshComponent = WeakMesh.Get();
+							
+								TArray<USkeletalMeshComponent*> ComponentsToProcess =
+									Local::GetComponentsToProcess(RigMeshComponent, DeformChildComponents, ExcludeChildComponentsWithTag);
 
-									for (USkeletalMeshComponent* ComponentToProcess : ComponentsToProcess)
+								for (USkeletalMeshComponent* ComponentToProcess : ComponentsToProcess)
+								{
+									// Currently only one deformer for all LODs so using LOD0 for now is fine
+									// It may change in the future depending on how per-LOD deformer instance is implemented
+									if (!ComponentToProcess->GetMeshDeformerInstanceForLOD(0))
 									{
-										// Currently only one deformer for all LODs so using LOD0 for now is fine
-										// It may change in the future depending on how per-LOD deformer instance is implemented
-										if (!ComponentToProcess->GetMeshDeformerInstanceForLOD(0))
-										{
-											// Every time we re-set the mesh deformers we have to wait for anim eval to complete to avoid anim thread accessing
-											// deformer instances while we are trying to modify it. Similar to OnUnregister for skeletal mesh component
-											ComponentToProcess->HandleExistingParallelEvaluationTask(true, false);
-										
-											// In case there is no mesh deformer running, force turning it on and use a default deformer
-											// in project settings
-											ComponentToProcess->SetAlwaysUseMeshDeformer(true);
-										}
+										// Every time we re-set the mesh deformers we have to wait for anim eval to complete to avoid anim thread accessing
+										// deformer instances while we are trying to modify it. Similar to OnUnregister for skeletal mesh component
+										ComponentToProcess->HandleExistingParallelEvaluationTask(true, false);
 								
-										if (UOptimusDeformerDynamicInstanceManager* DeformerInstanceManager = Cast<UOptimusDeformerDynamicInstanceManager>(ComponentToProcess->GetMeshDeformerInstanceForLOD(0)))
+										// In case there is no mesh deformer running, force turning it on and use a default deformer
+										// in project settings
+										ComponentToProcess->SetAlwaysUseMeshDeformer(true);
+									}
+						
+									if (UOptimusDeformerDynamicInstanceManager* DeformerInstanceManager = Cast<UOptimusDeformerDynamicInstanceManager>(ComponentToProcess->GetMeshDeformerInstanceForLOD(0)))
+									{
+										UControlRig* Rig = WeakRig.Get();
+										if (!DeformerInstanceManager->GetDeformerInstance(DeformerInstanceGuid))
 										{
-											UControlRig* Rig = WeakRig.Get();
-											if (!DeformerInstanceManager->GetDeformerInstance(DeformerInstanceGuid))
-											{
-												DeformerInstanceManager->AddProducerDeformer(Rig, DeformerInstanceGuid, DeformerGraph);
-											}
+											DeformerInstanceManager->AddProducerDeformer(Rig, DeformerInstanceGuid, DeformerGraph);
 										}
 									}
-								}	
-							}
-						}
-					}, TStatId(), NULL, ENamedThreads::GameThread);
+								}
+							}	
+						};
+
+					if (UOptimusDeformer* DeformerGraph = DeformerTrait->DeformerGraph.Get())
+					{
+						FFunctionGraphTask::CreateAndDispatchWhenReady([
+							AddDeformerToComponentsIfNeeded,
+							WeakDeformer = TWeakObjectPtr<UOptimusDeformer>(DeformerGraph)]()
+							{
+								if (WeakDeformer.IsValid())
+								{
+									AddDeformerToComponentsIfNeeded(WeakDeformer.Get());
+								}
+							}, TStatId(), NULL, ENamedThreads::GameThread);	
+					}
+					else
+					{
+						DeformerTrait->DeformerGraph.LoadAsync(
+							FLoadSoftObjectPathAsyncDelegate::CreateLambda(
+								[AddDeformerToComponentsIfNeeded](const FSoftObjectPath&, UObject* InLoadedObject)
+								{
+									if (UOptimusDeformer* DeformerGraph = Cast<UOptimusDeformer>(InLoadedObject))
+									{
+										AddDeformerToComponentsIfNeeded(DeformerGraph);
+									}
+								}));	
+					}
+				
 				}	
 			}
 			else if (const FRigVMTrait_OptimusVariableBase* VariableTrait = Scope.GetTrait<FRigVMTrait_OptimusVariableBase>())
