[Deformer Graph] Added additional validation to avoid crash due to lod mismatch when setting up data provider proxies.

diff --git a/Engine/Plugins/Animation/DeformerGraph/Source/OptimusCore/Private/DataInterfaces/OptimusDataInterfaceMorphTarget.cpp b/Engine/Plugins/Animation/DeformerGraph/Source/OptimusCore/Private/DataInterfaces/OptimusDataInterfaceMorphTarget.cpp
--- a/Engine/Plugins/Animation/DeformerGraph/Source/OptimusCore/Private/DataInterfaces/OptimusDataInterfaceMorphTarget.cpp
+++ b/Engine/Plugins/Animation/DeformerGraph/Source/OptimusCore/Private/DataInterfaces/OptimusDataInterfaceMorphTarget.cpp
@@ -124,6 +124,18 @@ bool FOptimusMorphTargetDataProviderProxy::IsValid(FValidationData const& InVali
 		return false;
 	}
 
+	const int32 LodIndex = SkeletalMeshObject->GetLOD();
+	FSkeletalMeshRenderData const& SkeletalMeshRenderData = SkeletalMeshObject->GetSkeletalMeshRenderData();
+	FSkeletalMeshLODRenderData const* LodRenderData = &SkeletalMeshRenderData.LODRenderData[LodIndex];
+
+	// this mismatch can happen during map load for one frame sometimes, when predicted lod != actual lod
+	// TODO: avoid the mismatch in the first place, use actual lod to query NumInvocations
+	if (LodRenderData->RenderSections.Num() != InValidationData.NumInvocations)
+	{
+		return false;
+	}
+
+	
 
 	return true;
 }
@@ -145,14 +157,10 @@ struct FMorphTargetDataInterfacePermutationIds
 void FOptimusMorphTargetDataProviderProxy::GatherPermutations(FPermutationData& InOutPermutationData) const
 {
 	const int32 LodIndex = SkeletalMeshObject->GetLOD();
-	FSkeletalMeshRenderData const& SkeletalMeshRenderData = SkeletalMeshObject->GetSkeletalMeshRenderData();
-	FSkeletalMeshLODRenderData const* LodRenderData = &SkeletalMeshRenderData.LODRenderData[LodIndex];
 
 	FMorphTargetDataInterfacePermutationIds PermutationIds(InOutPermutationData.PermutationVector);
 	for (int32 InvocationIndex = 0; InvocationIndex < InOutPermutationData.NumInvocations; ++InvocationIndex)
 	{
-		FSkelMeshRenderSection const& RenderSection = LodRenderData->RenderSections[InvocationIndex];
-
 		const bool bPreviousFrame = false;
 		FRHIShaderResourceView* MorphBufferSRV = FSkeletalMeshDeformerHelpers::GetMorphTargetBufferForReading(SkeletalMeshObject, LodIndex, InvocationIndex, FrameNumber, bPreviousFrame);
 		const bool bValidMorph = MorphBufferSRV != nullptr;
diff --git a/Engine/Plugins/Animation/DeformerGraph/Source/OptimusCore/Private/DataInterfaces/OptimusDataInterfaceSkinWeightsAsVertexMask.cpp b/Engine/Plugins/Animation/DeformerGraph/Source/OptimusCore/Private/DataInterfaces/OptimusDataInterfaceSkinWeightsAsVertexMask.cpp
--- a/Engine/Plugins/Animation/DeformerGraph/Source/OptimusCore/Private/DataInterfaces/OptimusDataInterfaceSkinWeightsAsVertexMask.cpp
+++ b/Engine/Plugins/Animation/DeformerGraph/Source/OptimusCore/Private/DataInterfaces/OptimusDataInterfaceSkinWeightsAsVertexMask.cpp
@@ -455,6 +455,13 @@ bool FOptimusSkinWeightsAsVertexMaskDataProviderProxy::IsValid(FValidationData c
 	const int32 LodIndex = SkeletalMeshObject->GetLOD();
 	FSkeletalMeshRenderData const& SkeletalMeshRenderData = SkeletalMeshObject->GetSkeletalMeshRenderData();
 	FSkeletalMeshLODRenderData const* LodRenderData = &SkeletalMeshRenderData.LODRenderData[LodIndex];
+
+	// this mismatch can happen during map load for one frame sometimes, when predicted lod != actual lod
+	// TODO: avoid the mismatch in the first place, use actual lod to query NumInvocations
+	if (LodRenderData->RenderSections.Num() != InValidationData.NumInvocations)
+	{
+		return false;
+	}
 	
 	FSkinWeightVertexBuffer const* WeightBuffer = LodRenderData->GetSkinWeightVertexBuffer();
 	if (LodRenderData->SkinWeightProfilesData.ContainsProfile(SkinWeightProfile))
diff --git a/Engine/Plugins/Animation/DeformerGraph/Source/OptimusCore/Private/DataInterfaces/OptimusDataInterfaceSkinnedMeshRead.cpp b/Engine/Plugins/Animation/DeformerGraph/Source/OptimusCore/Private/DataInterfaces/OptimusDataInterfaceSkinnedMeshRead.cpp
--- a/Engine/Plugins/Animation/DeformerGraph/Source/OptimusCore/Private/DataInterfaces/OptimusDataInterfaceSkinnedMeshRead.cpp
+++ b/Engine/Plugins/Animation/DeformerGraph/Source/OptimusCore/Private/DataInterfaces/OptimusDataInterfaceSkinnedMeshRead.cpp
@@ -203,7 +203,17 @@ bool FOptimusSkinnedMeshReadDataProviderProxy::IsValid(FValidationData const& In
 		return false;
 	}
 
+	const int32 LodIndex = SkeletalMeshObject->GetLOD();
+	FSkeletalMeshRenderData const& SkeletalMeshRenderData = SkeletalMeshObject->GetSkeletalMeshRenderData();
+	FSkeletalMeshLODRenderData const* LodRenderData = &SkeletalMeshRenderData.LODRenderData[LodIndex];
 
+	// this mismatch can happen during map load for one frame sometimes, when predicted lod != actual lod
+	// TODO: avoid the mismatch in the first place, use actual lod to query NumInvocations
+	if (LodRenderData->RenderSections.Num() != InValidationData.NumInvocations)
+	{
+		return false;
+	}
+	
 	return true;
 }
 
diff --git a/Engine/Plugins/Animation/DeformerGraph/Source/OptimusCore/Private/DataInterfaces/OptimusDataInterfaceSkinnedMeshWrite.cpp b/Engine/Plugins/Animation/DeformerGraph/Source/OptimusCore/Private/DataInterfaces/OptimusDataInterfaceSkinnedMeshWrite.cpp
--- a/Engine/Plugins/Animation/DeformerGraph/Source/OptimusCore/Private/DataInterfaces/OptimusDataInterfaceSkinnedMeshWrite.cpp
+++ b/Engine/Plugins/Animation/DeformerGraph/Source/OptimusCore/Private/DataInterfaces/OptimusDataInterfaceSkinnedMeshWrite.cpp
@@ -215,6 +215,17 @@ bool FOptimusSkinnedMeshWriteDataProviderProxy::IsValid(FValidationData const& I
 	{
 		return false;
 	}
+	
+	const int32 LodIndex = SkeletalMeshObject->GetLOD();
+	FSkeletalMeshRenderData const& SkeletalMeshRenderData = SkeletalMeshObject->GetSkeletalMeshRenderData();
+	FSkeletalMeshLODRenderData const* LodRenderData = &SkeletalMeshRenderData.LODRenderData[LodIndex];
+
+	// this mismatch can happen during map load for one frame sometimes, when predicted lod != actual lod
+	// TODO: avoid the mismatch in the first place, use actual lod to query NumInvocations
+	if (LodRenderData->RenderSections.Num() != InValidationData.NumInvocations)
+	{
+		return false;
+	}
 
 	return true;
 }
