`OptimusCore` is reachable from transactional code so needs to be enabled for AutoRTFM.

#rnx

diff --git a/Engine/Plugins/Animation/DeformerGraph/Source/OptimusCore/OptimusCore.Build.cs b/Engine/Plugins/Animation/DeformerGraph/Source/OptimusCore/OptimusCore.Build.cs
--- a/Engine/Plugins/Animation/DeformerGraph/Source/OptimusCore/OptimusCore.Build.cs
+++ b/Engine/Plugins/Animation/DeformerGraph/Source/OptimusCore/OptimusCore.Build.cs
@@ -43,8 +43,6 @@ namespace UnrealBuildTool.Rules
 		            }
 	            );
             }
-			
-			bDisableAutoRTFMInstrumentation = true;
         }
     }
 }
