[Deformer Graph] now uses weak reference to skinnedmesh from data providers to avoid GC issues

diff --git a/Engine/Plugins/Animation/DeformerGraph/Source/OptimusCore/Private/DataInterfaces/OptimusDataInterfaceAdvancedSkeleton.cpp b/Engine/Plugins/Animation/DeformerGraph/Source/OptimusCore/Private/DataInterfaces/OptimusDataInterfaceAdvancedSkeleton.cpp
--- a/Engine/Plugins/Animation/DeformerGraph/Source/OptimusCore/Private/DataInterfaces/OptimusDataInterfaceAdvancedSkeleton.cpp
+++ b/Engine/Plugins/Animation/DeformerGraph/Source/OptimusCore/Private/DataInterfaces/OptimusDataInterfaceAdvancedSkeleton.cpp
@@ -542,7 +542,7 @@ void FOptimusBoneTransformBuffer::AllocateResources(FRDGBuilder& GraphBuilder)
 
 void UOptimusAdvancedSkeletonDataProvider::SetDeformerInstance(UOptimusDeformerInstance* InInstance)
 {
-	DeformerInstance = InInstance;
+	WeakDeformerInstance = InInstance;
 }
 
 void UOptimusAdvancedSkeletonDataProvider::Init(const UOptimusAdvancedSkeletonDataInterface* InDataInterface, USkeletalMeshComponent* InSkeletalMesh)
@@ -721,9 +721,11 @@ FComputeDataProviderRenderProxy* UOptimusAdvancedSkeletonDataProvider::GetRender
 {
 	FOptimusAdvancedSkeletonDataProviderProxy* Proxy = new FOptimusAdvancedSkeletonDataProviderProxy();
 
-	if (SkeletalMesh && SkeletalMesh->MeshObject)
+	if (SkeletalMesh.IsValid() && SkeletalMesh->MeshObject)
 	{
-		if (const UOptimusAdvancedSkeletonDataInterface* DataInterface = WeakDataInterface.Get())
+		const UOptimusAdvancedSkeletonDataInterface* DataInterface = WeakDataInterface.Get();
+		UOptimusDeformerInstance* DeformerInstance = WeakDeformerInstance.Get();
+		if (DeformerInstance && DataInterface)
 		{
 			FOptimusValueContainerStruct ValueContainer =
 					DeformerInstance->GetDataInterfacePropertyOverride(
diff --git a/Engine/Plugins/Animation/DeformerGraph/Source/OptimusCore/Private/DataInterfaces/OptimusDataInterfaceAdvancedSkeleton.h b/Engine/Plugins/Animation/DeformerGraph/Source/OptimusCore/Private/DataInterfaces/OptimusDataInterfaceAdvancedSkeleton.h
--- a/Engine/Plugins/Animation/DeformerGraph/Source/OptimusCore/Private/DataInterfaces/OptimusDataInterfaceAdvancedSkeleton.h
+++ b/Engine/Plugins/Animation/DeformerGraph/Source/OptimusCore/Private/DataInterfaces/OptimusDataInterfaceAdvancedSkeleton.h
@@ -218,8 +218,7 @@ class UOptimusAdvancedSkeletonDataProvider :
 public:
 	void SetDeformerInstance(UOptimusDeformerInstance* InInstance) override;
 	
-	UPROPERTY(EditAnywhere, BlueprintReadWrite, Category = Binding)
-	TObjectPtr<USkeletalMeshComponent> SkeletalMesh = nullptr;
+	TWeakObjectPtr<USkeletalMeshComponent> SkeletalMesh = nullptr;
 
 	TArray<FOptimusAnimAttributeBufferRuntimeData> AttributeBufferRuntimeData;
 
@@ -244,8 +243,7 @@ public:
 
 	void ComputeBoneTransformsForLayeredSkinning(TArray<FTransform>& OutBoneTransforms, FSkeletalMeshLODRenderData const& InLodRenderData, const FReferenceSkeleton& InRefSkeleton);
 private:
-	UPROPERTY()
-	TObjectPtr<UOptimusDeformerInstance> DeformerInstance;
+	TWeakObjectPtr<UOptimusDeformerInstance> WeakDeformerInstance;
 
 	TWeakObjectPtr<const UOptimusAdvancedSkeletonDataInterface> WeakDataInterface;
 };
diff --git a/Engine/Plugins/Animation/DeformerGraph/Source/OptimusCore/Private/DataInterfaces/OptimusDataInterfaceAnimAttribute.cpp b/Engine/Plugins/Animation/DeformerGraph/Source/OptimusCore/Private/DataInterfaces/OptimusDataInterfaceAnimAttribute.cpp
--- a/Engine/Plugins/Animation/DeformerGraph/Source/OptimusCore/Private/DataInterfaces/OptimusDataInterfaceAnimAttribute.cpp
+++ b/Engine/Plugins/Animation/DeformerGraph/Source/OptimusCore/Private/DataInterfaces/OptimusDataInterfaceAnimAttribute.cpp
@@ -578,7 +578,7 @@ void UOptimusAnimAttributeDataProvider::Init(
 	{
 		// Skip this step in case that there is no skeletal mesh, this can happen if
 		// the preview scene does not have a preview mesh assigned
-		if (SkeletalMesh && SkeletalMesh->GetSkeletalMeshAsset())
+		if (SkeletalMesh.IsValid() && SkeletalMesh->GetSkeletalMeshAsset())
 		{
 			if (Attribute.BoneName != NAME_None)
 			{
diff --git a/Engine/Plugins/Animation/DeformerGraph/Source/OptimusCore/Private/DataInterfaces/OptimusDataInterfaceAnimAttribute.h b/Engine/Plugins/Animation/DeformerGraph/Source/OptimusCore/Private/DataInterfaces/OptimusDataInterfaceAnimAttribute.h
--- a/Engine/Plugins/Animation/DeformerGraph/Source/OptimusCore/Private/DataInterfaces/OptimusDataInterfaceAnimAttribute.h
+++ b/Engine/Plugins/Animation/DeformerGraph/Source/OptimusCore/Private/DataInterfaces/OptimusDataInterfaceAnimAttribute.h
@@ -203,8 +203,7 @@ public:
 		const TArray<FOptimusAnimAttributeDescription>& InAttributeArray
 	);
 	
-	UPROPERTY(EditAnywhere, BlueprintReadWrite, Category = Binding)
-	TObjectPtr<USkeletalMeshComponent> SkeletalMesh = nullptr;
+	TWeakObjectPtr<USkeletalMeshComponent> SkeletalMesh = nullptr;
 
 	TArray<FOptimusAnimAttributeRuntimeData> AttributeRuntimeData;
 
diff --git a/Engine/Plugins/Animation/DeformerGraph/Source/OptimusCore/Private/DataInterfaces/OptimusDataInterfaceCloth.cpp b/Engine/Plugins/Animation/DeformerGraph/Source/OptimusCore/Private/DataInterfaces/OptimusDataInterfaceCloth.cpp
--- a/Engine/Plugins/Animation/DeformerGraph/Source/OptimusCore/Private/DataInterfaces/OptimusDataInterfaceCloth.cpp
+++ b/Engine/Plugins/Animation/DeformerGraph/Source/OptimusCore/Private/DataInterfaces/OptimusDataInterfaceCloth.cpp
@@ -126,7 +126,7 @@ UComputeDataProvider* UOptimusClothDataInterface::CreateDataProvider(TObjectPtr<
 
 FComputeDataProviderRenderProxy* UOptimusClothDataProvider::GetRenderProxy()
 {
-	return new FOptimusClothDataProviderProxy(SkinnedMesh);
+	return new FOptimusClothDataProviderProxy(SkinnedMesh.Get());
 }
 
 
diff --git a/Engine/Plugins/Animation/DeformerGraph/Source/OptimusCore/Private/DataInterfaces/OptimusDataInterfaceCloth.h b/Engine/Plugins/Animation/DeformerGraph/Source/OptimusCore/Private/DataInterfaces/OptimusDataInterfaceCloth.h
--- a/Engine/Plugins/Animation/DeformerGraph/Source/OptimusCore/Private/DataInterfaces/OptimusDataInterfaceCloth.h
+++ b/Engine/Plugins/Animation/DeformerGraph/Source/OptimusCore/Private/DataInterfaces/OptimusDataInterfaceCloth.h
@@ -47,8 +47,7 @@ class UOptimusClothDataProvider : public UComputeDataProvider
 	GENERATED_BODY()
 
 public:
-	UPROPERTY(EditAnywhere, BlueprintReadWrite, Category = Binding)
-	TObjectPtr<USkinnedMeshComponent> SkinnedMesh = nullptr;
+	TWeakObjectPtr<USkinnedMeshComponent> SkinnedMesh = nullptr;
 
 	//~ Begin UComputeDataProvider Interface
 	FComputeDataProviderRenderProxy* GetRenderProxy() override;
diff --git a/Engine/Plugins/Animation/DeformerGraph/Source/OptimusCore/Private/DataInterfaces/OptimusDataInterfaceConnectivity.cpp b/Engine/Plugins/Animation/DeformerGraph/Source/OptimusCore/Private/DataInterfaces/OptimusDataInterfaceConnectivity.cpp
--- a/Engine/Plugins/Animation/DeformerGraph/Source/OptimusCore/Private/DataInterfaces/OptimusDataInterfaceConnectivity.cpp
+++ b/Engine/Plugins/Animation/DeformerGraph/Source/OptimusCore/Private/DataInterfaces/OptimusDataInterfaceConnectivity.cpp
@@ -229,7 +229,7 @@ UComputeDataProvider* UOptimusConnectivityDataInterface::CreateDataProvider(TObj
 
 FComputeDataProviderRenderProxy* UOptimusConnectivityDataProvider::GetRenderProxy()
 {
-	return new FOptimusConnectivityDataProviderProxy(SkinnedMesh, AdjacencyBufferPerLod);
+	return new FOptimusConnectivityDataProviderProxy(SkinnedMesh.Get(), AdjacencyBufferPerLod);
 }
 
 
diff --git a/Engine/Plugins/Animation/DeformerGraph/Source/OptimusCore/Private/DataInterfaces/OptimusDataInterfaceConnectivity.h b/Engine/Plugins/Animation/DeformerGraph/Source/OptimusCore/Private/DataInterfaces/OptimusDataInterfaceConnectivity.h
--- a/Engine/Plugins/Animation/DeformerGraph/Source/OptimusCore/Private/DataInterfaces/OptimusDataInterfaceConnectivity.h
+++ b/Engine/Plugins/Animation/DeformerGraph/Source/OptimusCore/Private/DataInterfaces/OptimusDataInterfaceConnectivity.h
@@ -53,8 +53,7 @@ class UOptimusConnectivityDataProvider : public UComputeDataProvider
 	GENERATED_BODY()
 
 public:
-	UPROPERTY(EditAnywhere, BlueprintReadWrite, Category = Binding)
-	TObjectPtr<USkinnedMeshComponent> SkinnedMesh = nullptr;
+	TWeakObjectPtr<USkinnedMeshComponent> SkinnedMesh = nullptr;
 
 	//~ Begin UComputeDataProvider Interface
 	FComputeDataProviderRenderProxy* GetRenderProxy() override;
diff --git a/Engine/Plugins/Animation/DeformerGraph/Source/OptimusCore/Private/DataInterfaces/OptimusDataInterfaceDebugDraw.cpp b/Engine/Plugins/Animation/DeformerGraph/Source/OptimusCore/Private/DataInterfaces/OptimusDataInterfaceDebugDraw.cpp
--- a/Engine/Plugins/Animation/DeformerGraph/Source/OptimusCore/Private/DataInterfaces/OptimusDataInterfaceDebugDraw.cpp
+++ b/Engine/Plugins/Animation/DeformerGraph/Source/OptimusCore/Private/DataInterfaces/OptimusDataInterfaceDebugDraw.cpp
@@ -139,7 +139,7 @@ UComputeDataProvider* UOptimusDebugDrawDataInterface::CreateDataProvider(TObject
 
 FComputeDataProviderRenderProxy* UOptimusDebugDrawDataProvider::GetRenderProxy()
 {
-	return new FOptimusDebugDrawDataProviderProxy(PrimitiveComponent, DebugDrawParameters);
+	return new FOptimusDebugDrawDataProviderProxy(PrimitiveComponent.Get(), DebugDrawParameters);
 }
 
 
diff --git a/Engine/Plugins/Animation/DeformerGraph/Source/OptimusCore/Private/DataInterfaces/OptimusDataInterfaceDebugDraw.h b/Engine/Plugins/Animation/DeformerGraph/Source/OptimusCore/Private/DataInterfaces/OptimusDataInterfaceDebugDraw.h
--- a/Engine/Plugins/Animation/DeformerGraph/Source/OptimusCore/Private/DataInterfaces/OptimusDataInterfaceDebugDraw.h
+++ b/Engine/Plugins/Animation/DeformerGraph/Source/OptimusCore/Private/DataInterfaces/OptimusDataInterfaceDebugDraw.h
@@ -88,10 +88,8 @@ class UOptimusDebugDrawDataProvider : public UComputeDataProvider
 	GENERATED_BODY()
 
 public:
-	UPROPERTY(EditAnywhere, BlueprintReadWrite, Category = Binding)
-	TObjectPtr<UPrimitiveComponent> PrimitiveComponent = nullptr;
+	TWeakObjectPtr<UPrimitiveComponent> PrimitiveComponent = nullptr;
 
-	UPROPERTY()
 	FOptimusDebugDrawParameters DebugDrawParameters;
 
 	//~ Begin UComputeDataProvider Interface
diff --git a/Engine/Plugins/Animation/DeformerGraph/Source/OptimusCore/Private/DataInterfaces/OptimusDataInterfaceDuplicateVertices.cpp b/Engine/Plugins/Animation/DeformerGraph/Source/OptimusCore/Private/DataInterfaces/OptimusDataInterfaceDuplicateVertices.cpp
--- a/Engine/Plugins/Animation/DeformerGraph/Source/OptimusCore/Private/DataInterfaces/OptimusDataInterfaceDuplicateVertices.cpp
+++ b/Engine/Plugins/Animation/DeformerGraph/Source/OptimusCore/Private/DataInterfaces/OptimusDataInterfaceDuplicateVertices.cpp
@@ -113,7 +113,7 @@ UComputeDataProvider* UOptimusDuplicateVerticesDataInterface::CreateDataProvider
 
 FComputeDataProviderRenderProxy* UOptimusDuplicateVerticesDataProvider::GetRenderProxy()
 {
-	return new FOptimusDuplicateVerticesDataProviderProxy(SkinnedMesh);
+	return new FOptimusDuplicateVerticesDataProviderProxy(SkinnedMesh.Get());
 }
 
 
diff --git a/Engine/Plugins/Animation/DeformerGraph/Source/OptimusCore/Private/DataInterfaces/OptimusDataInterfaceDuplicateVertices.h b/Engine/Plugins/Animation/DeformerGraph/Source/OptimusCore/Private/DataInterfaces/OptimusDataInterfaceDuplicateVertices.h
--- a/Engine/Plugins/Animation/DeformerGraph/Source/OptimusCore/Private/DataInterfaces/OptimusDataInterfaceDuplicateVertices.h
+++ b/Engine/Plugins/Animation/DeformerGraph/Source/OptimusCore/Private/DataInterfaces/OptimusDataInterfaceDuplicateVertices.h
@@ -51,8 +51,7 @@ class UOptimusDuplicateVerticesDataProvider : public UComputeDataProvider
 	GENERATED_BODY()
 
 public:
-	UPROPERTY(EditAnywhere, BlueprintReadWrite, Category = Binding)
-	TObjectPtr<USkinnedMeshComponent> SkinnedMesh = nullptr;
+	TWeakObjectPtr<USkinnedMeshComponent> SkinnedMesh = nullptr;
 
 	//~ Begin UComputeDataProvider Interface
 	FComputeDataProviderRenderProxy* GetRenderProxy() override;
diff --git a/Engine/Plugins/Animation/DeformerGraph/Source/OptimusCore/Private/DataInterfaces/OptimusDataInterfaceGraph.cpp b/Engine/Plugins/Animation/DeformerGraph/Source/OptimusCore/Private/DataInterfaces/OptimusDataInterfaceGraph.cpp
--- a/Engine/Plugins/Animation/DeformerGraph/Source/OptimusCore/Private/DataInterfaces/OptimusDataInterfaceGraph.cpp
+++ b/Engine/Plugins/Animation/DeformerGraph/Source/OptimusCore/Private/DataInterfaces/OptimusDataInterfaceGraph.cpp
@@ -109,11 +109,8 @@ void UOptimusGraphDataInterface::GetHLSL(FString& OutHLSL, FString const& InData
 UComputeDataProvider* UOptimusGraphDataInterface::CreateDataProvider(TObjectPtr<UObject> InBinding, uint64 InInputMask, uint64 InOutputMask) const
 {
 	UOptimusGraphDataProvider* Provider = NewObject<UOptimusGraphDataProvider>();
-	Provider->Init(Cast<UMeshComponent>(InBinding), Variables, ParameterBufferSize);
+	Provider->Init(Variables, ParameterBufferSize);
 
-
-	
-	
 	return Provider;
 }
 
@@ -137,9 +134,8 @@ PRAGMA_ENABLE_DEPRECATION_WARNINGS
 }
 
 
-void UOptimusGraphDataProvider::Init(UMeshComponent* InMeshComponent, const TArray<FOptimusGraphVariableDescription>& InVariables, int32 InParameterBufferSize)
+void UOptimusGraphDataProvider::Init(const TArray<FOptimusGraphVariableDescription>& InVariables, int32 InParameterBufferSize)
 {
-	MeshComponent = InMeshComponent;
 	Variables = InVariables;
 	ParameterBufferSize = InParameterBufferSize;
 
@@ -172,12 +168,12 @@ void UOptimusGraphDataProvider::Init(UMeshComponent* InMeshComponent, const TArr
 
 FComputeDataProviderRenderProxy* UOptimusGraphDataProvider::GetRenderProxy()
 {
-	return new FOptimusGraphDataProviderProxy(DeformerInstance, Variables, ParameterBufferSize, ParameterArrayMetadata);
+	return new FOptimusGraphDataProviderProxy(WeakDeformerInstance.Get(), Variables, ParameterBufferSize, ParameterArrayMetadata);
 }
 
 void UOptimusGraphDataProvider::SetDeformerInstance(UOptimusDeformerInstance* InInstance)
 {
-	DeformerInstance = InInstance;
+	WeakDeformerInstance = InInstance;
 }
 
 FOptimusGraphDataProviderProxy::FOptimusGraphDataProviderProxy(
diff --git a/Engine/Plugins/Animation/DeformerGraph/Source/OptimusCore/Private/DataInterfaces/OptimusDataInterfaceGraph.h b/Engine/Plugins/Animation/DeformerGraph/Source/OptimusCore/Private/DataInterfaces/OptimusDataInterfaceGraph.h
--- a/Engine/Plugins/Animation/DeformerGraph/Source/OptimusCore/Private/DataInterfaces/OptimusDataInterfaceGraph.h
+++ b/Engine/Plugins/Animation/DeformerGraph/Source/OptimusCore/Private/DataInterfaces/OptimusDataInterfaceGraph.h
@@ -91,8 +91,6 @@ class UOptimusGraphDataProvider :
 	GENERATED_BODY()
 
 public:
-	UPROPERTY(EditAnywhere, BlueprintReadWrite, Category = Binding)
-	TObjectPtr<UMeshComponent> MeshComponent = nullptr;
 
 	UPROPERTY()
 	TArray<FOptimusGraphVariableDescription> Variables;
@@ -109,7 +107,7 @@ public:
 	TArray<FArrayMetadata> ParameterArrayMetadata;
 	
 	
-	void Init(UMeshComponent* InMeshComponent, const TArray<FOptimusGraphVariableDescription>& InVariables, int32 InParameterBufferSize);
+	void Init(const TArray<FOptimusGraphVariableDescription>& InVariables, int32 InParameterBufferSize);
 	
 	//~ Begin UComputeDataProvider Interface
 	FComputeDataProviderRenderProxy* GetRenderProxy() override;
@@ -120,8 +118,7 @@ public:
 	//~ End IOptimusDeformerInstanceAccessor Interface
 
 private:
-	UPROPERTY()
-	TObjectPtr<UOptimusDeformerInstance> DeformerInstance = nullptr;
+	TWeakObjectPtr<UOptimusDeformerInstance> WeakDeformerInstance = nullptr;
 };
 
 class FOptimusGraphDataProviderProxy : public FComputeDataProviderRenderProxy
diff --git a/Engine/Plugins/Animation/DeformerGraph/Source/OptimusCore/Private/DataInterfaces/OptimusDataInterfaceHalfEdge.cpp b/Engine/Plugins/Animation/DeformerGraph/Source/OptimusCore/Private/DataInterfaces/OptimusDataInterfaceHalfEdge.cpp
--- a/Engine/Plugins/Animation/DeformerGraph/Source/OptimusCore/Private/DataInterfaces/OptimusDataInterfaceHalfEdge.cpp
+++ b/Engine/Plugins/Animation/DeformerGraph/Source/OptimusCore/Private/DataInterfaces/OptimusDataInterfaceHalfEdge.cpp
@@ -201,7 +201,7 @@ FComputeDataProviderRenderProxy* UOptimusHalfEdgeDataProvider::GetRenderProxy()
 	}
 #endif
 		
-	return new FOptimusHalfEdgeDataProviderProxy(SkinnedMesh, OnDemandHalfEdgeBuffers);
+	return new FOptimusHalfEdgeDataProviderProxy(SkinnedMesh.Get(), OnDemandHalfEdgeBuffers);
 }
 
 void UOptimusHalfEdgeDataProvider::SetOptimusDeformerAssetPath(const FTopLevelAssetPath& InPath)
diff --git a/Engine/Plugins/Animation/DeformerGraph/Source/OptimusCore/Private/DataInterfaces/OptimusDataInterfaceHalfEdge.h b/Engine/Plugins/Animation/DeformerGraph/Source/OptimusCore/Private/DataInterfaces/OptimusDataInterfaceHalfEdge.h
--- a/Engine/Plugins/Animation/DeformerGraph/Source/OptimusCore/Private/DataInterfaces/OptimusDataInterfaceHalfEdge.h
+++ b/Engine/Plugins/Animation/DeformerGraph/Source/OptimusCore/Private/DataInterfaces/OptimusDataInterfaceHalfEdge.h
@@ -61,8 +61,7 @@ class UOptimusHalfEdgeDataProvider :
 	GENERATED_BODY()
 
 public:
-	UPROPERTY(EditAnywhere, BlueprintReadWrite, Category = Binding)
-	TObjectPtr<USkinnedMeshComponent> SkinnedMesh = nullptr;
+	TWeakObjectPtr<USkinnedMeshComponent> SkinnedMesh = nullptr;
 
 	//~ Begin UObject Interface
 	void BeginDestroy() override;
diff --git a/Engine/Plugins/Animation/DeformerGraph/Source/OptimusCore/Private/DataInterfaces/OptimusDataInterfaceMorphTarget.cpp b/Engine/Plugins/Animation/DeformerGraph/Source/OptimusCore/Private/DataInterfaces/OptimusDataInterfaceMorphTarget.cpp
--- a/Engine/Plugins/Animation/DeformerGraph/Source/OptimusCore/Private/DataInterfaces/OptimusDataInterfaceMorphTarget.cpp
+++ b/Engine/Plugins/Animation/DeformerGraph/Source/OptimusCore/Private/DataInterfaces/OptimusDataInterfaceMorphTarget.cpp
@@ -100,7 +100,7 @@ UComputeDataProvider* UOptimusMorphTargetDataInterface::CreateDataProvider(TObje
 
 FComputeDataProviderRenderProxy* UOptimusMorphTargetDataProvider::GetRenderProxy()
 {
-	return new FOptimusMorphTargetDataProviderProxy(SkinnedMesh);
+	return new FOptimusMorphTargetDataProviderProxy(SkinnedMesh.Get());
 }
 
 
diff --git a/Engine/Plugins/Animation/DeformerGraph/Source/OptimusCore/Private/DataInterfaces/OptimusDataInterfaceMorphTarget.h b/Engine/Plugins/Animation/DeformerGraph/Source/OptimusCore/Private/DataInterfaces/OptimusDataInterfaceMorphTarget.h
--- a/Engine/Plugins/Animation/DeformerGraph/Source/OptimusCore/Private/DataInterfaces/OptimusDataInterfaceMorphTarget.h
+++ b/Engine/Plugins/Animation/DeformerGraph/Source/OptimusCore/Private/DataInterfaces/OptimusDataInterfaceMorphTarget.h
@@ -49,8 +49,7 @@ class UOptimusMorphTargetDataProvider : public UComputeDataProvider
 	GENERATED_BODY()
 
 public:
-	UPROPERTY(EditAnywhere, BlueprintReadWrite, Category = Binding)
-	TObjectPtr<USkinnedMeshComponent> SkinnedMesh = nullptr;
+	TWeakObjectPtr<USkinnedMeshComponent> SkinnedMesh = nullptr;
 
 	//~ Begin UComputeDataProvider Interface
 	FComputeDataProviderRenderProxy* GetRenderProxy() override;
diff --git a/Engine/Plugins/Animation/DeformerGraph/Source/OptimusCore/Private/DataInterfaces/OptimusDataInterfaceScene.cpp b/Engine/Plugins/Animation/DeformerGraph/Source/OptimusCore/Private/DataInterfaces/OptimusDataInterfaceScene.cpp
--- a/Engine/Plugins/Animation/DeformerGraph/Source/OptimusCore/Private/DataInterfaces/OptimusDataInterfaceScene.cpp
+++ b/Engine/Plugins/Animation/DeformerGraph/Source/OptimusCore/Private/DataInterfaces/OptimusDataInterfaceScene.cpp
@@ -94,7 +94,7 @@ UComputeDataProvider* UOptimusSceneDataInterface::CreateDataProvider(TObjectPtr<
 
 FComputeDataProviderRenderProxy* UOptimusSceneDataProvider::GetRenderProxy()
 {
-	return new FOptimusSceneDataProviderProxy(SceneComponent);
+	return new FOptimusSceneDataProviderProxy(SceneComponent.Get());
 }
 
 
diff --git a/Engine/Plugins/Animation/DeformerGraph/Source/OptimusCore/Private/DataInterfaces/OptimusDataInterfaceScene.h b/Engine/Plugins/Animation/DeformerGraph/Source/OptimusCore/Private/DataInterfaces/OptimusDataInterfaceScene.h
--- a/Engine/Plugins/Animation/DeformerGraph/Source/OptimusCore/Private/DataInterfaces/OptimusDataInterfaceScene.h
+++ b/Engine/Plugins/Animation/DeformerGraph/Source/OptimusCore/Private/DataInterfaces/OptimusDataInterfaceScene.h
@@ -44,8 +44,7 @@ class UOptimusSceneDataProvider : public UComputeDataProvider
 	GENERATED_BODY()
 
 public:
-	UPROPERTY(EditAnywhere, BlueprintReadWrite, Category = Binding)
-	TObjectPtr<USceneComponent> SceneComponent = nullptr;
+	TWeakObjectPtr<USceneComponent> SceneComponent = nullptr;
 
 	//~ Begin UComputeDataProvider Interface
 	FComputeDataProviderRenderProxy* GetRenderProxy() override;
diff --git a/Engine/Plugins/Animation/DeformerGraph/Source/OptimusCore/Private/DataInterfaces/OptimusDataInterfaceSkeleton.cpp b/Engine/Plugins/Animation/DeformerGraph/Source/OptimusCore/Private/DataInterfaces/OptimusDataInterfaceSkeleton.cpp
--- a/Engine/Plugins/Animation/DeformerGraph/Source/OptimusCore/Private/DataInterfaces/OptimusDataInterfaceSkeleton.cpp
+++ b/Engine/Plugins/Animation/DeformerGraph/Source/OptimusCore/Private/DataInterfaces/OptimusDataInterfaceSkeleton.cpp
@@ -124,7 +124,7 @@ UComputeDataProvider* UOptimusSkeletonDataInterface::CreateDataProvider(TObjectP
 
 FComputeDataProviderRenderProxy* UOptimusSkeletonDataProvider::GetRenderProxy()
 {
-	return new FOptimusSkeletonDataProviderProxy(SkinnedMesh);
+	return new FOptimusSkeletonDataProviderProxy(SkinnedMesh.Get());
 }
 
 
diff --git a/Engine/Plugins/Animation/DeformerGraph/Source/OptimusCore/Private/DataInterfaces/OptimusDataInterfaceSkeleton.h b/Engine/Plugins/Animation/DeformerGraph/Source/OptimusCore/Private/DataInterfaces/OptimusDataInterfaceSkeleton.h
--- a/Engine/Plugins/Animation/DeformerGraph/Source/OptimusCore/Private/DataInterfaces/OptimusDataInterfaceSkeleton.h
+++ b/Engine/Plugins/Animation/DeformerGraph/Source/OptimusCore/Private/DataInterfaces/OptimusDataInterfaceSkeleton.h
@@ -50,8 +50,7 @@ class UOptimusSkeletonDataProvider : public UComputeDataProvider
 	GENERATED_BODY()
 
 public:
-	UPROPERTY(EditAnywhere, BlueprintReadWrite, Category = Binding)
-	TObjectPtr<USkinnedMeshComponent> SkinnedMesh = nullptr;
+	TWeakObjectPtr<USkinnedMeshComponent> SkinnedMesh = nullptr;
 
 	//~ Begin UComputeDataProvider Interface
 	FComputeDataProviderRenderProxy* GetRenderProxy() override;
diff --git a/Engine/Plugins/Animation/DeformerGraph/Source/OptimusCore/Private/DataInterfaces/OptimusDataInterfaceSkinWeightsAsVertexMask.cpp b/Engine/Plugins/Animation/DeformerGraph/Source/OptimusCore/Private/DataInterfaces/OptimusDataInterfaceSkinWeightsAsVertexMask.cpp
--- a/Engine/Plugins/Animation/DeformerGraph/Source/OptimusCore/Private/DataInterfaces/OptimusDataInterfaceSkinWeightsAsVertexMask.cpp
+++ b/Engine/Plugins/Animation/DeformerGraph/Source/OptimusCore/Private/DataInterfaces/OptimusDataInterfaceSkinWeightsAsVertexMask.cpp
@@ -163,7 +163,7 @@ UComputeDataProvider* UOptimusSkinWeightsAsVertexMaskDataInterface::CreateDataPr
 
 void UOptimusSkinWeightsAsVertexMaskDataProvider::SetDeformerInstance(UOptimusDeformerInstance* InInstance)
 {
-	DeformerInstance = InInstance;
+	WeakDeformerInstance = InInstance;
 }
 
 void UOptimusSkinWeightsAsVertexMaskDataProvider::Init(const UOptimusSkinWeightsAsVertexMaskDataInterface* InDataInterface, USkeletalMeshComponent* InSkeletalMesh)
@@ -183,11 +183,11 @@ FComputeDataProviderRenderProxy* UOptimusSkinWeightsAsVertexMaskDataProvider::Ge
 {
 	FOptimusSkinWeightsAsVertexMaskDataProviderProxy* Proxy = new FOptimusSkinWeightsAsVertexMaskDataProviderProxy();
 
-	if (SkeletalMesh && SkeletalMesh->MeshObject)
+	if (SkeletalMesh.Get() && SkeletalMesh->MeshObject)
 	{
-
-		
-		if (const UOptimusSkinWeightsAsVertexMaskDataInterface* DataInterface = WeakDataInterface.Get())
+		const UOptimusSkinWeightsAsVertexMaskDataInterface* DataInterface = WeakDataInterface.Get();
+		UOptimusDeformerInstance* DeformerInstance = WeakDeformerInstance.Get();
+		if (DeformerInstance && DataInterface)
 		{
 			{
 				const FOptimusValueContainerStruct ValueContainer =
@@ -409,22 +409,23 @@ FComputeDataProviderRenderProxy* UOptimusSkinWeightsAsVertexMaskDataProvider::Ge
 					{
 						if (WeakThis.IsValid())
 						{
-							USkeletalMeshComponent* Mesh = WeakThis->SkeletalMesh;
-							const FReferenceSkeleton& RefSkeleton = Mesh->GetSkinnedAsset()->GetRefSkeleton();
+							if (USkeletalMeshComponent* Mesh = WeakThis->SkeletalMesh.Get())
+							{
+								const FReferenceSkeleton& RefSkeleton = Mesh->GetSkinnedAsset()->GetRefSkeleton();
 					
-							const TArray<FTransform>& CurrentComponentSpaceBoneMatrix = Mesh->GetComponentSpaceTransforms();
+								const TArray<FTransform>& CurrentComponentSpaceBoneMatrix = Mesh->GetComponentSpaceTransforms();
 
-							for (int32 Index = 0 ; Index < CurrentComponentSpaceBoneMatrix.Num(); Index++)
-							{
-								int32 ParentIndex = RefSkeleton.GetParentIndex(Index);
-								if (ParentIndex!=INDEX_NONE && WeakThis->CachedSelectedBones.Contains(ParentIndex))
+								for (int32 Index = 0 ; Index < CurrentComponentSpaceBoneMatrix.Num(); Index++)
 								{
-									const FTransform& Child =  CurrentComponentSpaceBoneMatrix[Index] * Mesh->GetComponentTransform();
-									const FTransform& Parent = CurrentComponentSpaceBoneMatrix[ParentIndex] * Mesh->GetComponentTransform();
-									DrawDebugLine(Mesh->GetWorld(), Child.GetLocation(), Parent.GetLocation(), WeakThis->DebugDrawColor, false, -1, SDPG_Foreground);	
-								}
+									int32 ParentIndex = RefSkeleton.GetParentIndex(Index);
+									if (ParentIndex!=INDEX_NONE && WeakThis->CachedSelectedBones.Contains(ParentIndex))
+									{
+										const FTransform& Child =  CurrentComponentSpaceBoneMatrix[Index] * Mesh->GetComponentTransform();
+										const FTransform& Parent = CurrentComponentSpaceBoneMatrix[ParentIndex] * Mesh->GetComponentTransform();
+										DrawDebugLine(Mesh->GetWorld(), Child.GetLocation(), Parent.GetLocation(), WeakThis->DebugDrawColor, false, -1, SDPG_Foreground);	
+									}
+								}	
 							}
-
 						}
 					}, TStatId(), NULL, ENamedThreads::GameThread);	
 #endif
diff --git a/Engine/Plugins/Animation/DeformerGraph/Source/OptimusCore/Private/DataInterfaces/OptimusDataInterfaceSkinWeightsAsVertexMask.h b/Engine/Plugins/Animation/DeformerGraph/Source/OptimusCore/Private/DataInterfaces/OptimusDataInterfaceSkinWeightsAsVertexMask.h
--- a/Engine/Plugins/Animation/DeformerGraph/Source/OptimusCore/Private/DataInterfaces/OptimusDataInterfaceSkinWeightsAsVertexMask.h
+++ b/Engine/Plugins/Animation/DeformerGraph/Source/OptimusCore/Private/DataInterfaces/OptimusDataInterfaceSkinWeightsAsVertexMask.h
@@ -98,8 +98,7 @@ public:
 	
 	void SetDeformerInstance(UOptimusDeformerInstance* InInstance) override;
 
-	UPROPERTY(EditAnywhere, BlueprintReadWrite, Category = Binding)
-	TObjectPtr<USkeletalMeshComponent> SkeletalMesh = nullptr;
+	TWeakObjectPtr<USkeletalMeshComponent> SkeletalMesh = nullptr;
 
 	FName SkinWeightProfile = NAME_None;
 	
@@ -119,8 +118,7 @@ public:
 
 
 private:
-	UPROPERTY()
-	TObjectPtr<UOptimusDeformerInstance> DeformerInstance;
+	TWeakObjectPtr<UOptimusDeformerInstance> WeakDeformerInstance;
 
 	TWeakObjectPtr<const UOptimusSkinWeightsAsVertexMaskDataInterface> WeakDataInterface;
 };
diff --git a/Engine/Plugins/Animation/DeformerGraph/Source/OptimusCore/Private/DataInterfaces/OptimusDataInterfaceSkinnedMesh.cpp b/Engine/Plugins/Animation/DeformerGraph/Source/OptimusCore/Private/DataInterfaces/OptimusDataInterfaceSkinnedMesh.cpp
--- a/Engine/Plugins/Animation/DeformerGraph/Source/OptimusCore/Private/DataInterfaces/OptimusDataInterfaceSkinnedMesh.cpp
+++ b/Engine/Plugins/Animation/DeformerGraph/Source/OptimusCore/Private/DataInterfaces/OptimusDataInterfaceSkinnedMesh.cpp
@@ -143,7 +143,7 @@ UComputeDataProvider* UOptimusSkinnedMeshDataInterface::CreateDataProvider(TObje
 
 FComputeDataProviderRenderProxy* UOptimusSkinnedMeshDataProvider::GetRenderProxy()
 {
-	return new FOptimusSkinnedMeshDataProviderProxy(SkinnedMesh);
+	return new FOptimusSkinnedMeshDataProviderProxy(SkinnedMesh.Get());
 }
 
 
diff --git a/Engine/Plugins/Animation/DeformerGraph/Source/OptimusCore/Private/DataInterfaces/OptimusDataInterfaceSkinnedMesh.h b/Engine/Plugins/Animation/DeformerGraph/Source/OptimusCore/Private/DataInterfaces/OptimusDataInterfaceSkinnedMesh.h
--- a/Engine/Plugins/Animation/DeformerGraph/Source/OptimusCore/Private/DataInterfaces/OptimusDataInterfaceSkinnedMesh.h
+++ b/Engine/Plugins/Animation/DeformerGraph/Source/OptimusCore/Private/DataInterfaces/OptimusDataInterfaceSkinnedMesh.h
@@ -49,8 +49,7 @@ class UOptimusSkinnedMeshDataProvider : public UComputeDataProvider
 	GENERATED_BODY()
 
 public:
-	UPROPERTY(EditAnywhere, BlueprintReadWrite, Category = Binding)
-	TObjectPtr<USkinnedMeshComponent> SkinnedMesh = nullptr;
+	TWeakObjectPtr<USkinnedMeshComponent> SkinnedMesh = nullptr;
 
 	//~ Begin UComputeDataProvider Interface
 	FComputeDataProviderRenderProxy* GetRenderProxy() override;
diff --git a/Engine/Plugins/Animation/DeformerGraph/Source/OptimusCore/Private/DataInterfaces/OptimusDataInterfaceSkinnedMeshRead.cpp b/Engine/Plugins/Animation/DeformerGraph/Source/OptimusCore/Private/DataInterfaces/OptimusDataInterfaceSkinnedMeshRead.cpp
--- a/Engine/Plugins/Animation/DeformerGraph/Source/OptimusCore/Private/DataInterfaces/OptimusDataInterfaceSkinnedMeshRead.cpp
+++ b/Engine/Plugins/Animation/DeformerGraph/Source/OptimusCore/Private/DataInterfaces/OptimusDataInterfaceSkinnedMeshRead.cpp
@@ -165,12 +165,15 @@ UComputeDataProvider* UOptimusSkinnedMeshReadDataInterface::CreateDataProvider(T
 
 FComputeDataProviderRenderProxy* UOptimusSkinnedMeshReadDataProvider::GetRenderProxy()
 {
-	return new FOptimusSkinnedMeshReadDataProviderProxy(SkinnedMesh, InputMask, DeformerInstance->OutputBuffersFromPreviousInstances, &LastLodIndexCachedByRenderProxy);
+	return new FOptimusSkinnedMeshReadDataProviderProxy(
+		SkinnedMesh.Get(),
+		InputMask,
+		WeakDeformerInstance.IsValid()? WeakDeformerInstance->OutputBuffersFromPreviousInstances : EMeshDeformerOutputBuffer::None, &LastLodIndexCachedByRenderProxy);
 }
 
 void UOptimusSkinnedMeshReadDataProvider::SetDeformerInstance(UOptimusDeformerInstance* InInstance)
 {
-	DeformerInstance = InInstance;
+	WeakDeformerInstance = InInstance;
 }
 
 FOptimusSkinnedMeshReadDataProviderProxy::FOptimusSkinnedMeshReadDataProviderProxy(USkinnedMeshComponent* InSkinnedMeshComponent, uint64 InInputMask, EMeshDeformerOutputBuffer InOutputBuffersFromPreviousInstances, int32* InLastLodIndexPtr) 
diff --git a/Engine/Plugins/Animation/DeformerGraph/Source/OptimusCore/Private/DataInterfaces/OptimusDataInterfaceSkinnedMeshRead.h b/Engine/Plugins/Animation/DeformerGraph/Source/OptimusCore/Private/DataInterfaces/OptimusDataInterfaceSkinnedMeshRead.h
--- a/Engine/Plugins/Animation/DeformerGraph/Source/OptimusCore/Private/DataInterfaces/OptimusDataInterfaceSkinnedMeshRead.h
+++ b/Engine/Plugins/Animation/DeformerGraph/Source/OptimusCore/Private/DataInterfaces/OptimusDataInterfaceSkinnedMeshRead.h
@@ -60,8 +60,7 @@ class UOptimusSkinnedMeshReadDataProvider :
 	GENERATED_BODY()
 
 public:
-	UPROPERTY(EditAnywhere, BlueprintReadWrite, Category = Binding)
-	TObjectPtr<USkinnedMeshComponent> SkinnedMesh = nullptr;
+	TWeakObjectPtr<USkinnedMeshComponent> SkinnedMesh = nullptr;
 
 	uint64 InputMask = 0;
 
@@ -77,8 +76,7 @@ public:
 	//~ End IOptimusDeformerInstanceAccessor Interface
 
 private:
-	UPROPERTY()
-	TObjectPtr<UOptimusDeformerInstance> DeformerInstance = nullptr;
+	TWeakObjectPtr<UOptimusDeformerInstance> WeakDeformerInstance = nullptr;
 };
 
 class FOptimusSkinnedMeshReadDataProviderProxy : public FComputeDataProviderRenderProxy
diff --git a/Engine/Plugins/Animation/DeformerGraph/Source/OptimusCore/Private/DataInterfaces/OptimusDataInterfaceSkinnedMeshVertexAttribute.cpp b/Engine/Plugins/Animation/DeformerGraph/Source/OptimusCore/Private/DataInterfaces/OptimusDataInterfaceSkinnedMeshVertexAttribute.cpp
--- a/Engine/Plugins/Animation/DeformerGraph/Source/OptimusCore/Private/DataInterfaces/OptimusDataInterfaceSkinnedMeshVertexAttribute.cpp
+++ b/Engine/Plugins/Animation/DeformerGraph/Source/OptimusCore/Private/DataInterfaces/OptimusDataInterfaceSkinnedMeshVertexAttribute.cpp
@@ -123,7 +123,9 @@ UComputeDataProvider* UOptimusSkinnedMeshVertexAttributeDataInterface::CreateDat
 
 FComputeDataProviderRenderProxy* UOptimusSkinnedMeshVertexAttributeDataProvider::GetRenderProxy()
 {
-	if (const UOptimusSkinnedMeshVertexAttributeDataInterface* DataInterface = WeakDataInterface.Get())
+	const UOptimusSkinnedMeshVertexAttributeDataInterface* DataInterface = WeakDataInterface.Get();
+	UOptimusDeformerInstance* DeformerInstance = WeakDeformerInstance.Get();
+	if (DeformerInstance && DataInterface)
 	{
 		{
 			FOptimusValueContainerStruct ValueContainer =
@@ -154,12 +156,12 @@ FComputeDataProviderRenderProxy* UOptimusSkinnedMeshVertexAttributeDataProvider:
 		}
 	}
 	
-	return new FOptimusSkinnedMeshVertexAttributeDataProviderProxy(SkinnedMeshComponent, AttributeName, DefaultValue);
+	return new FOptimusSkinnedMeshVertexAttributeDataProviderProxy(SkinnedMeshComponent.Get(), AttributeName, DefaultValue);
 }
 
 void UOptimusSkinnedMeshVertexAttributeDataProvider::SetDeformerInstance(UOptimusDeformerInstance* InInstance)
 {
-	DeformerInstance = InInstance;
+	WeakDeformerInstance = InInstance;
 }
 
 
@@ -169,7 +171,7 @@ FOptimusSkinnedMeshVertexAttributeDataProviderProxy::FOptimusSkinnedMeshVertexAt
 	float InDefaultValue
 	)
 {
-	SkeletalMeshObject = InSkinnedMeshComponent->MeshObject;
+	SkeletalMeshObject = InSkinnedMeshComponent ? InSkinnedMeshComponent->MeshObject : nullptr;
 	AttributeName = InAttributeName;
 	DefaultValue = InDefaultValue;
 }
diff --git a/Engine/Plugins/Animation/DeformerGraph/Source/OptimusCore/Private/DataInterfaces/OptimusDataInterfaceSkinnedMeshVertexAttribute.h b/Engine/Plugins/Animation/DeformerGraph/Source/OptimusCore/Private/DataInterfaces/OptimusDataInterfaceSkinnedMeshVertexAttribute.h
--- a/Engine/Plugins/Animation/DeformerGraph/Source/OptimusCore/Private/DataInterfaces/OptimusDataInterfaceSkinnedMeshVertexAttribute.h
+++ b/Engine/Plugins/Animation/DeformerGraph/Source/OptimusCore/Private/DataInterfaces/OptimusDataInterfaceSkinnedMeshVertexAttribute.h
@@ -58,8 +58,7 @@ class UOptimusSkinnedMeshVertexAttributeDataProvider :
 	GENERATED_BODY()
 
 public:
-	UPROPERTY(EditAnywhere, BlueprintReadWrite, Category = Binding)
-	TObjectPtr<USkinnedMeshComponent> SkinnedMeshComponent = nullptr;
+	TWeakObjectPtr<USkinnedMeshComponent> SkinnedMeshComponent = nullptr;
 
 	UPROPERTY(EditAnywhere, BlueprintReadWrite, Category = VertexAttribute)
 	FName AttributeName;
@@ -77,8 +76,7 @@ public:
 
 	
 private:
-	UPROPERTY()
-	TObjectPtr<UOptimusDeformerInstance> DeformerInstance;
+	TWeakObjectPtr<UOptimusDeformerInstance> WeakDeformerInstance;
 };
 
 class FOptimusSkinnedMeshVertexAttributeDataProviderProxy : public FComputeDataProviderRenderProxy
diff --git a/Engine/Plugins/Animation/DeformerGraph/Source/OptimusCore/Private/DataInterfaces/OptimusDataInterfaceSkinnedMeshWrite.cpp b/Engine/Plugins/Animation/DeformerGraph/Source/OptimusCore/Private/DataInterfaces/OptimusDataInterfaceSkinnedMeshWrite.cpp
--- a/Engine/Plugins/Animation/DeformerGraph/Source/OptimusCore/Private/DataInterfaces/OptimusDataInterfaceSkinnedMeshWrite.cpp
+++ b/Engine/Plugins/Animation/DeformerGraph/Source/OptimusCore/Private/DataInterfaces/OptimusDataInterfaceSkinnedMeshWrite.cpp
@@ -170,7 +170,7 @@ FComputeDataProviderRenderProxy* UOptimusSkinnedMeshWriteDataProvider::GetRender
 {
 	FOptimusSkinnedMeshWriteDataProviderProxy* Proxy = new FOptimusSkinnedMeshWriteDataProviderProxy();
 	
-	Proxy->SkeletalMeshObject = SkinnedMesh? SkinnedMesh->MeshObject : nullptr;
+	Proxy->SkeletalMeshObject = SkinnedMesh.IsValid() ? SkinnedMesh->MeshObject : nullptr;
 	Proxy->OutputMask = OutputMask;
 	Proxy->LastLodIndexPtr = &LastLodIndexCachedByRenderProxy;
 
diff --git a/Engine/Plugins/Animation/DeformerGraph/Source/OptimusCore/Private/DataInterfaces/OptimusDataInterfaceSkinnedMeshWrite.h b/Engine/Plugins/Animation/DeformerGraph/Source/OptimusCore/Private/DataInterfaces/OptimusDataInterfaceSkinnedMeshWrite.h
--- a/Engine/Plugins/Animation/DeformerGraph/Source/OptimusCore/Private/DataInterfaces/OptimusDataInterfaceSkinnedMeshWrite.h
+++ b/Engine/Plugins/Animation/DeformerGraph/Source/OptimusCore/Private/DataInterfaces/OptimusDataInterfaceSkinnedMeshWrite.h
@@ -73,8 +73,7 @@ public:
 	~UOptimusSkinnedMeshWriteDataProvider() = default;
 	UOptimusSkinnedMeshWriteDataProvider(FVTableHelper& Helper);
 	
-	UPROPERTY(EditAnywhere, BlueprintReadWrite, Category = Binding)
-	TObjectPtr<USkinnedMeshComponent> SkinnedMesh = nullptr;
+	TWeakObjectPtr<USkinnedMeshComponent> SkinnedMesh = nullptr;
 
 	uint64 OutputMask = 0;
 
