[Deformer Graph] added a missing nullcheck in half edge data provider

diff --git a/Engine/Plugins/Animation/DeformerGraph/Source/OptimusCore/Private/DataInterfaces/OptimusDataInterfaceHalfEdge.cpp b/Engine/Plugins/Animation/DeformerGraph/Source/OptimusCore/Private/DataInterfaces/OptimusDataInterfaceHalfEdge.cpp
--- a/Engine/Plugins/Animation/DeformerGraph/Source/OptimusCore/Private/DataInterfaces/OptimusDataInterfaceHalfEdge.cpp
+++ b/Engine/Plugins/Animation/DeformerGraph/Source/OptimusCore/Private/DataInterfaces/OptimusDataInterfaceHalfEdge.cpp
@@ -215,6 +215,12 @@ void UOptimusHalfEdgeDataProvider::SetOptimusDeformerAssetPath(const FTopLevelAs
 void UOptimusHalfEdgeDataProvider::ValidateSkinnedAsset()
 {
 	check(IsInGameThread());
+
+	if (!SkinnedMesh.IsValid())
+	{
+		return;
+	}
+	
 	if (USkinnedAsset* Asset = SkinnedMesh->GetSkinnedAsset())
 	{
 		if (!ValidatedAssets.Contains(Asset))
