[Deformer Graph] Make sure skin weight profiles are requested earlier so that we can avoid TPose on the frame that the deformer becomes active / LOD changes

diff --git a/Engine/Plugins/Animation/DeformerGraph/Source/OptimusCore/Private/DataInterfaces/OptimusDataInterfaceAdvancedSkeleton.cpp b/Engine/Plugins/Animation/DeformerGraph/Source/OptimusCore/Private/DataInterfaces/OptimusDataInterfaceAdvancedSkeleton.cpp
--- a/Engine/Plugins/Animation/DeformerGraph/Source/OptimusCore/Private/DataInterfaces/OptimusDataInterfaceAdvancedSkeleton.cpp
+++ b/Engine/Plugins/Animation/DeformerGraph/Source/OptimusCore/Private/DataInterfaces/OptimusDataInterfaceAdvancedSkeleton.cpp
@@ -596,7 +596,9 @@ void UOptimusAdvancedSkeletonDataProvider::Init(const UOptimusAdvancedSkeletonDa
 		ParameterBufferSize = MetadataScope.ShaderParameterMetadata->GetSize();
 	}
 
-
+	// Request profiles early here for all the LODs that can run this deformer so that we can avoid requesting profiles
+	// during gameplay or LOD transitions as it can cause T-Posing for a few frames due to GPU readback delay
+	Optimus::RequestSkinWeightProfileForDeformer(InSkeletalMesh, SkinWeightProfile);
 }
 
 void UOptimusAdvancedSkeletonDataProvider::ComputeBoneTransformsForLayeredSkinning(
@@ -749,35 +751,12 @@ FComputeDataProviderRenderProxy* UOptimusAdvancedSkeletonDataProvider::GetRender
 		
 		FSkeletalMeshObject* SkeletalMeshObject = SkeletalMesh->MeshObject;
 		const int32 LodIndex = SkeletalMeshObject->GetLOD();
-		FSkeletalMeshRenderData const& SkeletalMeshRenderData = SkeletalMeshObject->GetSkeletalMeshRenderData();
-		FSkeletalMeshLODRenderData const& LodRenderData = SkeletalMeshRenderData.LODRenderData[LodIndex];
-
-		bool bSkinWeightBufferReady = false;
-
-		if (SkinWeightProfile.IsNone())
-		{
-			bSkinWeightBufferReady = true;
-		}
-		else if (LodRenderData.SkinWeightProfilesData.ContainsProfile(SkinWeightProfile))
-		{
-			// Retrieve this profile's skin weight buffer
-			const FSkinWeightProfileStack ProfileStack{SkinWeightProfile};
-			FSkinWeightVertexBuffer* Buffer = LodRenderData.SkinWeightProfilesData.GetOverrideBuffer(ProfileStack);
-			bSkinWeightBufferReady = Buffer ? true : false;
-		}
+		const FSkeletalMeshRenderData& SkeletalMeshRenderData = SkeletalMeshObject->GetSkeletalMeshRenderData();
+		const FSkeletalMeshLODRenderData& LodRenderData = SkeletalMeshRenderData.LODRenderData[LodIndex];
 
-		if (!bSkinWeightBufferReady)
+		if (!Optimus::IsSkinWeightProfileAvailable(LodRenderData, SkinWeightProfile))
 		{
-			if (LodRenderData.SkinWeightProfilesData.ContainsProfile(SkinWeightProfile))
-			{
-				const FSkinWeightProfileStack ProfileStack{SkinWeightProfile};
-				// Put in a skin weight profile request
-				if (FSkinWeightProfileManager* Manager = FSkinWeightProfileManager::Get(SkeletalMesh->GetWorld()))
-				{
-					FRequestFinished DummyCallback = [](TWeakObjectPtr<USkeletalMesh> WeakMesh, FSkinWeightProfileStack ProfileStack){};
-					Manager->RequestSkinWeightProfileStack(ProfileStack, SkeletalMesh->GetSkinnedAsset(), this, DummyCallback);
-				}
-			}
+			Optimus::RequestSkinWeightProfileForDeformer(SkeletalMesh.Get(), SkinWeightProfile, LodIndex);
 		}
 		else
 		{
diff --git a/Engine/Plugins/Animation/DeformerGraph/Source/OptimusCore/Private/DataInterfaces/OptimusDataInterfaceSkinWeightsAsVertexMask.cpp b/Engine/Plugins/Animation/DeformerGraph/Source/OptimusCore/Private/DataInterfaces/OptimusDataInterfaceSkinWeightsAsVertexMask.cpp
--- a/Engine/Plugins/Animation/DeformerGraph/Source/OptimusCore/Private/DataInterfaces/OptimusDataInterfaceSkinWeightsAsVertexMask.cpp
+++ b/Engine/Plugins/Animation/DeformerGraph/Source/OptimusCore/Private/DataInterfaces/OptimusDataInterfaceSkinWeightsAsVertexMask.cpp
@@ -5,6 +5,7 @@
 #include "OptimusDataDomain.h"
 #include "OptimusDataTypeRegistry.h"
 #include "OptimusDeformerInstance.h"
+#include "OptimusHelpers.h"
 #include "ShaderCompilerCore.h"
 #include "ShaderParameterMetadataBuilder.h"
 #include "SkeletalMeshDeformerHelpers.h"
@@ -177,6 +178,10 @@ void UOptimusSkinWeightsAsVertexMaskDataProvider::Init(const UOptimusSkinWeights
 	ExpandTowardsLeaf = InDataInterface->ExpandTowardsLeaf;
 	bDebugDrawIncludedBones = InDataInterface->bDebugDrawIncludedBones;
 	DebugDrawColor = InDataInterface->DebugDrawColor;
+
+	// Request profiles early here for all the LODs that can run this deformer so that we can avoid requesting profiles
+	// during gameplay or LOD transitions as it can cause T-Posing for a few frames due to GPU readback delay
+	Optimus::RequestSkinWeightProfileForDeformer(InSkeletalMesh, SkinWeightProfile);
 }
 
 FComputeDataProviderRenderProxy* UOptimusSkinWeightsAsVertexMaskDataProvider::GetRenderProxy()
@@ -277,32 +282,9 @@ FComputeDataProviderRenderProxy* UOptimusSkinWeightsAsVertexMaskDataProvider::Ge
 		FSkeletalMeshRenderData const& SkeletalMeshRenderData = SkeletalMeshObject->GetSkeletalMeshRenderData();
 		FSkeletalMeshLODRenderData const& CurrentLodRenderData = SkeletalMeshRenderData.LODRenderData[CurrentLodIndex];
 
-		bool bSkinWeightBufferReady = false;
-
-		if (SkinWeightProfile.IsNone())
-		{
-			bSkinWeightBufferReady = true;
-		}
-		else if (CurrentLodRenderData.SkinWeightProfilesData.ContainsProfile(SkinWeightProfile))
+		if (!Optimus::IsSkinWeightProfileAvailable(CurrentLodRenderData, SkinWeightProfile))
 		{
-			// Retrieve this profile's skin weight buffer
-			const FSkinWeightProfileStack ProfileStack{SkinWeightProfile};
-			FSkinWeightVertexBuffer* Buffer = CurrentLodRenderData.SkinWeightProfilesData.GetOverrideBuffer(ProfileStack);
-			bSkinWeightBufferReady = Buffer ? true : false;
-		}
-
-		if (!bSkinWeightBufferReady)
-		{
-			if (CurrentLodRenderData.SkinWeightProfilesData.ContainsProfile(SkinWeightProfile))
-			{
-				const FSkinWeightProfileStack ProfileStack{SkinWeightProfile};
-				// Put in a skin weight profile request
-				if (FSkinWeightProfileManager* Manager = FSkinWeightProfileManager::Get(SkeletalMesh->GetWorld()))
-				{
-					FRequestFinished DummyCallback = [](TWeakObjectPtr<USkeletalMesh> WeakMesh, FSkinWeightProfileStack ProfileStack){};
-					Manager->RequestSkinWeightProfileStack(ProfileStack, SkeletalMesh->GetSkinnedAsset(), this, DummyCallback);
-				}		
-			}
+			Optimus::RequestSkinWeightProfileForDeformer(SkeletalMesh.Get(), SkinWeightProfile, CurrentLodIndex);
 		}
 		else
 		{
diff --git a/Engine/Plugins/Animation/DeformerGraph/Source/OptimusCore/Private/OptimusHelpers.cpp b/Engine/Plugins/Animation/DeformerGraph/Source/OptimusCore/Private/OptimusHelpers.cpp
--- a/Engine/Plugins/Animation/DeformerGraph/Source/OptimusCore/Private/OptimusHelpers.cpp
+++ b/Engine/Plugins/Animation/DeformerGraph/Source/OptimusCore/Private/OptimusHelpers.cpp
@@ -9,6 +9,10 @@
 #include "ComputeFramework/ShaderParamTypeDefinition.h"
 #include "ShaderParameterMetadataBuilder.h"
 #include "ShaderParameterMetadata.h"
+#include "SkeletalRenderPublic.h"
+#include "Animation/SkinWeightProfileManager.h"
+#include "Rendering/SkeletalMeshRenderData.h"
+#include "Components/SkeletalMeshComponent.h"
 #include "StructUtils/UserDefinedStruct.h"
 
 #define LOCTEXT_NAMESPACE "OptimusHelpers"
@@ -559,4 +563,75 @@ FString Optimus::ExtractSourceValueName(const FString& InUniqueValueName)
 	return InUniqueValueName;
 }
 
-#undef LOCTEXT_NAMESPACE
\ No newline at end of file
+bool Optimus::IsSkinWeightProfileAvailable(const FSkeletalMeshLODRenderData& InLodRenderData, FName InSkinWeightProfile)
+{
+	if (InSkinWeightProfile.IsNone())
+	{
+		return true;
+	}
+
+	const FSkinWeightProfileStack ProfileStack{InSkinWeightProfile};
+	FSkinWeightVertexBuffer* Buffer = InLodRenderData.SkinWeightProfilesData.GetOverrideBuffer(ProfileStack);
+	return Buffer != nullptr;
+}
+
+void Optimus::RequestSkinWeightProfileForDeformer(USkeletalMeshComponent* InSkeletalMeshComponent, FName InSkinWeightProfile, TOptional<int32> InLOD)
+{
+	if (!IsValid(InSkeletalMeshComponent))
+	{
+		return;
+	}
+	FSkeletalMeshObject* SkeletalMeshObject = InSkeletalMeshComponent->MeshObject;
+	if (!SkeletalMeshObject)
+	{
+		return;
+	}
+
+	FSkeletalMeshRenderData const& SkeletalMeshRenderData = SkeletalMeshObject->GetSkeletalMeshRenderData();
+	
+	for (int32 LODIndex = 0; LODIndex < SkeletalMeshRenderData.NumInlinedLODs; LODIndex++)
+	{
+		if (!InSkeletalMeshComponent->HasMeshDeformer(LODIndex))
+		{
+			continue;
+		}
+
+		FSkeletalMeshLODRenderData const& LodRenderData = SkeletalMeshRenderData.LODRenderData[LODIndex];
+		if (!LodRenderData.SkinWeightProfilesData.ContainsProfile(InSkinWeightProfile))
+		{
+			continue;
+		}
+		
+		bool bShouldRequest = false;
+		if (InLOD)
+		{
+			if (*InLOD == LODIndex)
+			{
+				bShouldRequest = true;
+			}
+		}
+		else
+		{
+			bShouldRequest = true;
+		}
+
+		if (!bShouldRequest)
+		{
+			continue;
+		}
+		
+		if (!IsSkinWeightProfileAvailable(LodRenderData, InSkinWeightProfile))
+		{
+			const FSkinWeightProfileStack ProfileStack{InSkinWeightProfile};
+			// Put in a skin weight profile request
+			if (FSkinWeightProfileManager* Manager = FSkinWeightProfileManager::Get(InSkeletalMeshComponent->GetWorld()))
+			{
+				FRequestFinished DummyCallback = [](TWeakObjectPtr<USkeletalMesh> WeakMesh, FSkinWeightProfileStack ProfileStack){};
+				Manager->RequestSkinWeightProfileStack(ProfileStack, InSkeletalMeshComponent->GetSkinnedAsset(), InSkeletalMeshComponent, DummyCallback, LODIndex);
+			}
+		}	
+	}
+}
+
+
+#undef LOCTEXT_NAMESPACE
diff --git a/Engine/Plugins/Animation/DeformerGraph/Source/OptimusCore/Public/OptimusHelpers.h b/Engine/Plugins/Animation/DeformerGraph/Source/OptimusCore/Public/OptimusHelpers.h
--- a/Engine/Plugins/Animation/DeformerGraph/Source/OptimusCore/Public/OptimusHelpers.h
+++ b/Engine/Plugins/Animation/DeformerGraph/Source/OptimusCore/Public/OptimusHelpers.h
@@ -20,6 +20,8 @@ class FShaderParametersMetadata;
 class FShaderParametersMetadataBuilder;
 struct FShaderValueTypeHandle;
 class UActorComponent;
+class USkeletalMeshComponent;
+class FSkeletalMeshLODRenderData;
 
 namespace Optimus
 {
@@ -163,4 +165,8 @@ namespace Optimus
 
 	FString MakeUniqueValueName(const FString& InValueName, int32 InUniqueIndex);
 	FString ExtractSourceValueName(const FString& InUniqueValueName);
+
+	bool IsSkinWeightProfileAvailable(const FSkeletalMeshLODRenderData& InLODRenderData, FName InSkinWeightProfile);
+	
+	void RequestSkinWeightProfileForDeformer(USkeletalMeshComponent* InSkeletalMeshComponent, FName InSkinWeightProfile, TOptional<int32> InLOD = {});
 }
diff --git a/Engine/Source/Runtime/Engine/Classes/Components/SkinnedMeshComponent.h b/Engine/Source/Runtime/Engine/Classes/Components/SkinnedMeshComponent.h
--- a/Engine/Source/Runtime/Engine/Classes/Components/SkinnedMeshComponent.h
+++ b/Engine/Source/Runtime/Engine/Classes/Components/SkinnedMeshComponent.h
@@ -1213,7 +1213,8 @@ public:
 
 	ENGINE_API bool IsSkinCacheAllowed(int32 LodIdx) const;
 
-	bool HasMeshDeformer() const { return GetActiveMeshDeformers().Deformers.Num() > 0; }
+	ENGINE_API bool HasMeshDeformer() const;
+	ENGINE_API bool HasMeshDeformer(int32 LodIdx) const;
 
 	UE_DEPRECATED(5.5, "This is no longer relevant. Dynamic data is always updated immediately.")
 	bool GetForceUpdateDynamicDataImmediately() const { return true; }
diff --git a/Engine/Source/Runtime/Engine/Private/Components/SkinnedMeshComponent.cpp b/Engine/Source/Runtime/Engine/Private/Components/SkinnedMeshComponent.cpp
--- a/Engine/Source/Runtime/Engine/Private/Components/SkinnedMeshComponent.cpp
+++ b/Engine/Source/Runtime/Engine/Private/Components/SkinnedMeshComponent.cpp
@@ -2395,6 +2395,26 @@ bool USkinnedMeshComponent::IsSkinCacheAllowed(int32 LodIdx) const
 	return bComponentEnabled;
 }
 
+bool USkinnedMeshComponent::HasMeshDeformer() const
+{
+	return GetActiveMeshDeformers().Deformers.Num() > 0;
+}
+
+bool USkinnedMeshComponent::HasMeshDeformer(int32 LodIdx) const
+{
+	if (!HasMeshDeformer())
+	{
+		return false;
+	}
+
+	const FMeshDeformerSet& Set = GetActiveMeshDeformers();
+	if (!Set.DeformerIndexForLOD.IsValidIndex(LodIdx))
+	{
+		return false;
+	}
+	return Set.DeformerIndexForLOD[LodIdx] != INDEX_NONE;
+}
+
 void USkinnedMeshComponent::GetBoneNames(TArray<FName>& BoneNames)
 {
 	if (GetSkinnedAsset() == nullptr)
