[Deformer Graph] Fixed copy pasting resource node with matrix3x4 type creates a node with no pin

diff --git a/Engine/Plugins/Animation/DeformerGraph/Source/OptimusCore/Private/Nodes/OptimusNode_ResourceAccessorBase.cpp b/Engine/Plugins/Animation/DeformerGraph/Source/OptimusCore/Private/Nodes/OptimusNode_ResourceAccessorBase.cpp
--- a/Engine/Plugins/Animation/DeformerGraph/Source/OptimusCore/Private/Nodes/OptimusNode_ResourceAccessorBase.cpp
+++ b/Engine/Plugins/Animation/DeformerGraph/Source/OptimusCore/Private/Nodes/OptimusNode_ResourceAccessorBase.cpp
@@ -111,7 +111,7 @@ void UOptimusNode_ResourceAccessorBase::ExportCustomProperties(FOutputDevice& Ou
 {
 	if (const UOptimusResourceDescription* Res = ResourceDesc.Get())
 	{
-		Out.Logf(TEXT("%sCustomProperties ResourceDefinition Name=\"%s\" Type=%s DataDomain=\"%s\"\n"),
+		Out.Logf(TEXT("%sCustomProperties ResourceDefinition Name=\"%s\" Type=\"%s\" DataDomain=\"%s\"\n"),
 			FCString::Spc(Indent), *Res->ResourceName.ToString(), *Res->DataType->TypeName.ToString(),
 			*Res->DataDomain.ToString());
 	}
