[Deformer Graph] Avoid replacing the loaded instance settings object directly by creating a new object using the same name, instead, use an unique name

diff --git a/Engine/Plugins/Animation/DeformerGraph/Source/OptimusCore/Private/OptimusDeformer.cpp b/Engine/Plugins/Animation/DeformerGraph/Source/OptimusCore/Private/OptimusDeformer.cpp
--- a/Engine/Plugins/Animation/DeformerGraph/Source/OptimusCore/Private/OptimusDeformer.cpp
+++ b/Engine/Plugins/Animation/DeformerGraph/Source/OptimusCore/Private/OptimusDeformer.cpp
@@ -4024,8 +4024,9 @@ void UOptimusDeformer::PostEditChangeProperty(FPropertyChangedEvent& PropertyCha
 
 UMeshDeformerInstanceSettings* UOptimusDeformer::CreateSettingsInstance(UMeshComponent* InMeshComponent)
 {
-	const FName SettingsName(GetName() + TEXT("_Settings"));
+	FName SettingsName(GetName() + TEXT("_Settings"));
 	const EObjectFlags CreateObjectFlags = InMeshComponent->HasAnyFlags(RF_ArchetypeObject) ? RF_Public : RF_NoFlags; // Make public when stored in a BP.
+	SettingsName = MakeUniqueObjectName(InMeshComponent, UOptimusDeformerInstanceSettings::StaticClass(), SettingsName);
 	UOptimusDeformerInstanceSettings *Settings = NewObject<UOptimusDeformerInstanceSettings>(InMeshComponent, SettingsName, CreateObjectFlags);
 	Settings->InitializeSettings(this, InMeshComponent);
 	return Settings;
