[Deformer Graph] remove reset to default button for unsupported property fields on variable / resource

diff --git a/Engine/Plugins/Animation/DeformerGraph/Source/OptimusCore/Public/OptimusResourceDescription.h b/Engine/Plugins/Animation/DeformerGraph/Source/OptimusCore/Public/OptimusResourceDescription.h
--- a/Engine/Plugins/Animation/DeformerGraph/Source/OptimusCore/Public/OptimusResourceDescription.h
+++ b/Engine/Plugins/Animation/DeformerGraph/Source/OptimusCore/Public/OptimusResourceDescription.h
@@ -35,11 +35,11 @@ public:
 	/** Returns the index of the variable within the container */ 
 	UE_API int32 GetIndex() const;
 
-	UPROPERTY(EditAnywhere, BlueprintReadOnly, Category = ResourceDescription)
+	UPROPERTY(EditAnywhere, BlueprintReadOnly, Category = ResourceDescription, meta=(NoResetToDefault))
 	FName ResourceName;
 
 	/** The the data type of each element of the resource */
-	UPROPERTY(EditAnywhere, Category=ResourceDescription, meta=(UseInResource))
+	UPROPERTY(EditAnywhere, Category=ResourceDescription, meta=(UseInResource, NoResetToDefault))
 	FOptimusDataTypeRef DataType;
 
 	/** The component binding that this resource description is bound to */
diff --git a/Engine/Plugins/Animation/DeformerGraph/Source/OptimusCore/Public/OptimusVariableDescription.h b/Engine/Plugins/Animation/DeformerGraph/Source/OptimusCore/Public/OptimusVariableDescription.h
--- a/Engine/Plugins/Animation/DeformerGraph/Source/OptimusCore/Public/OptimusVariableDescription.h
+++ b/Engine/Plugins/Animation/DeformerGraph/Source/OptimusCore/Public/OptimusVariableDescription.h
@@ -62,11 +62,11 @@ public:
 	FGuid Guid;
 
 	/** Name of the variable */
-	UPROPERTY(EditAnywhere, Category = VariableDefinition)
+	UPROPERTY(EditAnywhere, Category = VariableDefinition,  meta=(NoResetToDefault))
 	FName VariableName;
 
 	/** The data type of the variable */
-	UPROPERTY(EditAnywhere, Category = VariableDefinition, meta=(UseInVariable, UseInProperty))
+	UPROPERTY(EditAnywhere, Category = VariableDefinition, meta=(UseInVariable, UseInProperty, NoResetToDefault))
 	FOptimusDataTypeRef DataType;
 
 	/** The default value for the variable. */
