Optimus: Enable documentation policy and add missing comments

#rnx

diff --git a/Engine/Plugins/Animation/DeformerGraph/Source/OptimusCore/Private/ControlRig/RigUnit_Optimus.h b/Engine/Plugins/Animation/DeformerGraph/Source/OptimusCore/Private/ControlRig/RigUnit_Optimus.h
--- a/Engine/Plugins/Animation/DeformerGraph/Source/OptimusCore/Private/ControlRig/RigUnit_Optimus.h
+++ b/Engine/Plugins/Animation/DeformerGraph/Source/OptimusCore/Private/ControlRig/RigUnit_Optimus.h
@@ -9,7 +9,10 @@
 
 #define UE_API OPTIMUSCORE_API
 
-USTRUCT()
+/*
+ * The trait representing the optimus deformer kernel for Control Rig
+ */
+USTRUCT(meta=(DocumentationPolicy = "Strict"))
 struct FRigVMTrait_OptimusDeformer: public FRigVMTrait
 {
 	GENERATED_BODY()
@@ -24,14 +27,17 @@ struct FRigVMTrait_OptimusDeformer: public FRigVMTrait
 	TSoftObjectPtr<UOptimusDeformer> DeformerGraph;
 };
 
-
-USTRUCT()
+/*
+ * The trait representing the base settings for an optimus deformer in Control Rig
+ */
+USTRUCT(meta=(DocumentationPolicy = "Strict"))
 struct FRigVMTrait_OptimusDeformerSettings: public FRigVMTrait
 {
 	GENERATED_BODY()
 
 	FRigVMTrait_OptimusDeformerSettings() = default;
 
+	// The phase at which the deformer will execute
 	UPROPERTY(EditAnywhere, Category = "Trait", meta=(Input))
 	EOptimusDeformerExecutionPhase ExecutionPhase = EOptimusDeformerExecutionPhase::AfterDefaultDeformer;
 	
@@ -72,9 +78,10 @@ struct FRigUnit_AddOptimusDeformer: public FRigUnitMutable
 	FGuid DeformerInstanceGuid;
 };
 
-
-
-USTRUCT()
+/*
+ * The base class for all optimus deformer variable traits for Control Rig
+ */
+USTRUCT(meta=(DocumentationPolicy = "Strict"))
 struct FRigVMTrait_OptimusVariableBase: public FRigVMTrait
 {
 	GENERATED_BODY()
@@ -83,6 +90,9 @@ struct FRigVMTrait_OptimusVariableBase: public FRigVMTrait
 	static const TCHAR* GetValuePinName();
 };
 
+/*
+ * A trait to set an Int variable on an optimus deformer in Control Rig
+ */
 USTRUCT(BlueprintType)
 struct FRigVMTrait_SetDeformerIntVariable: public FRigVMTrait_OptimusVariableBase
 {
@@ -94,6 +104,9 @@ struct FRigVMTrait_SetDeformerIntVariable: public FRigVMTrait_OptimusVariableBas
 	int32 Value = 0;	
 };
 
+/*
+ * A trait to set an Int array variable on an optimus deformer in Control Rig
+ */
 USTRUCT(BlueprintType)
 struct FRigVMTrait_SetDeformerIntArrayVariable: public FRigVMTrait_OptimusVariableBase
 {
@@ -105,6 +118,9 @@ struct FRigVMTrait_SetDeformerIntArrayVariable: public FRigVMTrait_OptimusVariab
 	TArray<int32> Value;	
 };
 
+/*
+ * A trait to set an Int2 variable on an optimus deformer in Control Rig
+ */
 USTRUCT(BlueprintType)
 struct FRigVMTrait_SetDeformerInt2Variable: public FRigVMTrait_OptimusVariableBase
 {
@@ -116,6 +132,9 @@ struct FRigVMTrait_SetDeformerInt2Variable: public FRigVMTrait_OptimusVariableBa
 	FIntPoint Value = FIntPoint::ZeroValue;	
 };
 
+/*
+ * A trait to set an Int2 array variable on an optimus deformer in Control Rig
+ */
 USTRUCT(BlueprintType)
 struct FRigVMTrait_SetDeformerInt2ArrayVariable: public FRigVMTrait_OptimusVariableBase
 {
@@ -127,6 +146,9 @@ struct FRigVMTrait_SetDeformerInt2ArrayVariable: public FRigVMTrait_OptimusVaria
 	TArray<FIntPoint> Value;	
 };
 
+/*
+ * A trait to set an Int3 variable on an optimus deformer in Control Rig
+ */
 USTRUCT(BlueprintType)
 struct FRigVMTrait_SetDeformerInt3Variable: public FRigVMTrait_OptimusVariableBase
 {
@@ -138,6 +160,9 @@ struct FRigVMTrait_SetDeformerInt3Variable: public FRigVMTrait_OptimusVariableBa
 	FIntVector Value = FIntVector::ZeroValue;	
 };
 
+/*
+ * A trait to set an Int3 array variable on an optimus deformer in Control Rig
+ */
 USTRUCT(BlueprintType)
 struct FRigVMTrait_SetDeformerInt3ArrayVariable: public FRigVMTrait_OptimusVariableBase
 {
@@ -149,6 +174,9 @@ struct FRigVMTrait_SetDeformerInt3ArrayVariable: public FRigVMTrait_OptimusVaria
 	TArray<FIntVector> Value;	
 };
 
+/*
+ * A trait to set an Int4 variable on an optimus deformer in Control Rig
+ */
 USTRUCT(BlueprintType)
 struct FRigVMTrait_SetDeformerInt4Variable: public FRigVMTrait_OptimusVariableBase
 {
@@ -160,6 +188,9 @@ struct FRigVMTrait_SetDeformerInt4Variable: public FRigVMTrait_OptimusVariableBa
 	FIntVector4 Value = FIntVector4::ZeroValue;	
 };
 
+/*
+ * A trait to set an Int4 array variable on an optimus deformer in Control Rig
+ */
 USTRUCT(BlueprintType)
 struct FRigVMTrait_SetDeformerInt4ArrayVariable: public FRigVMTrait_OptimusVariableBase
 {
@@ -171,9 +202,9 @@ struct FRigVMTrait_SetDeformerInt4ArrayVariable: public FRigVMTrait_OptimusVaria
 	TArray<FIntVector4> Value;	
 };
 
-
-
-
+/*
+ * A trait to set a double variable on an optimus deformer in Control Rig
+ */
 USTRUCT(BlueprintType)
 struct FRigVMTrait_SetDeformerFloatVariable: public FRigVMTrait_OptimusVariableBase
 {
@@ -185,6 +216,9 @@ struct FRigVMTrait_SetDeformerFloatVariable: public FRigVMTrait_OptimusVariableB
 	double Value = 0.0;	
 };
 
+/*
+ * A trait to set a double array variable on an optimus deformer in Control Rig
+ */
 USTRUCT(BlueprintType)
 struct FRigVMTrait_SetDeformerFloatArrayVariable: public FRigVMTrait_OptimusVariableBase
 {
@@ -196,6 +230,9 @@ struct FRigVMTrait_SetDeformerFloatArrayVariable: public FRigVMTrait_OptimusVari
 	TArray<double> Value;	
 };
 
+/*
+ * A trait to set a Vector2 variable on an optimus deformer in Control Rig
+ */
 USTRUCT(BlueprintType)
 struct FRigVMTrait_SetDeformerVector2Variable: public FRigVMTrait_OptimusVariableBase
 {
@@ -207,7 +244,9 @@ struct FRigVMTrait_SetDeformerVector2Variable: public FRigVMTrait_OptimusVariabl
 	FVector2D Value = FVector2D::ZeroVector;	
 };
 
-
+/*
+ * A trait to set a Vector2 array variable on an optimus deformer in Control Rig
+ */
 USTRUCT(BlueprintType)
 struct FRigVMTrait_SetDeformerVector2ArrayVariable: public FRigVMTrait_OptimusVariableBase
 {
@@ -219,6 +258,9 @@ struct FRigVMTrait_SetDeformerVector2ArrayVariable: public FRigVMTrait_OptimusVa
 	TArray<FVector2D> Value;	
 };
 
+/*
+ * A trait to set a Vector variable on an optimus deformer in Control Rig
+ */
 USTRUCT(BlueprintType)
 struct FRigVMTrait_SetDeformerVectorVariable: public FRigVMTrait_OptimusVariableBase
 {
@@ -230,6 +272,9 @@ struct FRigVMTrait_SetDeformerVectorVariable: public FRigVMTrait_OptimusVariable
 	FVector Value = FVector::ZeroVector; 
 };
 
+/*
+ * A trait to set a Vector array variable on an optimus deformer in Control Rig
+ */
 USTRUCT(BlueprintType)
 struct FRigVMTrait_SetDeformerVectorArrayVariable: public FRigVMTrait_OptimusVariableBase
 {
@@ -241,7 +286,9 @@ struct FRigVMTrait_SetDeformerVectorArrayVariable: public FRigVMTrait_OptimusVar
 	TArray<FVector> Value;	
 };
 
-
+/*
+ * A trait to set a Vector4 variable on an optimus deformer in Control Rig
+ */
 USTRUCT(BlueprintType)
 struct FRigVMTrait_SetDeformerVector4Variable: public FRigVMTrait_OptimusVariableBase
 {
@@ -253,6 +300,9 @@ struct FRigVMTrait_SetDeformerVector4Variable: public FRigVMTrait_OptimusVariabl
 	FVector4 Value = FVector4::Zero();	
 };
 
+/*
+ * A trait to set a Vector4 array variable on an optimus deformer in Control Rig
+ */
 USTRUCT(BlueprintType)
 struct FRigVMTrait_SetDeformerVector4ArrayVariable: public FRigVMTrait_OptimusVariableBase
 {
@@ -264,6 +314,9 @@ struct FRigVMTrait_SetDeformerVector4ArrayVariable: public FRigVMTrait_OptimusVa
 	TArray<FVector4> Value;	
 };
 
+/*
+ * A trait to set a LinearColor variable on an optimus deformer in Control Rig
+ */
 USTRUCT(BlueprintType)
 struct FRigVMTrait_SetDeformerLinearColorVariable: public FRigVMTrait_OptimusVariableBase
 {
@@ -275,6 +328,9 @@ struct FRigVMTrait_SetDeformerLinearColorVariable: public FRigVMTrait_OptimusVar
 	FLinearColor Value = FLinearColor::Black;	
 };
 
+/*
+ * A trait to set a LinearColor array variable on an optimus deformer in Control Rig
+ */
 USTRUCT(BlueprintType)
 struct FRigVMTrait_SetDeformerLinearColorArrayVariable: public FRigVMTrait_OptimusVariableBase
 {
@@ -286,7 +342,9 @@ struct FRigVMTrait_SetDeformerLinearColorArrayVariable: public FRigVMTrait_Optim
 	TArray<FLinearColor> Value;	
 };
 
-
+/*
+ * A trait to set a Quat variable on an optimus deformer in Control Rig
+ */
 USTRUCT(BlueprintType)
 struct FRigVMTrait_SetDeformerQuatVariable: public FRigVMTrait_OptimusVariableBase
 {
@@ -298,7 +356,9 @@ struct FRigVMTrait_SetDeformerQuatVariable: public FRigVMTrait_OptimusVariableBa
 	FQuat Value = FQuat::Identity;	
 };
 
-
+/*
+ * A trait to set a Quat array variable on an optimus deformer in Control Rig
+ */
 USTRUCT(BlueprintType)
 struct FRigVMTrait_SetDeformerQuatArrayVariable: public FRigVMTrait_OptimusVariableBase
 {
@@ -310,6 +370,9 @@ struct FRigVMTrait_SetDeformerQuatArrayVariable: public FRigVMTrait_OptimusVaria
 	TArray<FQuat> Value;	
 };
 
+/*
+ * A trait to set a Rotator variable on an optimus deformer in Control Rig
+ */
 USTRUCT(BlueprintType)
 struct FRigVMTrait_SetDeformerRotatorVariable: public FRigVMTrait_OptimusVariableBase
 {
@@ -321,6 +384,9 @@ struct FRigVMTrait_SetDeformerRotatorVariable: public FRigVMTrait_OptimusVariabl
 	FRotator Value = FRotator::ZeroRotator;	
 };
 
+/*
+ * A trait to set a Rotator array variable on an optimus deformer in Control Rig
+ */
 USTRUCT(BlueprintType)
 struct FRigVMTrait_SetDeformerRotatorArrayVariable: public FRigVMTrait_OptimusVariableBase
 {
@@ -332,6 +398,9 @@ struct FRigVMTrait_SetDeformerRotatorArrayVariable: public FRigVMTrait_OptimusVa
 	TArray<FRotator> Value;	
 };
 
+/*
+ * A trait to set a Transform variable on an optimus deformer in Control Rig
+ */
 USTRUCT(BlueprintType)
 struct FRigVMTrait_SetDeformerTransformVariable: public FRigVMTrait_OptimusVariableBase
 {
@@ -343,6 +412,9 @@ struct FRigVMTrait_SetDeformerTransformVariable: public FRigVMTrait_OptimusVaria
 	FTransform Value = FTransform::Identity;	
 };
 
+/*
+ * A trait to set a Transform array variable on an optimus deformer in Control Rig
+ */
 USTRUCT(BlueprintType)
 struct FRigVMTrait_SetDeformerTransformArrayVariable: public FRigVMTrait_OptimusVariableBase
 {
@@ -354,6 +426,9 @@ struct FRigVMTrait_SetDeformerTransformArrayVariable: public FRigVMTrait_Optimus
 	TArray<FTransform> Value;	
 };
 
+/*
+ * A trait to set a Name variable on an optimus deformer in Control Rig
+ */
 USTRUCT(BlueprintType)
 struct FRigVMTrait_SetDeformerNameVariable: public FRigVMTrait_OptimusVariableBase
 {
@@ -365,6 +440,9 @@ struct FRigVMTrait_SetDeformerNameVariable: public FRigVMTrait_OptimusVariableBa
 	FName Value = NAME_None;	
 };
 
+/*
+ * A trait to set a Name array variable on an optimus deformer in Control Rig
+ */
 USTRUCT(BlueprintType)
 struct FRigVMTrait_SetDeformerNameArrayVariable: public FRigVMTrait_OptimusVariableBase
 {
@@ -376,6 +454,9 @@ struct FRigVMTrait_SetDeformerNameArrayVariable: public FRigVMTrait_OptimusVaria
 	TArray<FName> Value;	
 };
 
+/*
+ * A trait to set a Bool variable on an optimus deformer in Control Rig
+ */
 USTRUCT(BlueprintType)
 struct FRigVMTrait_SetDeformerBoolVariable: public FRigVMTrait_OptimusVariableBase
 {
@@ -387,6 +468,9 @@ struct FRigVMTrait_SetDeformerBoolVariable: public FRigVMTrait_OptimusVariableBa
 	bool Value = false;	
 };
 
+/*
+ * A trait to set a Bool array variable on an optimus deformer in Control Rig
+ */
 USTRUCT(BlueprintType)
 struct FRigVMTrait_SetDeformerBoolArrayVariable: public FRigVMTrait_OptimusVariableBase
 {
