[DeformerGraph] Avoid scheduling deformer work during ProcessTasksUntilIdle as it may be part of a save callstack when updating progress bar, during which certain operations are not allowed and may lead to crash, such as staticfindobjectfast(), changing to use game thread tick group instead

diff --git a/Engine/Plugins/Animation/DeformerGraph/Source/OptimusCore/Private/ControlRig/RigUnit_Optimus.cpp b/Engine/Plugins/Animation/DeformerGraph/Source/OptimusCore/Private/ControlRig/RigUnit_Optimus.cpp
--- a/Engine/Plugins/Animation/DeformerGraph/Source/OptimusCore/Private/ControlRig/RigUnit_Optimus.cpp
+++ b/Engine/Plugins/Animation/DeformerGraph/Source/OptimusCore/Private/ControlRig/RigUnit_Optimus.cpp
@@ -362,15 +362,14 @@ FRigUnit_AddOptimusDeformer_Execute()
 
 					if (UOptimusDeformer* DeformerGraph = DeformerTrait->DeformerGraph.Get())
 					{
-						FFunctionGraphTask::CreateAndDispatchWhenReady([
-							AddDeformerToComponentsIfNeeded,
-							WeakDeformer = TWeakObjectPtr<UOptimusDeformer>(DeformerGraph)]()
-							{
-								if (WeakDeformer.IsValid())
+						ExecuteOnGameThread(UE_SOURCE_LOCATION,
+							[AddDeformerToComponentsIfNeeded, WeakDeformer = TWeakObjectPtr<UOptimusDeformer>(DeformerGraph)]()
 								{
-									AddDeformerToComponentsIfNeeded(WeakDeformer.Get());
-								}
-							}, TStatId(), NULL, ENamedThreads::GameThread);	
+									if (WeakDeformer.IsValid())
+									{
+										AddDeformerToComponentsIfNeeded(WeakDeformer.Get());
+									}
+								});
 					}
 					else
 					{
diff --git a/Engine/Plugins/Animation/DeformerGraph/Source/OptimusCore/Private/DataInterfaces/OptimusDataInterfaceHalfEdge.cpp b/Engine/Plugins/Animation/DeformerGraph/Source/OptimusCore/Private/DataInterfaces/OptimusDataInterfaceHalfEdge.cpp
--- a/Engine/Plugins/Animation/DeformerGraph/Source/OptimusCore/Private/DataInterfaces/OptimusDataInterfaceHalfEdge.cpp
+++ b/Engine/Plugins/Animation/DeformerGraph/Source/OptimusCore/Private/DataInterfaces/OptimusDataInterfaceHalfEdge.cpp
@@ -191,13 +191,13 @@ FComputeDataProviderRenderProxy* UOptimusHalfEdgeDataProvider::GetRenderProxy()
 #if WITH_EDITOR
 	if (Optimus::IsAssetValidationEnabled())
 	{
-		FFunctionGraphTask::CreateAndDispatchWhenReady([WeakThis = TWeakObjectPtr(this)]()
-		{
-			if (UOptimusHalfEdgeDataProvider* DataProvider = Cast<UOptimusHalfEdgeDataProvider>(WeakThis))
+		ExecuteOnGameThread(UE_SOURCE_LOCATION, [WeakThis = TWeakObjectPtr(this)]()
 			{
-				DataProvider->ValidateSkinnedAsset();
-			}
-		}, TStatId(), NULL, ENamedThreads::GameThread);
+				if (UOptimusHalfEdgeDataProvider* DataProvider = Cast<UOptimusHalfEdgeDataProvider>(WeakThis))
+				{
+					DataProvider->ValidateSkinnedAsset();
+				}
+			});
 	}
 #endif
 		
diff --git a/Engine/Plugins/Animation/DeformerGraph/Source/OptimusCore/Private/DataInterfaces/OptimusDataInterfaceSkinWeightsAsVertexMask.cpp b/Engine/Plugins/Animation/DeformerGraph/Source/OptimusCore/Private/DataInterfaces/OptimusDataInterfaceSkinWeightsAsVertexMask.cpp
--- a/Engine/Plugins/Animation/DeformerGraph/Source/OptimusCore/Private/DataInterfaces/OptimusDataInterfaceSkinWeightsAsVertexMask.cpp
+++ b/Engine/Plugins/Animation/DeformerGraph/Source/OptimusCore/Private/DataInterfaces/OptimusDataInterfaceSkinWeightsAsVertexMask.cpp
@@ -387,29 +387,29 @@ FComputeDataProviderRenderProxy* UOptimusSkinWeightsAsVertexMaskDataProvider::Ge
 			if (bDebugDrawIncludedBones)
 			{
 #if WITH_EDITOR
-				FFunctionGraphTask::CreateAndDispatchWhenReady([ WeakThis = TWeakObjectPtr(this)]()
+				ExecuteOnGameThread(UE_SOURCE_LOCATION, [ WeakThis = TWeakObjectPtr(this)]()
 					{
 						if (WeakThis.IsValid())
 						{
 							if (USkeletalMeshComponent* Mesh = WeakThis->SkeletalMesh.Get())
 							{
 								const FReferenceSkeleton& RefSkeleton = Mesh->GetSkinnedAsset()->GetRefSkeleton();
-					
+
 								const TArray<FTransform>& CurrentComponentSpaceBoneMatrix = Mesh->GetComponentSpaceTransforms();
 
-								for (int32 Index = 0 ; Index < CurrentComponentSpaceBoneMatrix.Num(); Index++)
+								for (int32 Index = 0; Index < CurrentComponentSpaceBoneMatrix.Num(); Index++)
 								{
 									int32 ParentIndex = RefSkeleton.GetParentIndex(Index);
-									if (ParentIndex!=INDEX_NONE && WeakThis->CachedSelectedBones.Contains(ParentIndex))
+									if (ParentIndex != INDEX_NONE && WeakThis->CachedSelectedBones.Contains(ParentIndex))
 									{
-										const FTransform& Child =  CurrentComponentSpaceBoneMatrix[Index] * Mesh->GetComponentTransform();
+										const FTransform& Child = CurrentComponentSpaceBoneMatrix[Index] * Mesh->GetComponentTransform();
 										const FTransform& Parent = CurrentComponentSpaceBoneMatrix[ParentIndex] * Mesh->GetComponentTransform();
-										DrawDebugLine(Mesh->GetWorld(), Child.GetLocation(), Parent.GetLocation(), WeakThis->DebugDrawColor, false, -1, SDPG_Foreground);	
+										DrawDebugLine(Mesh->GetWorld(), Child.GetLocation(), Parent.GetLocation(), WeakThis->DebugDrawColor, false, -1, SDPG_Foreground);
 									}
-								}	
+								}
 							}
 						}
-					}, TStatId(), NULL, ENamedThreads::GameThread);	
+					});	
 #endif
 			}
 			
diff --git a/Engine/Plugins/Animation/DeformerGraph/Source/OptimusCore/Private/OptimusDeformer.cpp b/Engine/Plugins/Animation/DeformerGraph/Source/OptimusCore/Private/OptimusDeformer.cpp
--- a/Engine/Plugins/Animation/DeformerGraph/Source/OptimusCore/Private/OptimusDeformer.cpp
+++ b/Engine/Plugins/Animation/DeformerGraph/Source/OptimusCore/Private/OptimusDeformer.cpp
@@ -3602,15 +3602,14 @@ void UOptimusDeformer::PostLoad()
 				// so in some case it is not possible make sure they are fully loaded during PostLoad.
 				// Delaying the compilation here such that it is performed after all dependencies are fully loaded (hopefully)
 				// It is best to resave the asset to avoid going down this path
-				FFunctionGraphTask::CreateAndDispatchWhenReady(
+				ExecuteOnGameThread(UE_SOURCE_LOCATION,
 					[WeakThis = TWeakObjectPtr(this)]()
-					{
-						if (WeakThis.IsValid())
 						{
-							WeakThis->Compile();
-						}
-					}
-				, TStatId(), NULL, ENamedThreads::GameThread);	
+							if (WeakThis.IsValid())
+							{
+								WeakThis->Compile();
+							}
+						});
 
 				UE_LOG(LogOptimusCore, Warning, TEXT("Deformer Graph should be recompiled and resaved to avoid issues in the packaged build, Asset: %s"),  *GetPackage()->GetName());
 			}
