[Deformer Graph] Added stronger name validation for deformer graph items

diff --git a/Engine/Plugins/Animation/DeformerGraph/Source/OptimusEditor/Private/OptimusNameValidator.h b/Engine/Plugins/Animation/DeformerGraph/Source/OptimusEditor/Private/OptimusNameValidator.h
--- a/Engine/Plugins/Animation/DeformerGraph/Source/OptimusEditor/Private/OptimusNameValidator.h
+++ b/Engine/Plugins/Animation/DeformerGraph/Source/OptimusEditor/Private/OptimusNameValidator.h
@@ -16,8 +16,8 @@ public:
 	
 	/// Disallowed characters in a graph/variable/buffer name. Period is disallowed since it's
 	/// a path separator (like with BP), and slash is also disallowed since we use it for 
-	/// constructing name paths for undo.
-	static const TCHAR *InvalidCharacters() { return TEXT("./$"); }
+	/// constructing name paths for undo, space is also not allowed since it complicates unique name checking
+	static const TCHAR *InvalidCharacters() { return TEXT("./$ "); }
 
 	// INameValidatorInterface overrides
 	EValidatorResult IsValid(const FName& InName, bool bInOriginal = false) override;
