Deformer Graph: Scene Component Data Interface

#Synced-CL 49084977

diff --git a/Engine/Plugins/Animation/DeformerGraph/Shaders/Private/DataInterfaceSceneComponent.ush b/Engine/Plugins/Animation/DeformerGraph/Shaders/Private/DataInterfaceSceneComponent.ush
new file mode 100644
--- /dev/null
+++ b/Engine/Plugins/Animation/DeformerGraph/Shaders/Private/DataInterfaceSceneComponent.ush
@@ -0,0 +1,8 @@
+// Copyright Epic Games, Inc. All Rights Reserved.
+
+float4x4 {DataInterfaceName}_WorldTransform;
+
+float4x4 ReadWorldTransform_{DataInterfaceName}()
+{
+	return {DataInterfaceName}_WorldTransform;
+}
\ No newline at end of file
diff --git a/Engine/Plugins/Animation/DeformerGraph/Source/OptimusCore/Private/DataInterfaces/OptimusDataInterfaceSceneComponent.cpp b/Engine/Plugins/Animation/DeformerGraph/Source/OptimusCore/Private/DataInterfaces/OptimusDataInterfaceSceneComponent.cpp
new file mode 100644
--- /dev/null
+++ b/Engine/Plugins/Animation/DeformerGraph/Source/OptimusCore/Private/DataInterfaces/OptimusDataInterfaceSceneComponent.cpp
@@ -0,0 +1,116 @@
+// Copyright Epic Games, Inc. All Rights Reserved.
+
+#include "OptimusDataInterfaceSceneComponent.h"
+
+#include "Components/SceneComponent.h"
+#include "ComputeFramework/ShaderParamTypeDefinition.h"
+#include "OptimusHelpers.h"
+#include "SceneInterface.h"
+#include "ShaderCompilerCore.h"
+#include "ShaderParameterMetadataBuilder.h"
+
+#include UE_INLINE_GENERATED_CPP_BY_NAME(OptimusDataInterfaceSceneComponent)
+
+FString UOptimusSceneComponentDataInterface::GetDisplayName() const
+{
+	return TEXT("Scene Component Data");
+}
+
+TArray<FOptimusCDIPinDefinition> UOptimusSceneComponentDataInterface::GetPinDefinitions() const
+{
+	TArray<FOptimusCDIPinDefinition> Defs;
+	Defs.Add({ "WorldTransform", "ReadWorldTransform" });
+	return Defs;
+}
+
+
+TSubclassOf<UActorComponent> UOptimusSceneComponentDataInterface::GetRequiredComponentClass() const
+{
+	return USceneComponent::StaticClass();
+}
+
+
+void UOptimusSceneComponentDataInterface::GetSupportedInputs(TArray<FShaderFunctionDefinition>& OutFunctions) const
+{
+	OutFunctions.AddDefaulted_GetRef()
+		.SetName(TEXT("ReadWorldTransform"))
+		.AddReturnType(EShaderFundamentalType::Float, 4, 4);
+}
+
+BEGIN_SHADER_PARAMETER_STRUCT(FSceneComponentDataInterfaceParameters, )
+	SHADER_PARAMETER(FMatrix44f, WorldTransform)
+END_SHADER_PARAMETER_STRUCT()
+
+void UOptimusSceneComponentDataInterface::GetShaderParameters(TCHAR const* UID, FShaderParametersMetadataBuilder& InOutBuilder, FShaderParametersMetadataAllocations& InOutAllocations) const
+{
+	InOutBuilder.AddNestedStruct<FSceneComponentDataInterfaceParameters>(UID);
+}
+
+TCHAR const* UOptimusSceneComponentDataInterface::TemplateFilePath = TEXT("/Plugin/Optimus/Private/DataInterfaceSceneComponent.ush");
+
+TCHAR const* UOptimusSceneComponentDataInterface::GetShaderVirtualPath() const
+{
+	return TemplateFilePath;
+}
+
+void UOptimusSceneComponentDataInterface::GetShaderHash(FString& InOutKey) const
+{
+	GetShaderFileHash(TemplateFilePath, EShaderPlatform::SP_PCD3D_SM5).AppendString(InOutKey);
+}
+
+void UOptimusSceneComponentDataInterface::GetHLSL(FString& OutHLSL, FString const& InDataInterfaceName) const
+{
+ 	TMap<FString, FStringFormatArg> TemplateArgs =
+	{
+		{ TEXT("DataInterfaceName"), InDataInterfaceName },
+	};
+
+	FString TemplateFile;
+	LoadShaderSourceFile(TemplateFilePath, EShaderPlatform::SP_PCD3D_SM5, &TemplateFile, nullptr);
+	OutHLSL += FString::Format(*TemplateFile, TemplateArgs);
+}
+
+UComputeDataProvider* UOptimusSceneComponentDataInterface::CreateDataProvider() const
+{
+	return NewObject<UOptimusSceneComponentDataProvider>();
+}
+
+void UOptimusSceneComponentDataProvider::Initialize(const UComputeDataInterface* InDataInterface, UObject* InBinding, uint64 InInputMask, uint64 InOutputMask)
+{
+	SceneComponent = Cast<USceneComponent>(InBinding);
+}
+
+FComputeDataProviderRenderProxy* UOptimusSceneComponentDataProvider::GetRenderProxy()
+{
+	return new FOptimusSceneComponentDataProviderProxy(SceneComponent.Get());
+}
+
+
+FOptimusSceneComponentDataProviderProxy::FOptimusSceneComponentDataProviderProxy(USceneComponent* SceneComponent)
+{
+	bIsValid = SceneComponent != nullptr;
+	if (bIsValid)
+	{
+		WorldTransform = SceneComponent->GetComponentTransform();
+	}
+}
+
+bool FOptimusSceneComponentDataProviderProxy::IsValid(FValidationData const& InValidationData) const
+{
+	if (InValidationData.ParameterStructSize != sizeof(FParameters))
+	{
+		return false;
+	}
+
+	return bIsValid;
+}
+
+void FOptimusSceneComponentDataProviderProxy::GatherDispatchData(FDispatchData const& InDispatchData)
+{
+	const TStridedView<FParameters> ParameterArray = MakeStridedParameterView<FParameters>(InDispatchData);
+	for (int32 InvocationIndex = 0; InvocationIndex < ParameterArray.Num(); ++InvocationIndex)
+	{
+		FParameters& Parameters = ParameterArray[InvocationIndex];
+		Parameters.WorldTransform = Optimus::ConvertFTransformToFMatrix44f(WorldTransform);
+	}
+}
diff --git a/Engine/Plugins/Animation/DeformerGraph/Source/OptimusCore/Private/DataInterfaces/OptimusDataInterfaceSceneComponent.h b/Engine/Plugins/Animation/DeformerGraph/Source/OptimusCore/Private/DataInterfaces/OptimusDataInterfaceSceneComponent.h
new file mode 100644
--- /dev/null
+++ b/Engine/Plugins/Animation/DeformerGraph/Source/OptimusCore/Private/DataInterfaces/OptimusDataInterfaceSceneComponent.h
@@ -0,0 +1,70 @@
+// Copyright Epic Games, Inc. All Rights Reserved.
+
+#pragma once
+
+#include "OptimusComputeDataInterface.h"
+#include "ComputeFramework/ComputeDataProvider.h"
+#include "OptimusDataInterfaceSceneComponent.generated.h"
+
+class FSceneComponentDataInterfaceParameters;
+class USceneComponent;
+
+/** Compute Framework Data Interface for reading general scene component data. */
+UCLASS(Category = ComputeFramework)
+class UOptimusSceneComponentDataInterface : public UOptimusComputeDataInterface
+{
+	GENERATED_BODY()
+
+public:
+	//~ Begin UOptimusComputeDataInterface Interface
+	virtual FString GetDisplayName() const override;
+	virtual TArray<FOptimusCDIPinDefinition> GetPinDefinitions() const override;
+	virtual TSubclassOf<UActorComponent> GetRequiredComponentClass() const override;
+	//~ End UOptimusComputeDataInterface Interface
+
+	//~ Begin UComputeDataInterface Interface
+	virtual TCHAR const* GetClassName() const override { return TEXT("SceneComponent"); }
+	virtual bool CanSupportUnifiedDispatch() const override { return true; }
+	virtual void GetSupportedInputs(TArray<FShaderFunctionDefinition>& OutFunctions) const override;
+	virtual void GetShaderParameters(TCHAR const* UID, FShaderParametersMetadataBuilder& InOutBuilder, FShaderParametersMetadataAllocations& InOutAllocations) const override;
+	virtual TCHAR const* GetShaderVirtualPath() const override;
+	virtual void GetShaderHash(FString& InOutKey) const override;
+	virtual void GetHLSL(FString& OutHLSL, FString const& InDataInterfaceName) const override;
+	virtual UComputeDataProvider* CreateDataProvider() const override;
+	//~ End UComputeDataInterface Interface
+
+private:
+	static TCHAR const* TemplateFilePath;
+};
+
+/** Compute Framework Data Provider for reading general scene data. */
+UCLASS(BlueprintType, editinlinenew, Category = ComputeFramework)
+class UOptimusSceneComponentDataProvider : public UComputeDataProvider
+{
+	GENERATED_BODY()
+
+public:
+	TWeakObjectPtr<USceneComponent> SceneComponent = nullptr;
+
+	//~ Begin UComputeDataProvider Interface
+	virtual void Initialize(const UComputeDataInterface* InDataInterface, UObject* InBinding, uint64 InInputMask, uint64 InOutputMask) override;
+	virtual FComputeDataProviderRenderProxy* GetRenderProxy() override;
+	//~ End UComputeDataProvider Interface
+};
+
+class FOptimusSceneComponentDataProviderProxy : public FComputeDataProviderRenderProxy
+{
+public:
+	FOptimusSceneComponentDataProviderProxy(USceneComponent* SceneComponent);
+
+	//~ Begin FComputeDataProviderRenderProxy Interface
+	virtual bool IsValid(FValidationData const& InValidationData) const override;
+	virtual void GatherDispatchData(FDispatchData const& InDispatchData) override;
+	//~ End FComputeDataProviderRenderProxy Interface
+
+private:
+	using FParameters = FSceneComponentDataInterfaceParameters;
+
+	bool bIsValid = false;
+	FTransform WorldTransform = FTransform::Identity;
+};
