[Deformer Graph] removed extra rigvm registration of deformer graph type, which should not be used as a general type in the graph.

#Synced-CL 50430604

diff --git a/Engine/Plugins/Animation/DeformerGraph/Source/OptimusCore/Private/OptimusCoreModule.cpp b/Engine/Plugins/Animation/DeformerGraph/Source/OptimusCore/Private/OptimusCoreModule.cpp
--- a/Engine/Plugins/Animation/DeformerGraph/Source/OptimusCore/Private/OptimusCoreModule.cpp
+++ b/Engine/Plugins/Animation/DeformerGraph/Source/OptimusCore/Private/OptimusCoreModule.cpp
@@ -14,7 +14,6 @@
 #include "AssetRegistry/AssetRegistryModule.h"
 #include "AssetRegistry/AssetData.h"
 #include "Nodes/OptimusNode_FunctionReference.h"
-#include "RigVMCore/RigVMRegistry.h"
 
 // Unique serialization id for Optimus .
 const FGuid FOptimusObjectVersion::GUID(0x93ede1aa, 0x10ca7375, 0x4df98a28, 0x49b157a0);
@@ -30,14 +29,6 @@ void FOptimusCoreModule::StartupModule()
 	FOptimusDataTypeRegistry::RegisterBuiltinTypes();
 	FOptimusDataTypeRegistry::RegisterEngineCallbacks();
 	UOptimusComputeDataInterface::RegisterAllTypes();
-
-	FRigVMRegistry& RigVMRegistry = FRigVMRegistry::Get();
-	static TPair<UClass*, FRigVMRegistry::ERegisterObjectOperation> const AllowedObjectTypes[] =
-	{
-		{ UOptimusDeformer::StaticClass(), FRigVMRegistry::ERegisterObjectOperation::Class },
-	};
-
-	RigVMRegistry.RegisterObjectTypes(AllowedObjectTypes);
 }
 
 void FOptimusCoreModule::ShutdownModule()
