[Backout] - CL50684474
#rnx
Original CL Desc
-----------------------------------------------------------------
[Deformer Graph] removed extra rigvm registration of deformer graph type, which should not be used as a general type in the graph.

diff --git a/Engine/Plugins/Animation/DeformerGraph/Source/OptimusCore/Private/OptimusCoreModule.cpp b/Engine/Plugins/Animation/DeformerGraph/Source/OptimusCore/Private/OptimusCoreModule.cpp
--- a/Engine/Plugins/Animation/DeformerGraph/Source/OptimusCore/Private/OptimusCoreModule.cpp
+++ b/Engine/Plugins/Animation/DeformerGraph/Source/OptimusCore/Private/OptimusCoreModule.cpp
@@ -14,6 +14,7 @@
 #include "AssetRegistry/AssetRegistryModule.h"
 #include "AssetRegistry/AssetData.h"
 #include "Nodes/OptimusNode_FunctionReference.h"
+#include "RigVMCore/RigVMRegistry.h"
 
 // Unique serialization id for Optimus .
 const FGuid FOptimusObjectVersion::GUID(0x93ede1aa, 0x10ca7375, 0x4df98a28, 0x49b157a0);
@@ -29,6 +30,14 @@ void FOptimusCoreModule::StartupModule()
 	FOptimusDataTypeRegistry::RegisterBuiltinTypes();
 	FOptimusDataTypeRegistry::RegisterEngineCallbacks();
 	UOptimusComputeDataInterface::RegisterAllTypes();
+
+	FRigVMRegistry& RigVMRegistry = FRigVMRegistry::Get();
+	static TPair<UClass*, FRigVMRegistry::ERegisterObjectOperation> const AllowedObjectTypes[] =
+	{
+		{ UOptimusDeformer::StaticClass(), FRigVMRegistry::ERegisterObjectOperation::Class },
+	};
+
+	RigVMRegistry.RegisterObjectTypes(AllowedObjectTypes);
 }
 
 void FOptimusCoreModule::ShutdownModule()
