Deformer & ControlRig: do not add child temp actors
- especially as control shape actors can be skeletal meshes now

#rnx

#Synced-CL 51558725

diff --git a/Engine/Plugins/Animation/DeformerGraph/Source/OptimusCore/Private/ControlRig/RigUnit_Optimus.cpp b/Engine/Plugins/Animation/DeformerGraph/Source/OptimusCore/Private/ControlRig/RigUnit_Optimus.cpp
--- a/Engine/Plugins/Animation/DeformerGraph/Source/OptimusCore/Private/ControlRig/RigUnit_Optimus.cpp
+++ b/Engine/Plugins/Animation/DeformerGraph/Source/OptimusCore/Private/ControlRig/RigUnit_Optimus.cpp
@@ -274,18 +274,29 @@ FRigUnit_AddOptimusDeformer_Execute()
 			static TArray<USkeletalMeshComponent*> GetComponentsToProcess(USkeletalMeshComponent* RigMeshComponent, bool DeformChildComponents, FName ExcludeChildComponentsWithTag)
 			{
 				TArray<USkeletalMeshComponent*> ComponentsToProcess = {RigMeshComponent};
-
+				
 				if (DeformChildComponents)
 				{
 					TArray<USceneComponent*> ChildComponents;
 					RigMeshComponent->GetChildrenComponents(true, ChildComponents);
+					
+					const AActor* RigMeshActor = RigMeshComponent->GetOwner();
 					for (USceneComponent* Component : ChildComponents)
 					{
 						if (USkeletalMeshComponent* SkeletalMeshComponent = Cast<USkeletalMeshComponent>(Component))
 						{
-							if (!Component->ComponentHasTag(ExcludeChildComponentsWithTag))
+							const AActor* ChildMeshActor = SkeletalMeshComponent->GetOwner();						
+#if WITH_EDITOR
+							const bool bIsTempActor = ChildMeshActor && ChildMeshActor != RigMeshActor && ChildMeshActor->bIsEditorPreviewActor;
+#else
+							constexpr bool bIsTempActor = false;
+#endif						
+							if (!bIsTempActor)
 							{
-								ComponentsToProcess.Add(SkeletalMeshComponent);
+								if (!Component->ComponentHasTag(ExcludeChildComponentsWithTag))
+								{
+									ComponentsToProcess.Add(SkeletalMeshComponent);
+								}
 							}
 						}
 					}	
@@ -293,7 +304,6 @@ FRigUnit_AddOptimusDeformer_Execute()
 
 				return ComponentsToProcess;
 			}
-			
 		};
 		
 		TArray<USkeletalMeshComponent*> ComponentsToProcess =
