Optimization to avoid creating duplicate section buffers in Optimus.

#Synced-CL 52152383

diff --git a/Engine/Plugins/Animation/DeformerGraph/Source/OptimusCore/Private/DataInterfaces/OptimusDataInterfaceAdvancedSkeleton.cpp b/Engine/Plugins/Animation/DeformerGraph/Source/OptimusCore/Private/DataInterfaces/OptimusDataInterfaceAdvancedSkeleton.cpp
--- a/Engine/Plugins/Animation/DeformerGraph/Source/OptimusCore/Private/DataInterfaces/OptimusDataInterfaceAdvancedSkeleton.cpp
+++ b/Engine/Plugins/Animation/DeformerGraph/Source/OptimusCore/Private/DataInterfaces/OptimusDataInterfaceAdvancedSkeleton.cpp
@@ -479,7 +479,7 @@ UComputeDataProvider* UOptimusAdvancedSkeletonDataInterface::CreateDataProvider(
 	Provider->Init(this, Cast<USkeletalMeshComponent>(InBinding));
 	return Provider;
 }
-void FOptimusBoneTransformBuffer::SetData(FSkeletalMeshLODRenderData const& InLodRenderData, const TArray<FTransform>& InBoneTransforms)
+void FOptimusBoneTransformBuffer::SetData(FSkeletalMeshLODRenderData const& InLodRenderData, const TArray<FTransform>& InBoneTransforms, bool bUnifiedBoneMap)
 {
 	if (InBoneTransforms.IsEmpty())
 	{
@@ -489,10 +489,12 @@ void FOptimusBoneTransformBuffer::SetData(FSkeletalMeshLODRenderData const& InLo
 	const FOptimusDataTypeHandle Matrix34TypeHandle = FOptimusDataTypeRegistry::Get().FindType(Matrix34TypeName);
 	int32 Matrix34ShaderSize = Matrix34TypeHandle->ShaderValueSize;
 
-	BufferData.AddDefaulted(InLodRenderData.RenderSections.Num());
-	NumBones.AddDefaulted(InLodRenderData.RenderSections.Num());
-				
-	for (int32 SectionIndex = 0; SectionIndex < InLodRenderData.RenderSections.Num(); ++SectionIndex)
+	const int32 NumUniqueSections = bUnifiedBoneMap ? 1 : InLodRenderData.RenderSections.Num();
+
+	BufferData.AddDefaulted(NumUniqueSections);
+	NumBones.AddDefaulted(NumUniqueSections);
+
+	for (int32 SectionIndex = 0; SectionIndex < NumUniqueSections; ++SectionIndex)
 	{
 		FSkelMeshRenderSection const& RenderSection = InLodRenderData.RenderSections[SectionIndex];
 		NumBones[SectionIndex] = RenderSection.BoneMap.Num();
@@ -502,9 +504,6 @@ void FOptimusBoneTransformBuffer::SetData(FSkeletalMeshLODRenderData const& InLo
 		{
 			int32 FinalBoneIndex = RenderSection.BoneMap[BoneIndex];
 			uint8* BoneDataPtr = BufferData[SectionIndex].GetData() + BoneIndex * Matrix34ShaderSize;
-
-			const uint8* ValuePtr = reinterpret_cast<const uint8*>(&InBoneTransforms[FinalBoneIndex]);
-						
 			Optimus::ConvertFTransformToFMatrix3x4(InBoneTransforms[FinalBoneIndex],
 				{ {BoneDataPtr , Matrix34ShaderSize} ,{}} );
 		}
@@ -879,19 +878,22 @@ void FOptimusAdvancedSkeletonDataProviderProxy::AllocateResources(FRDGBuilder& G
 	const int32 LodIndex = SkeletalMeshObject->GetLOD();
 	FSkeletalMeshRenderData const& SkeletalMeshRenderData = SkeletalMeshObject->GetSkeletalMeshRenderData();
 	FSkeletalMeshLODRenderData const& LodRenderData = SkeletalMeshRenderData.LODRenderData[LodIndex];
-	
+
+	const bool bUnifiedBoneMap = SkeletalMeshRenderData.HasUnifiedBoneMap();
+	const int32 NumUniqueSections = bUnifiedBoneMap ? 1 : LodRenderData.RenderSections.Num();
+
 	FallbackSRV = GraphBuilder.CreateSRV(FRDGBufferSRVDesc(GSystemTextures.GetDefaultStructuredBuffer(GraphBuilder, sizeof(int32))));
 
-	LayeredBoneMatrixBuffer.SetData(LodRenderData, LayeredBoneMatrices);
+	LayeredBoneMatrixBuffer.SetData(LodRenderData, LayeredBoneMatrices, bUnifiedBoneMap);
 	
 	LayeredBoneMatrixBuffer.AllocateResources(GraphBuilder);
-	
+
 	// Per-Bone Animation Attribute 
-	BuffersPerAttributePerSection.AddDefaulted(LodRenderData.RenderSections.Num());
-	BufferSRVsPerAttributePerSection.AddDefaulted(LodRenderData.RenderSections.Num());	
+	BuffersPerAttributePerSection.AddDefaulted(NumUniqueSections);
+	BufferSRVsPerAttributePerSection.AddDefaulted(NumUniqueSections);
 
-	AttributeBuffers.AddDefaulted(LodRenderData.RenderSections.Num());
-	for (int32 SectionIndex = 0; SectionIndex < LodRenderData.RenderSections.Num(); ++SectionIndex)
+	AttributeBuffers.AddDefaulted(NumUniqueSections);
+	for (int32 SectionIndex = 0; SectionIndex < NumUniqueSections; ++SectionIndex)
 	{
 		FSkelMeshRenderSection const& RenderSection = LodRenderData.RenderSections[SectionIndex];
 		AttributeBuffers[SectionIndex].AddDefaulted(AttributeBufferRuntimeData.Num());
@@ -917,14 +919,14 @@ void FOptimusAdvancedSkeletonDataProviderProxy::AllocateResources(FRDGBuilder& G
 		
 	}
 	
-	for (int32 InvocationIndex = 0 ; InvocationIndex < LodRenderData.RenderSections.Num(); InvocationIndex++)
+	for (int32 InvocationIndex = 0 ; InvocationIndex < NumUniqueSections; InvocationIndex++)
 	{
 		BuffersPerAttributePerSection[InvocationIndex].AddDefaulted(AttributeBuffers[InvocationIndex].Num());
 		BufferSRVsPerAttributePerSection[InvocationIndex].AddDefaulted(AttributeBuffers[InvocationIndex].Num());
 		
-		FSkelMeshRenderSection const& RenderSection = LodRenderData.RenderSections[InvocationIndex];	
+		FSkelMeshRenderSection const& RenderSection = LodRenderData.RenderSections[InvocationIndex];
 		int32 NumBones = RenderSection.BoneMap.Num();
-		
+
 		for (int32 AttributeIndex = 0; AttributeIndex < AttributeBuffers[InvocationIndex].Num(); ++AttributeIndex)
 		{
 			const TArray<uint8>& BoneData = AttributeBuffers[InvocationIndex][AttributeIndex];
@@ -1016,6 +1018,7 @@ void FOptimusAdvancedSkeletonDataProviderProxy::GatherDispatchData(FDispatchData
 	const int32 LodIndex = SkeletalMeshObject->GetLOD();
 	FSkeletalMeshRenderData const& SkeletalMeshRenderData = SkeletalMeshObject->GetSkeletalMeshRenderData();
 	FSkeletalMeshLODRenderData const* LodRenderData = &SkeletalMeshRenderData.LODRenderData[LodIndex];
+	const bool bUnifiedBoneMap = SkeletalMeshRenderData.HasUnifiedBoneMap();
 
 	FRHIShaderResourceView* NullSRVBinding = GWhiteVertexBufferWithSRV->ShaderResourceViewRHI.GetReference();
 
@@ -1037,7 +1040,9 @@ void FOptimusAdvancedSkeletonDataProviderProxy::GatherDispatchData(FDispatchData
 		
 		const bool bPreviousFrame = false;
 		FRHIShaderResourceView* BoneBufferSRV = FSkeletalMeshDeformerHelpers::GetBoneBufferForReading(SkeletalMeshObject, LodIndex, InvocationIndex, bPreviousFrame);
-	
+
+		const int32 BoneBufferIndex = bUnifiedBoneMap ? 0 : InvocationIndex;
+
 		FDefaultParameters* Parameters = reinterpret_cast<FDefaultParameters*>(ParameterBufferPtr);
 		Parameters->NumVertices = LodRenderData->GetNumVertices();
 		Parameters->NumBoneInfluences = WeightBuffer->GetMaxBoneInfluences();
@@ -1047,12 +1052,12 @@ void FOptimusAdvancedSkeletonDataProviderProxy::GatherDispatchData(FDispatchData
 		Parameters->InputWeightStream = SkinWeightBufferSRV != nullptr ? SkinWeightBufferSRV : NullSRVBinding;
 		Parameters->InputWeightLookupStream = InputWeightLookupStreamSRV != nullptr ? InputWeightLookupStreamSRV : NullSRVBinding;
 		Parameters->LayeredBoneMatrices =
-			LayeredBoneMatrixBuffer.HasData() && LayeredBoneMatrixBuffer.BufferSRVPerSection[InvocationIndex] != nullptr ?
-				LayeredBoneMatrixBuffer.BufferSRVPerSection[InvocationIndex] : FallbackSRV;
+			LayeredBoneMatrixBuffer.HasData() && LayeredBoneMatrixBuffer.BufferSRVPerSection[BoneBufferIndex] != nullptr ?
+				LayeredBoneMatrixBuffer.BufferSRVPerSection[BoneBufferIndex] : FallbackSRV;
 
 		for(int32 AttributeIndex = 0 ; AttributeIndex < AttributeBufferRuntimeData.Num(); AttributeIndex++)
 		{
-			*((FRDGBufferSRV**)(ParameterBufferPtr + AttributeBufferRuntimeData[AttributeIndex].Offset)) = BufferSRVsPerAttributePerSection[InvocationIndex][AttributeIndex];
+			*((FRDGBufferSRV**)(ParameterBufferPtr + AttributeBufferRuntimeData[AttributeIndex].Offset)) = BufferSRVsPerAttributePerSection[BoneBufferIndex][AttributeIndex];
 		}
 	}
 }
diff --git a/Engine/Plugins/Animation/DeformerGraph/Source/OptimusCore/Private/DataInterfaces/OptimusDataInterfaceAdvancedSkeleton.h b/Engine/Plugins/Animation/DeformerGraph/Source/OptimusCore/Private/DataInterfaces/OptimusDataInterfaceAdvancedSkeleton.h
--- a/Engine/Plugins/Animation/DeformerGraph/Source/OptimusCore/Private/DataInterfaces/OptimusDataInterfaceAdvancedSkeleton.h
+++ b/Engine/Plugins/Animation/DeformerGraph/Source/OptimusCore/Private/DataInterfaces/OptimusDataInterfaceAdvancedSkeleton.h
@@ -200,7 +200,7 @@ struct FOptimusBoneTransformBuffer
 	TArray<TArray<uint8>> BufferData;
 	TArray<int32> NumBones;
 
-	void SetData(FSkeletalMeshLODRenderData const& InLodRenderData, const TArray<FTransform>& InBoneTransforms);
+	void SetData(FSkeletalMeshLODRenderData const& InLodRenderData, const TArray<FTransform>& InBoneTransforms, bool bInUnifiedBoneMap);
 	bool HasData() const;
 
 	void AllocateResources(FRDGBuilder& GraphBuilder);
diff --git a/Engine/Plugins/Animation/DeformerGraph/Source/OptimusCore/Private/DataInterfaces/OptimusDataInterfaceSkinWeightsAsVertexMask.cpp b/Engine/Plugins/Animation/DeformerGraph/Source/OptimusCore/Private/DataInterfaces/OptimusDataInterfaceSkinWeightsAsVertexMask.cpp
--- a/Engine/Plugins/Animation/DeformerGraph/Source/OptimusCore/Private/DataInterfaces/OptimusDataInterfaceSkinWeightsAsVertexMask.cpp
+++ b/Engine/Plugins/Animation/DeformerGraph/Source/OptimusCore/Private/DataInterfaces/OptimusDataInterfaceSkinWeightsAsVertexMask.cpp
@@ -296,7 +296,7 @@ FComputeDataProviderRenderProxy* UOptimusSkinWeightsAsVertexMaskDataProvider::Ge
 				CachedBoneIsSelectedPerSectionPerLod.Reset();
 				
 				const FReferenceSkeleton& RefSkeleton = SkeletalMesh->GetSkinnedAsset()->GetRefSkeleton();
-
+				const bool bUnifiedBoneMap = SkeletalMeshRenderData.HasUnifiedBoneMap();
 
 				// Negative distance : Towards Root
 				// Positive distance : Towards Leaf
@@ -358,7 +358,7 @@ FComputeDataProviderRenderProxy* UOptimusSkinWeightsAsVertexMaskDataProvider::Ge
 						}
 					}
 				}
-				
+
 				for (int32 LODIndex = 0; LODIndex < SkeletalMeshRenderData.LODRenderData.Num(); LODIndex++)
 				{
 					FSkeletalMeshLODRenderData const& LocalLodRenderData = SkeletalMeshRenderData.LODRenderData[LODIndex];
@@ -380,6 +380,18 @@ FComputeDataProviderRenderProxy* UOptimusSkinWeightsAsVertexMaskDataProvider::Ge
 								BoneIsSelectedRef[BoneIndex] = true;
 							}
 						}
+
+						// Only process first section when unified.
+						if (bUnifiedBoneMap)
+						{
+							break;
+						}
+					}
+
+					// Unified bone map is cross-LOD, so LOD 0 / Section 0 covers all bones. Break the LOD loop too.
+					if (bUnifiedBoneMap)
+					{
+						break;
 					}
 				}
 			}
@@ -471,16 +483,19 @@ void FOptimusSkinWeightsAsVertexMaskDataProviderProxy::AllocateResources(FRDGBui
 	const int32 LodIndex = SkeletalMeshObject->GetLOD();
 	FSkeletalMeshRenderData const& SkeletalMeshRenderData = SkeletalMeshObject->GetSkeletalMeshRenderData();
 	FSkeletalMeshLODRenderData const* LodRenderData = &SkeletalMeshRenderData.LODRenderData[LodIndex];
-	
-	
-	BoneIsSelectedBuffersPerSection.AddDefaulted(LodRenderData->RenderSections.Num());
-	BoneIsSelectedBufferSRVsPerSection.AddDefaulted(LodRenderData->RenderSections.Num());	
-	
-	for (int32 InvocationIndex = 0 ; InvocationIndex < LodRenderData->RenderSections.Num(); InvocationIndex++)
+	const bool bUnifiedBoneMap = SkeletalMeshRenderData.HasUnifiedBoneMap();
+	const int32 NumUniqueSections = bUnifiedBoneMap ? 1 : LodRenderData->RenderSections.Num();
+
+	BoneIsSelectedBuffersPerSection.AddDefaulted(NumUniqueSections);
+	BoneIsSelectedBufferSRVsPerSection.AddDefaulted(NumUniqueSections);
+
+	for (int32 InvocationIndex = 0 ; InvocationIndex < NumUniqueSections; InvocationIndex++)
 	{
-		int32 NumBones = LodRenderData->RenderSections[InvocationIndex].BoneMap.Num();
-		
-		const TArray<uint32>& BoneIsSelected = BoneIsSelectedPerSectionPerLod[LodIndex][InvocationIndex];
+		const TArray<uint32>& BoneIsSelected = bUnifiedBoneMap ? BoneIsSelectedPerSectionPerLod[0][0] : BoneIsSelectedPerSectionPerLod[LodIndex][InvocationIndex];
+
+		// When using unified bone map, the BoneIsSelected array was built from LOD 0's bone map.
+		// The current LOD may have a different bone count. Use the smaller of the two to avoid OOB.
+		int32 NumBones = FMath::Min(LodRenderData->RenderSections[InvocationIndex].BoneMap.Num(), BoneIsSelected.Num());
 
 		int32 Stride = sizeof(uint32);
 		int32 ByteSize = NumBones * Stride;
@@ -566,6 +581,7 @@ void FOptimusSkinWeightsAsVertexMaskDataProviderProxy::GatherDispatchData(FDispa
 	const int32 LodIndex = SkeletalMeshObject->GetLOD();
 	FSkeletalMeshRenderData const& SkeletalMeshRenderData = SkeletalMeshObject->GetSkeletalMeshRenderData();
 	FSkeletalMeshLODRenderData const* LodRenderData = &SkeletalMeshRenderData.LODRenderData[LodIndex];
+	const bool bUnifiedBoneMap = SkeletalMeshRenderData.HasUnifiedBoneMap();
 
 	FRHIShaderResourceView* NullSRVBinding = GWhiteVertexBufferWithSRV->ShaderResourceViewRHI.GetReference();
 	
@@ -591,6 +607,6 @@ void FOptimusSkinWeightsAsVertexMaskDataProviderProxy::GatherDispatchData(FDispa
 		Parameters.InputWeightIndexSize = WeightBuffer->GetBoneIndexByteSize() | (WeightBuffer->GetBoneWeightByteSize() << 8);
 		Parameters.InputWeightStream = SkinWeightBufferSRV != nullptr ? SkinWeightBufferSRV : NullSRVBinding;
 		Parameters.InputWeightLookupStream = InputWeightLookupStreamSRV != nullptr ? InputWeightLookupStreamSRV : NullSRVBinding;
-		Parameters.BoneIsSelected = BoneIsSelectedBufferSRVsPerSection[InvocationIndex];
+		Parameters.BoneIsSelected = BoneIsSelectedBufferSRVsPerSection[bUnifiedBoneMap ? 0 : InvocationIndex];
 	}
 }
diff --git a/Engine/Plugins/Animation/MLDeformer/MLDeformerFramework/Source/MLDeformerFramework/Private/DeformerGraphDataInterfaceSkeletonWithQuats.cpp b/Engine/Plugins/Animation/MLDeformer/MLDeformerFramework/Source/MLDeformerFramework/Private/DeformerGraphDataInterfaceSkeletonWithQuats.cpp
--- a/Engine/Plugins/Animation/MLDeformer/MLDeformerFramework/Source/MLDeformerFramework/Private/DeformerGraphDataInterfaceSkeletonWithQuats.cpp
+++ b/Engine/Plugins/Animation/MLDeformer/MLDeformerFramework/Source/MLDeformerFramework/Private/DeformerGraphDataInterfaceSkeletonWithQuats.cpp
@@ -231,9 +231,14 @@ namespace UE::MLDeformer
 		FSkeletalMeshRenderData const& SkeletalMeshRenderData = SkeletalMeshObject->GetSkeletalMeshRenderData();
 		FSkeletalMeshLODRenderData const* LodRenderData = &SkeletalMeshRenderData.LODRenderData[LodIndex];
 		const int32 NumSections = LodRenderData->RenderSections.Num();
-		PerSectionRefToLocalQuats.SetNum(NumSections);
-		PerSectionRefToLocalQuatsSRVs.SetNum(NumSections);
-		for (int32 SectionIndex = 0; SectionIndex < NumSections; SectionIndex++)
+
+		const bool bUnifiedBoneMap = SkeletalMeshRenderData.HasUnifiedBoneMap();
+		const int32 NumUniqueSections = bUnifiedBoneMap ? 1 : NumSections;
+
+		PerSectionRefToLocalQuatsSRVs.SetNum(NumUniqueSections);
+		PerSectionRefToLocalQuats.SetNum(NumUniqueSections);
+
+		for (int32 SectionIndex = 0; SectionIndex < NumUniqueSections; SectionIndex++)
 		{
 			TArray<FQuat4f> &SectionRefToLocalQuats = PerSectionRefToLocalQuats[SectionIndex];
 			FRDGBufferSRVRef &SectionRefToLocalQuatsSRV = PerSectionRefToLocalQuatsSRVs[SectionIndex];
@@ -266,6 +271,7 @@ namespace UE::MLDeformer
 		const int32 LodIndex = SkeletalMeshObject->GetLOD();
 		FSkeletalMeshRenderData const& SkeletalMeshRenderData = SkeletalMeshObject->GetSkeletalMeshRenderData();
 		FSkeletalMeshLODRenderData const* LodRenderData = &SkeletalMeshRenderData.LODRenderData[LodIndex];
+		const bool bUnifiedBoneMap = SkeletalMeshRenderData.HasUnifiedBoneMap();
 
 		FRHIShaderResourceView* NullSRVBinding = GWhiteVertexBufferWithSRV->ShaderResourceViewRHI.GetReference();
 
@@ -278,6 +284,8 @@ namespace UE::MLDeformer
 		const TStridedView<FParameters> ParameterArray = MakeStridedParameterView<FParameters>(InDispatchData);
 		for (int32 InvocationIndex = 0; InvocationIndex < ParameterArray.Num(); ++InvocationIndex)
 		{
+			const int32 BoneBufferIndex = bUnifiedBoneMap ? 0 : InvocationIndex;
+
 			const bool bPreviousFrame = false;
 			FRHIShaderResourceView* BoneBufferSRV = FSkeletalMeshDeformerHelpers::GetBoneBufferForReading(SkeletalMeshObject, LodIndex, InvocationIndex, bPreviousFrame);
 			
@@ -287,7 +295,7 @@ namespace UE::MLDeformer
 			Parameters.InputWeightStride = WeightBuffer->GetConstantInfluencesVertexStride();
 			Parameters.InputWeightIndexSize = WeightBuffer->GetBoneIndexByteSize() | (WeightBuffer->GetBoneWeightByteSize() << 8);
 			Parameters.BoneMatrices = BoneBufferSRV != nullptr ? BoneBufferSRV : NullSRVBinding;
-			Parameters.BoneQuats = PerSectionRefToLocalQuatsSRVs.IsValidIndex(InvocationIndex) ? PerSectionRefToLocalQuatsSRVs[InvocationIndex] : nullptr;
+			Parameters.BoneQuats = PerSectionRefToLocalQuatsSRVs.IsValidIndex(BoneBufferIndex) ? PerSectionRefToLocalQuatsSRVs[BoneBufferIndex] : nullptr;
 			Parameters.InputWeightStream = SkinWeightBufferSRV != nullptr ? SkinWeightBufferSRV : NullSRVBinding;
 			Parameters.InputWeightLookupStream = InputWeightLookupStreamSRV != nullptr ? InputWeightLookupStreamSRV : NullSRVBinding;
 		}
diff --git a/Engine/Plugins/Animation/MLDeformer/NearestNeighborModel/Source/NearestNeighborModel/Private/OptimusDataInterfaceSkeletonWithQuats.cpp b/Engine/Plugins/Animation/MLDeformer/NearestNeighborModel/Source/NearestNeighborModel/Private/OptimusDataInterfaceSkeletonWithQuats.cpp
--- a/Engine/Plugins/Animation/MLDeformer/NearestNeighborModel/Source/NearestNeighborModel/Private/OptimusDataInterfaceSkeletonWithQuats.cpp
+++ b/Engine/Plugins/Animation/MLDeformer/NearestNeighborModel/Source/NearestNeighborModel/Private/OptimusDataInterfaceSkeletonWithQuats.cpp
@@ -235,9 +235,13 @@ void FOptimusSkeletonWithQuatsDataProviderProxy::AllocateResources(FRDGBuilder&
 	FSkeletalMeshRenderData const& SkeletalMeshRenderData = SkeletalMeshObject->GetSkeletalMeshRenderData();
 	FSkeletalMeshLODRenderData const* LodRenderData = &SkeletalMeshRenderData.LODRenderData[LodIndex];
 	const int32 NumSections = LodRenderData->RenderSections.Num();
-	PerSectionRefToLocalQuats.SetNum(NumSections);
-	PerSectionRefToLocalQuatsSRVs.SetNum(NumSections);
-	for (int32 SectionIndex = 0; SectionIndex < NumSections; SectionIndex++)
+
+	const bool bUnifiedBoneMap = SkeletalMeshRenderData.HasUnifiedBoneMap();
+	const int32 NumUniqueSections = bUnifiedBoneMap ? 1 : NumSections;
+
+	PerSectionRefToLocalQuats.SetNum(NumUniqueSections);
+	PerSectionRefToLocalQuatsSRVs.SetNum(NumUniqueSections);
+	for (int32 SectionIndex = 0; SectionIndex < NumUniqueSections; SectionIndex++)
 	{
 		TArray<FQuat4f> &SectionRefToLocalQuats = PerSectionRefToLocalQuats[SectionIndex];
 		FRDGBufferSRVRef &SectionRefToLocalQuatsSRV = PerSectionRefToLocalQuatsSRVs[SectionIndex];
@@ -270,13 +274,14 @@ void FOptimusSkeletonWithQuatsDataProviderProxy::GatherDispatchData(FDispatchDat
 	const int32 LodIndex = SkeletalMeshObject->GetLOD();
 	FSkeletalMeshRenderData const& SkeletalMeshRenderData = SkeletalMeshObject->GetSkeletalMeshRenderData();
 	FSkeletalMeshLODRenderData const* LodRenderData = &SkeletalMeshRenderData.LODRenderData[LodIndex];
+	const bool bUnifiedBoneMap = SkeletalMeshRenderData.HasUnifiedBoneMap();
 
 	FRHIShaderResourceView* NullSRVBinding = GWhiteVertexBufferWithSRV->ShaderResourceViewRHI.GetReference();
 
 	const TStridedView<FParameters> ParameterArray = MakeStridedParameterView<FParameters>(InDispatchData);
 	for (int32 InvocationIndex = 0; InvocationIndex < ParameterArray.Num(); ++InvocationIndex)
 	{
-		FSkelMeshRenderSection const& RenderSection = LodRenderData->RenderSections[InvocationIndex];
+		const int32 BoneBufferIndex = bUnifiedBoneMap ? 0 : InvocationIndex;
 
 		const bool bPreviousFrame = false;
 		FRHIShaderResourceView* BoneBufferSRV = FSkeletalMeshDeformerHelpers::GetBoneBufferForReading(SkeletalMeshObject, LodIndex, InvocationIndex, bPreviousFrame);
@@ -293,7 +298,7 @@ void FOptimusSkeletonWithQuatsDataProviderProxy::GatherDispatchData(FDispatchDat
 		Parameters.InputWeightStride = WeightBuffer->GetConstantInfluencesVertexStride();
 		Parameters.InputWeightIndexSize = WeightBuffer->GetBoneIndexByteSize() | (WeightBuffer->GetBoneWeightByteSize() << 8);
 		Parameters.BoneMatrices = BoneBufferSRV != nullptr ? BoneBufferSRV : NullSRVBinding;
-		Parameters.BoneQuats = PerSectionRefToLocalQuatsSRVs.IsValidIndex(InvocationIndex) ? PerSectionRefToLocalQuatsSRVs[InvocationIndex] : nullptr;
+		Parameters.BoneQuats = PerSectionRefToLocalQuatsSRVs.IsValidIndex(BoneBufferIndex) ? PerSectionRefToLocalQuatsSRVs[BoneBufferIndex] : nullptr;
 		Parameters.InputWeightStream = SkinWeightBufferSRV != nullptr ? SkinWeightBufferSRV : NullSRVBinding;
 		Parameters.InputWeightLookupStream = InputWeightLookupStreamSRV != nullptr ? InputWeightLookupStreamSRV : NullSRVBinding;
 	}
