[Deformer Graph] Multiple fixes for Optimus runtime

- Fix Optimus data interface MeshObject thread safety (use-after-free on secondary bound components), it now uses render thread apis to query the mesh object instead of caching it from game thread

- Created a scene extension to track which SKM component is deformed by Optimus. It also tracks LOD changes and Output buffers for each SKM at each execution group. Output buffer tracking is used to ensure a deformer can only read the deformer output of the a secondary component if it was produced in a earlier execution group. Reading from the same execution group is currently unsupported

#Synced-CL 52183294

diff --git a/Engine/Plugins/Animation/DeformerGraph/Source/OptimusCore/OptimusCore.Build.cs b/Engine/Plugins/Animation/DeformerGraph/Source/OptimusCore/OptimusCore.Build.cs
--- a/Engine/Plugins/Animation/DeformerGraph/Source/OptimusCore/OptimusCore.Build.cs
+++ b/Engine/Plugins/Animation/DeformerGraph/Source/OptimusCore/OptimusCore.Build.cs
@@ -1,5 +1,6 @@
 // Copyright Epic Games, Inc. All Rights Reserved.
 
+using System.IO;
 namespace UnrealBuildTool.Rules
 {
     public class OptimusCore : ModuleRules
@@ -8,6 +9,8 @@ namespace UnrealBuildTool.Rules
 		{
 			NumIncludedBytesPerUnityCPPOverride = 688128; // best unity size found from using UBT ProfileUnitySizes mode
 
+			PrivateIncludePaths.Add(Path.Combine(EngineDirectory, "Source/Runtime/Renderer/Private"));
+
             PublicDependencyModuleNames.AddRange(
 				new string[]
 				{
diff --git a/Engine/Plugins/Animation/DeformerGraph/Source/OptimusCore/Private/DataInterfaces/OptimusDataInterfaceAdvancedSkeleton.cpp b/Engine/Plugins/Animation/DeformerGraph/Source/OptimusCore/Private/DataInterfaces/OptimusDataInterfaceAdvancedSkeleton.cpp
--- a/Engine/Plugins/Animation/DeformerGraph/Source/OptimusCore/Private/DataInterfaces/OptimusDataInterfaceAdvancedSkeleton.cpp
+++ b/Engine/Plugins/Animation/DeformerGraph/Source/OptimusCore/Private/DataInterfaces/OptimusDataInterfaceAdvancedSkeleton.cpp
@@ -823,7 +823,9 @@ FComputeDataProviderRenderProxy* UOptimusAdvancedSkeletonDataProvider::GetRender
 			
 
 			// Pipe data into proxy
-			Proxy->SkeletalMeshObject = SkeletalMeshObject;
+			Proxy->SceneInterface = SkeletalMesh->GetScene();
+			Proxy->ComponentId = SkeletalMesh->GetPrimitiveSceneId();
+			
 			Proxy->SkinWeightProfile = SkinWeightProfile;
 			Proxy->ParameterBuffer.AddDefaulted(ParameterBufferSize);
 			Proxy->LayeredBoneMatrices = MoveTemp(LayeredBoneMatrices);
@@ -842,11 +844,12 @@ FOptimusAdvancedSkeletonDataProviderProxy::FOptimusAdvancedSkeletonDataProviderP
 
 bool FOptimusAdvancedSkeletonDataProviderProxy::IsValid(FValidationData const& InValidationData) const
 {
-	if (InValidationData.ParameterStructSize != ParameterBuffer.Num())
+	FSkeletalMeshObject* SkeletalMeshObject = FSkeletalMeshDeformerHelpers::GetSkeletalMeshObject(SceneInterface, ComponentId);
+	if (!SkeletalMeshObject)
 	{
 		return false;
 	}
-	if (SkeletalMeshObject == nullptr)
+	if (InValidationData.ParameterStructSize != ParameterBuffer.Num())
 	{
 		return false;
 	}
@@ -875,6 +878,7 @@ bool FOptimusAdvancedSkeletonDataProviderProxy::IsValid(FValidationData const& I
 
 void FOptimusAdvancedSkeletonDataProviderProxy::AllocateResources(FRDGBuilder& GraphBuilder, FAllocationData const& InAllocationData)
 {
+	FSkeletalMeshObject* SkeletalMeshObject = FSkeletalMeshDeformerHelpers::GetSkeletalMeshObject(SceneInterface, ComponentId);
 	const int32 LodIndex = SkeletalMeshObject->GetLOD();
 	FSkeletalMeshRenderData const& SkeletalMeshRenderData = SkeletalMeshObject->GetSkeletalMeshRenderData();
 	FSkeletalMeshLODRenderData const& LodRenderData = SkeletalMeshRenderData.LODRenderData[LodIndex];
@@ -978,6 +982,7 @@ struct FAdvancedSkeletonDataInterfacePermutationIds
 
 void FOptimusAdvancedSkeletonDataProviderProxy::GatherPermutations(FPermutationData& InOutPermutationData) const
 {
+	FSkeletalMeshObject* SkeletalMeshObject = FSkeletalMeshDeformerHelpers::GetSkeletalMeshObject(SceneInterface, ComponentId);
 	const int32 LodIndex = SkeletalMeshObject->GetLOD();
 	FSkeletalMeshRenderData const& SkeletalMeshRenderData = SkeletalMeshObject->GetSkeletalMeshRenderData();
 	FSkeletalMeshLODRenderData const* LodRenderData = &SkeletalMeshRenderData.LODRenderData[LodIndex];
@@ -1015,6 +1020,7 @@ void FOptimusAdvancedSkeletonDataProviderProxy::GatherPermutations(FPermutationD
 
 void FOptimusAdvancedSkeletonDataProviderProxy::GatherDispatchData(FDispatchData const& InDispatchData)
 {
+	FSkeletalMeshObject* SkeletalMeshObject = FSkeletalMeshDeformerHelpers::GetSkeletalMeshObject(SceneInterface, ComponentId);
 	const int32 LodIndex = SkeletalMeshObject->GetLOD();
 	FSkeletalMeshRenderData const& SkeletalMeshRenderData = SkeletalMeshObject->GetSkeletalMeshRenderData();
 	FSkeletalMeshLODRenderData const* LodRenderData = &SkeletalMeshRenderData.LODRenderData[LodIndex];
diff --git a/Engine/Plugins/Animation/DeformerGraph/Source/OptimusCore/Private/DataInterfaces/OptimusDataInterfaceAdvancedSkeleton.h b/Engine/Plugins/Animation/DeformerGraph/Source/OptimusCore/Private/DataInterfaces/OptimusDataInterfaceAdvancedSkeleton.h
--- a/Engine/Plugins/Animation/DeformerGraph/Source/OptimusCore/Private/DataInterfaces/OptimusDataInterfaceAdvancedSkeleton.h
+++ b/Engine/Plugins/Animation/DeformerGraph/Source/OptimusCore/Private/DataInterfaces/OptimusDataInterfaceAdvancedSkeleton.h
@@ -9,11 +9,13 @@
 #include "OptimusValueContainerStruct.h"
 #include "ComputeFramework/ComputeDataProvider.h"
 
+#include "PrimitiveComponentId.h"
 #include "OptimusDataInterfaceAdvancedSkeleton.generated.h"
 
 #define UE_API OPTIMUSCORE_API
 
 struct FOptimusAnimAttributeDescription;
+class FSceneInterface;
 class FSkeletalMeshObject;
 class FAdvancedSkeletonDataInterfaceDefaultParameters;
 class USkeletalMeshComponent;
@@ -263,7 +265,8 @@ public:
 	using FDefaultParameters = FAdvancedSkeletonDataInterfaceDefaultParameters;
 
 	// Data supplied by data provider
-	FSkeletalMeshObject* SkeletalMeshObject = nullptr;
+	FSceneInterface* SceneInterface = nullptr;
+	FPrimitiveComponentId ComponentId;
 	FName SkinWeightProfile = NAME_None;
 	
 	TArray<uint8> ParameterBuffer;
diff --git a/Engine/Plugins/Animation/DeformerGraph/Source/OptimusCore/Private/DataInterfaces/OptimusDataInterfaceCloth.cpp b/Engine/Plugins/Animation/DeformerGraph/Source/OptimusCore/Private/DataInterfaces/OptimusDataInterfaceCloth.cpp
--- a/Engine/Plugins/Animation/DeformerGraph/Source/OptimusCore/Private/DataInterfaces/OptimusDataInterfaceCloth.cpp
+++ b/Engine/Plugins/Animation/DeformerGraph/Source/OptimusCore/Private/DataInterfaces/OptimusDataInterfaceCloth.cpp
@@ -132,10 +132,11 @@ FComputeDataProviderRenderProxy* UOptimusClothDataProvider::GetRenderProxy()
 
 FOptimusClothDataProviderProxy::FOptimusClothDataProviderProxy(USkinnedMeshComponent* SkinnedMeshComponent)
 {
-	if (SkinnedMeshComponent != nullptr && SkinnedMeshComponent->GetScene() != nullptr)
+	if (SkinnedMeshComponent != nullptr && SkinnedMeshComponent->IsRegistered())
 	{
-		SkeletalMeshObject = SkinnedMeshComponent == nullptr ? nullptr : SkinnedMeshComponent->MeshObject;
-		FrameNumber = SkinnedMeshComponent == nullptr ? 0 : SkinnedMeshComponent->GetScene()->GetFrameNumber() + 1; // +1 matches the logic for FrameNumberToPrepare in FSkeletalMeshObjectGPUSkin::Update()
+		SceneInterface = SkinnedMeshComponent->GetScene();
+		ComponentId = SkinnedMeshComponent->GetPrimitiveSceneId();
+		FrameNumber = SkinnedMeshComponent->GetScene() != nullptr ? SkinnedMeshComponent->GetScene()->GetFrameNumber() + 1 : 0; // +1 matches the logic for FrameNumberToPrepare in FSkeletalMeshObjectGPUSkin::Update()
 
 		if (USkeletalMeshComponent* SkeletalMeshComponent = Cast<USkeletalMeshComponent>(SkinnedMeshComponent))
 		{
@@ -151,7 +152,9 @@ bool FOptimusClothDataProviderProxy::IsValid(FValidationData const& InValidation
 	{
 		return false;
 	}
-	if (SkeletalMeshObject == nullptr)
+
+	FSkeletalMeshObject* SkeletalMeshObject = FSkeletalMeshDeformerHelpers::GetSkeletalMeshObject(SceneInterface, ComponentId);
+	if (!SkeletalMeshObject)
 	{
 		return false;
 	}
@@ -179,6 +182,7 @@ struct FClothDataInterfacePermutationIds
 
 void FOptimusClothDataProviderProxy::GatherPermutations(FPermutationData& InOutPermutationData) const
 {
+	FSkeletalMeshObject* SkeletalMeshObject = FSkeletalMeshDeformerHelpers::GetSkeletalMeshObject(SceneInterface, ComponentId);
 	const int32 LodIndex = SkeletalMeshObject->GetLOD();
 	FSkeletalMeshRenderData const& SkeletalMeshRenderData = SkeletalMeshObject->GetSkeletalMeshRenderData();
 	FSkeletalMeshLODRenderData const* LodRenderData = &SkeletalMeshRenderData.LODRenderData[LodIndex];
@@ -196,6 +200,7 @@ void FOptimusClothDataProviderProxy::GatherPermutations(FPermutationData& InOutP
 
 void FOptimusClothDataProviderProxy::GatherDispatchData(FDispatchData const& InDispatchData)
 {
+	FSkeletalMeshObject* SkeletalMeshObject = FSkeletalMeshDeformerHelpers::GetSkeletalMeshObject(SceneInterface, ComponentId);
 	const int32 LodIndex = SkeletalMeshObject->GetLOD();
 	FSkeletalMeshRenderData const& SkeletalMeshRenderData = SkeletalMeshObject->GetSkeletalMeshRenderData();
 	FSkeletalMeshLODRenderData const* LodRenderData = &SkeletalMeshRenderData.LODRenderData[LodIndex];
diff --git a/Engine/Plugins/Animation/DeformerGraph/Source/OptimusCore/Private/DataInterfaces/OptimusDataInterfaceCloth.h b/Engine/Plugins/Animation/DeformerGraph/Source/OptimusCore/Private/DataInterfaces/OptimusDataInterfaceCloth.h
--- a/Engine/Plugins/Animation/DeformerGraph/Source/OptimusCore/Private/DataInterfaces/OptimusDataInterfaceCloth.h
+++ b/Engine/Plugins/Animation/DeformerGraph/Source/OptimusCore/Private/DataInterfaces/OptimusDataInterfaceCloth.h
@@ -4,11 +4,13 @@
 
 #include "OptimusComputeDataInterface.h"
 #include "ComputeFramework/ComputeDataProvider.h"
+#include "PrimitiveComponentId.h"
 #include "OptimusDataInterfaceCloth.generated.h"
 
 #define UE_API OPTIMUSCORE_API
 
 class FClothDataInterfaceParameters;
+class FSceneInterface;
 class FSkeletalMeshObject;
 class USkinnedMeshComponent;
 
@@ -68,7 +70,8 @@ public:
 private:
 	using FParameters = FClothDataInterfaceParameters;
 
-	FSkeletalMeshObject* SkeletalMeshObject = nullptr;
+	FSceneInterface* SceneInterface = nullptr;
+	FPrimitiveComponentId ComponentId;
 	float ClothBlendWeight = 1.0f;
 	FVector3f MeshScale = FVector3f::OneVector;
 	uint32 FrameNumber = 0;
diff --git a/Engine/Plugins/Animation/DeformerGraph/Source/OptimusCore/Private/DataInterfaces/OptimusDataInterfaceConnectivity.cpp b/Engine/Plugins/Animation/DeformerGraph/Source/OptimusCore/Private/DataInterfaces/OptimusDataInterfaceConnectivity.cpp
--- a/Engine/Plugins/Animation/DeformerGraph/Source/OptimusCore/Private/DataInterfaces/OptimusDataInterfaceConnectivity.cpp
+++ b/Engine/Plugins/Animation/DeformerGraph/Source/OptimusCore/Private/DataInterfaces/OptimusDataInterfaceConnectivity.cpp
@@ -234,9 +234,13 @@ FComputeDataProviderRenderProxy* UOptimusConnectivityDataProvider::GetRenderProx
 
 
 FOptimusConnectivityDataProviderProxy::FOptimusConnectivityDataProviderProxy(USkinnedMeshComponent* SkinnedMeshComponent, TArray< TArray<uint32> >& InAdjacencyBufferPerLod)
-	: SkeletalMeshObject(SkinnedMeshComponent != nullptr ? SkinnedMeshComponent->MeshObject : nullptr)
-	, AdjacencyBufferPerLod(InAdjacencyBufferPerLod)
+	: AdjacencyBufferPerLod(InAdjacencyBufferPerLod)
 {
+	if (SkinnedMeshComponent != nullptr && SkinnedMeshComponent->IsRegistered())
+	{
+		SceneInterface = SkinnedMeshComponent->GetScene();
+		ComponentId = SkinnedMeshComponent->GetPrimitiveSceneId();
+	}
 }
 
 bool FOptimusConnectivityDataProviderProxy::IsValid(FValidationData const& InValidationData) const
@@ -245,7 +249,9 @@ bool FOptimusConnectivityDataProviderProxy::IsValid(FValidationData const& InVal
 	{
 		return false;
 	}
-	if (SkeletalMeshObject == nullptr || AdjacencyBufferPerLod.Num() == 0)
+
+	FSkeletalMeshObject* SkeletalMeshObject = FSkeletalMeshDeformerHelpers::GetSkeletalMeshObject(SceneInterface, ComponentId);
+	if (!SkeletalMeshObject || AdjacencyBufferPerLod.Num() == 0)
 	{
 		return false;
 	}
@@ -255,6 +261,7 @@ bool FOptimusConnectivityDataProviderProxy::IsValid(FValidationData const& InVal
 
 void FOptimusConnectivityDataProviderProxy::AllocateResources(FRDGBuilder& GraphBuilder, FAllocationData const& InAllocationData)
 {
+	FSkeletalMeshObject* SkeletalMeshObject = FSkeletalMeshDeformerHelpers::GetSkeletalMeshObject(SceneInterface, ComponentId);
 	const int32 LodIndex = SkeletalMeshObject->GetLOD();
 	TArray<uint32> const& ConnectivityData = AdjacencyBufferPerLod[LodIndex];
 
@@ -266,6 +273,7 @@ void FOptimusConnectivityDataProviderProxy::AllocateResources(FRDGBuilder& Graph
 
 void FOptimusConnectivityDataProviderProxy::GatherDispatchData(FDispatchData const& InDispatchData)
 {
+	FSkeletalMeshObject* SkeletalMeshObject = FSkeletalMeshDeformerHelpers::GetSkeletalMeshObject(SceneInterface, ComponentId);
 	const int32 LodIndex = SkeletalMeshObject->GetLOD();
 	FSkeletalMeshRenderData const& SkeletalMeshRenderData = SkeletalMeshObject->GetSkeletalMeshRenderData();
 	FSkeletalMeshLODRenderData const* LodRenderData = &SkeletalMeshRenderData.LODRenderData[LodIndex];
diff --git a/Engine/Plugins/Animation/DeformerGraph/Source/OptimusCore/Private/DataInterfaces/OptimusDataInterfaceConnectivity.h b/Engine/Plugins/Animation/DeformerGraph/Source/OptimusCore/Private/DataInterfaces/OptimusDataInterfaceConnectivity.h
--- a/Engine/Plugins/Animation/DeformerGraph/Source/OptimusCore/Private/DataInterfaces/OptimusDataInterfaceConnectivity.h
+++ b/Engine/Plugins/Animation/DeformerGraph/Source/OptimusCore/Private/DataInterfaces/OptimusDataInterfaceConnectivity.h
@@ -4,6 +4,7 @@
 
 #include "OptimusComputeDataInterface.h"
 #include "ComputeFramework/ComputeDataProvider.h"
+#include "PrimitiveComponentId.h"
 #include "OptimusDataInterfaceConnectivity.generated.h"
 
 #define UE_API OPTIMUSCORE_API
@@ -11,6 +12,7 @@
 class FConnectivityDataInterfaceParameters;
 class FRDGBuffer;
 class FRDGBufferSRV;
+class FSceneInterface;
 class FSkeletalMeshObject;
 class USkinnedMeshComponent;
 
@@ -76,7 +78,8 @@ public:
 private:
 	using FParameters = FConnectivityDataInterfaceParameters;
 
-	FSkeletalMeshObject* SkeletalMeshObject = nullptr;
+	FSceneInterface* SceneInterface = nullptr;
+	FPrimitiveComponentId ComponentId;
 	TArray< TArray<uint32> > const& AdjacencyBufferPerLod;
 
 	FRDGBuffer* ConnectivityBuffer = nullptr;
diff --git a/Engine/Plugins/Animation/DeformerGraph/Source/OptimusCore/Private/DataInterfaces/OptimusDataInterfaceDuplicateVertices.cpp b/Engine/Plugins/Animation/DeformerGraph/Source/OptimusCore/Private/DataInterfaces/OptimusDataInterfaceDuplicateVertices.cpp
--- a/Engine/Plugins/Animation/DeformerGraph/Source/OptimusCore/Private/DataInterfaces/OptimusDataInterfaceDuplicateVertices.cpp
+++ b/Engine/Plugins/Animation/DeformerGraph/Source/OptimusCore/Private/DataInterfaces/OptimusDataInterfaceDuplicateVertices.cpp
@@ -2,9 +2,11 @@
 
 #include "OptimusDataInterfaceDuplicateVertices.h"
 
+#include "Components/SkinnedMeshComponent.h"
 #include "ComponentSources/OptimusSkinnedMeshComponentSource.h"
 #include "ComputeFramework/ComputeKernelPermutationVector.h"
 #include "ComputeFramework/ShaderParamTypeDefinition.h"
+#include "SkeletalMeshDeformerHelpers.h"
 #include "OptimusDataDomain.h"
 #include "Rendering/SkeletalMeshLODRenderData.h"
 #include "Rendering/SkeletalMeshRenderData.h"
@@ -119,7 +121,11 @@ FComputeDataProviderRenderProxy* UOptimusDuplicateVerticesDataProvider::GetRende
 
 FOptimusDuplicateVerticesDataProviderProxy::FOptimusDuplicateVerticesDataProviderProxy(USkinnedMeshComponent* SkinnedMeshComponent)
 {
-	SkeletalMeshObject = SkinnedMeshComponent != nullptr ? SkinnedMeshComponent->MeshObject : nullptr;
+	if (SkinnedMeshComponent != nullptr && SkinnedMeshComponent->IsRegistered())
+	{
+		SceneInterface = SkinnedMeshComponent->GetScene();
+		ComponentId = SkinnedMeshComponent->GetPrimitiveSceneId();
+	}
 }
 
 bool FOptimusDuplicateVerticesDataProviderProxy::IsValid(FValidationData const& InValidationData) const
@@ -128,7 +134,9 @@ bool FOptimusDuplicateVerticesDataProviderProxy::IsValid(FValidationData const&
 	{
 		return false;
 	}
-	if (SkeletalMeshObject == nullptr)
+
+	FSkeletalMeshObject* SkeletalMeshObject = FSkeletalMeshDeformerHelpers::GetSkeletalMeshObject(SceneInterface, ComponentId);
+	if (!SkeletalMeshObject)
 	{
 		return false;
 	}
@@ -156,6 +164,7 @@ struct FDuplicateVerticesDataInterfacePermutationIds
 
 void FOptimusDuplicateVerticesDataProviderProxy::GatherPermutations(FPermutationData& InOutPermutationData) const
 {
+	FSkeletalMeshObject* SkeletalMeshObject = FSkeletalMeshDeformerHelpers::GetSkeletalMeshObject(SceneInterface, ComponentId);
 	const int32 LodIndex = SkeletalMeshObject->GetLOD();
 	FSkeletalMeshRenderData const& SkeletalMeshRenderData = SkeletalMeshObject->GetSkeletalMeshRenderData();
 	FSkeletalMeshLODRenderData const* LodRenderData = &SkeletalMeshRenderData.LODRenderData[LodIndex];
@@ -175,6 +184,7 @@ void FOptimusDuplicateVerticesDataProviderProxy::GatherPermutations(FPermutation
 
 void FOptimusDuplicateVerticesDataProviderProxy::GatherDispatchData(FDispatchData const& InDispatchData)
 {
+	FSkeletalMeshObject* SkeletalMeshObject = FSkeletalMeshDeformerHelpers::GetSkeletalMeshObject(SceneInterface, ComponentId);
 	const int32 LodIndex = SkeletalMeshObject->GetLOD();
 	FSkeletalMeshRenderData const& SkeletalMeshRenderData = SkeletalMeshObject->GetSkeletalMeshRenderData();
 	FSkeletalMeshLODRenderData const* LodRenderData = &SkeletalMeshRenderData.LODRenderData[LodIndex];
diff --git a/Engine/Plugins/Animation/DeformerGraph/Source/OptimusCore/Private/DataInterfaces/OptimusDataInterfaceDuplicateVertices.h b/Engine/Plugins/Animation/DeformerGraph/Source/OptimusCore/Private/DataInterfaces/OptimusDataInterfaceDuplicateVertices.h
--- a/Engine/Plugins/Animation/DeformerGraph/Source/OptimusCore/Private/DataInterfaces/OptimusDataInterfaceDuplicateVertices.h
+++ b/Engine/Plugins/Animation/DeformerGraph/Source/OptimusCore/Private/DataInterfaces/OptimusDataInterfaceDuplicateVertices.h
@@ -4,12 +4,14 @@
 
 #include "OptimusComputeDataInterface.h"
 #include "ComputeFramework/ComputeDataProvider.h"
+#include "PrimitiveComponentId.h"
 
 #include "OptimusDataInterfaceDuplicateVertices.generated.h"
 
 #define UE_API OPTIMUSCORE_API
 
 class FDuplicateVerticesDataInterfaceParameters;
+class FSceneInterface;
 class FSkeletalMeshObject;
 class USkinnedMeshComponent;
 
@@ -72,7 +74,8 @@ public:
 private:
 	using FParameters = FDuplicateVerticesDataInterfaceParameters;
 
-	FSkeletalMeshObject* SkeletalMeshObject;
+	FSceneInterface* SceneInterface = nullptr;
+	FPrimitiveComponentId ComponentId;
 };
 
 #undef UE_API
diff --git a/Engine/Plugins/Animation/DeformerGraph/Source/OptimusCore/Private/DataInterfaces/OptimusDataInterfaceHalfEdge.cpp b/Engine/Plugins/Animation/DeformerGraph/Source/OptimusCore/Private/DataInterfaces/OptimusDataInterfaceHalfEdge.cpp
--- a/Engine/Plugins/Animation/DeformerGraph/Source/OptimusCore/Private/DataInterfaces/OptimusDataInterfaceHalfEdge.cpp
+++ b/Engine/Plugins/Animation/DeformerGraph/Source/OptimusCore/Private/DataInterfaces/OptimusDataInterfaceHalfEdge.cpp
@@ -3,6 +3,7 @@
 #include "OptimusDataInterfaceHalfEdge.h"
 
 #include "Components/SkinnedMeshComponent.h"
+#include "SkeletalMeshDeformerHelpers.h"
 #include "ComponentSources/OptimusSkinnedMeshComponentSource.h"
 #include "ComputeFramework/ShaderParamTypeDefinition.h"
 #include "OptimusDataDomain.h"
@@ -251,18 +252,23 @@ void UOptimusHalfEdgeDataProvider::ValidateSkinnedAsset()
 FOptimusHalfEdgeDataProviderProxy::FOptimusHalfEdgeDataProviderProxy(
 	USkinnedMeshComponent* InSkinnedMeshComponent, 
 	TArray<FHalfEdgeBuffers>& InOnDemandHalfEdgeBuffers)
-	: SkeletalMeshObject(InSkinnedMeshComponent != nullptr ? InSkinnedMeshComponent->MeshObject : nullptr)
-	, OnDemandHalfEdgeBuffers(InOnDemandHalfEdgeBuffers)
+	: OnDemandHalfEdgeBuffers(InOnDemandHalfEdgeBuffers)
 {
+	if (InSkinnedMeshComponent && InSkinnedMeshComponent->IsRegistered())
+	{
+		SceneInterface = InSkinnedMeshComponent->GetScene();
+		ComponentId = InSkinnedMeshComponent->GetPrimitiveSceneId();
+	}
 }
 
 bool FOptimusHalfEdgeDataProviderProxy::IsValid(FValidationData const& InValidationData) const
 {
-	if (InValidationData.ParameterStructSize != sizeof(FParameters))
+	FSkeletalMeshObject* SkeletalMeshObject = FSkeletalMeshDeformerHelpers::GetSkeletalMeshObject(SceneInterface, ComponentId);
+	if (!SkeletalMeshObject)
 	{
 		return false;
 	}
-	if (SkeletalMeshObject == nullptr)
+	if (InValidationData.ParameterStructSize != sizeof(FParameters))
 	{
 		return false;
 	}
@@ -289,6 +295,8 @@ bool FOptimusHalfEdgeDataProviderProxy::IsValid(FValidationData const& InValidat
 
 void FOptimusHalfEdgeDataProviderProxy::AllocateResources(FRDGBuilder& GraphBuilder, FAllocationData const& InAllocationData)
 {
+	FSkeletalMeshObject* SkeletalMeshObject = FSkeletalMeshDeformerHelpers::GetSkeletalMeshObject(SceneInterface, ComponentId);
+
 	FallbackSRV = GraphBuilder.CreateSRV(FRDGBufferSRVDesc(GSystemTextures.GetDefaultStructuredBuffer(GraphBuilder, sizeof(int32))));
 	
 	const int32 LodIndex = SkeletalMeshObject->GetLOD();
@@ -322,6 +330,8 @@ void FOptimusHalfEdgeDataProviderProxy::AllocateResources(FRDGBuilder& GraphBuil
 
 void FOptimusHalfEdgeDataProviderProxy::GatherDispatchData(FDispatchData const& InDispatchData)
 {
+	FSkeletalMeshObject* SkeletalMeshObject = FSkeletalMeshDeformerHelpers::GetSkeletalMeshObject(SceneInterface, ComponentId);
+
 	const int32 LodIndex = SkeletalMeshObject->GetLOD();
 	FSkeletalMeshRenderData const& SkeletalMeshRenderData = SkeletalMeshObject->GetSkeletalMeshRenderData();
 	FSkeletalMeshLODRenderData const* LodRenderData = &SkeletalMeshRenderData.LODRenderData[LodIndex];
diff --git a/Engine/Plugins/Animation/DeformerGraph/Source/OptimusCore/Private/DataInterfaces/OptimusDataInterfaceHalfEdge.h b/Engine/Plugins/Animation/DeformerGraph/Source/OptimusCore/Private/DataInterfaces/OptimusDataInterfaceHalfEdge.h
--- a/Engine/Plugins/Animation/DeformerGraph/Source/OptimusCore/Private/DataInterfaces/OptimusDataInterfaceHalfEdge.h
+++ b/Engine/Plugins/Animation/DeformerGraph/Source/OptimusCore/Private/DataInterfaces/OptimusDataInterfaceHalfEdge.h
@@ -7,6 +7,7 @@
 #include "OptimusComputeDataInterface.h"
 #include "ComputeFramework/ComputeDataProvider.h"
 #include "RenderCommandFence.h"
+#include "PrimitiveComponentId.h"
 #include "OptimusDataInterfaceHalfEdge.generated.h"
 
 #define UE_API OPTIMUSCORE_API
@@ -15,6 +16,7 @@ class FHalfEdgeDataInterfaceParameters;
 class FHalfEdgeBuffers;
 class FRDGBuffer;
 class FRDGBufferSRV;
+class FSceneInterface;
 class FSkeletalMeshObject;
 class USkinnedMeshComponent;
 class USkinnedAsset;
@@ -102,8 +104,9 @@ public:
 private:
 	using FParameters = FHalfEdgeDataInterfaceParameters;
 
-	FSkeletalMeshObject* SkeletalMeshObject = nullptr;
-	
+	FSceneInterface* SceneInterface = nullptr;
+	FPrimitiveComponentId ComponentId;
+
 	TArray<FHalfEdgeBuffers>& OnDemandHalfEdgeBuffers;
 	
 	bool bUseBufferFromRenderData = false;
diff --git a/Engine/Plugins/Animation/DeformerGraph/Source/OptimusCore/Private/DataInterfaces/OptimusDataInterfaceMorphTarget.cpp b/Engine/Plugins/Animation/DeformerGraph/Source/OptimusCore/Private/DataInterfaces/OptimusDataInterfaceMorphTarget.cpp
--- a/Engine/Plugins/Animation/DeformerGraph/Source/OptimusCore/Private/DataInterfaces/OptimusDataInterfaceMorphTarget.cpp
+++ b/Engine/Plugins/Animation/DeformerGraph/Source/OptimusCore/Private/DataInterfaces/OptimusDataInterfaceMorphTarget.cpp
@@ -106,9 +106,10 @@ FComputeDataProviderRenderProxy* UOptimusMorphTargetDataProvider::GetRenderProxy
 
 FOptimusMorphTargetDataProviderProxy::FOptimusMorphTargetDataProviderProxy(USkinnedMeshComponent* SkinnedMeshComponent)
 {
-	if (SkinnedMeshComponent != nullptr && SkinnedMeshComponent->GetScene() != nullptr)
+	if (SkinnedMeshComponent != nullptr && SkinnedMeshComponent->IsRegistered())
 	{
-		SkeletalMeshObject = SkinnedMeshComponent->MeshObject;
+		SceneInterface = SkinnedMeshComponent->GetScene();
+		ComponentId = SkinnedMeshComponent->GetPrimitiveSceneId();
 		FrameNumber = SkinnedMeshComponent->GetScene()->GetFrameNumber() + 1; // +1 matches the logic for FrameNumberToPrepare in FSkeletalMeshObjectGPUSkin::Update()
 	}
 }
@@ -119,24 +120,13 @@ bool FOptimusMorphTargetDataProviderProxy::IsValid(FValidationData const& InVali
 	{
 		return false;
 	}
-	if (SkeletalMeshObject == nullptr)
-	{
-		return false;
-	}
-
-	const int32 LodIndex = SkeletalMeshObject->GetLOD();
-	FSkeletalMeshRenderData const& SkeletalMeshRenderData = SkeletalMeshObject->GetSkeletalMeshRenderData();
-	FSkeletalMeshLODRenderData const* LodRenderData = &SkeletalMeshRenderData.LODRenderData[LodIndex];
 
-	// this mismatch can happen during map load for one frame sometimes, when predicted lod != actual lod
-	// TODO: avoid the mismatch in the first place, use actual lod to query NumInvocations
-	if (LodRenderData->RenderSections.Num() != InValidationData.NumInvocations)
+	FSkeletalMeshObject* SkeletalMeshObject = FSkeletalMeshDeformerHelpers::GetSkeletalMeshObject(SceneInterface, ComponentId);
+	if (!SkeletalMeshObject)
 	{
 		return false;
 	}
 
-	
-
 	return true;
 }
 
@@ -156,13 +146,15 @@ struct FMorphTargetDataInterfacePermutationIds
 
 void FOptimusMorphTargetDataProviderProxy::GatherPermutations(FPermutationData& InOutPermutationData) const
 {
+	FSkeletalMeshObject* SkeletalMeshObject = FSkeletalMeshDeformerHelpers::GetSkeletalMeshObject(SceneInterface, ComponentId);
+
 	const int32 LodIndex = SkeletalMeshObject->GetLOD();
 
 	FMorphTargetDataInterfacePermutationIds PermutationIds(InOutPermutationData.PermutationVector);
 	for (int32 InvocationIndex = 0; InvocationIndex < InOutPermutationData.NumInvocations; ++InvocationIndex)
 	{
-		const bool bPreviousFrame = false;
-		FRHIShaderResourceView* MorphBufferSRV = FSkeletalMeshDeformerHelpers::GetMorphTargetBufferForReading(SkeletalMeshObject, LodIndex, InvocationIndex, FrameNumber, bPreviousFrame);
+		constexpr bool bPreviousFrame = false;
+		FRHIShaderResourceView* MorphBufferSRV = FSkeletalMeshDeformerHelpers::GetMorphTargetBufferForReading(SkeletalMeshObject, LodIndex, bPreviousFrame);
 		const bool bValidMorph = MorphBufferSRV != nullptr;
 
 		InOutPermutationData.PermutationIds[InvocationIndex] |= (bValidMorph ? PermutationIds.EnableDeformerMorphTarget : 0);
@@ -171,6 +163,8 @@ void FOptimusMorphTargetDataProviderProxy::GatherPermutations(FPermutationData&
 
 void FOptimusMorphTargetDataProviderProxy::GatherDispatchData(FDispatchData const& InDispatchData)
 {
+	FSkeletalMeshObject* SkeletalMeshObject = FSkeletalMeshDeformerHelpers::GetSkeletalMeshObject(SceneInterface, ComponentId);
+
 	FRHIShaderResourceView* NullSRVBinding = GWhiteVertexBufferWithSRV->ShaderResourceViewRHI.GetReference();
 
 	const int32 LodIndex = SkeletalMeshObject->GetLOD();
@@ -180,8 +174,8 @@ void FOptimusMorphTargetDataProviderProxy::GatherDispatchData(FDispatchData cons
 	const TStridedView<FParameters> ParameterArray = MakeStridedParameterView<FParameters>(InDispatchData);
 	for (int32 InvocationIndex = 0; InvocationIndex < ParameterArray.Num(); ++InvocationIndex)
 	{
-		const bool bPreviousFrame = false;
-		FRHIShaderResourceView* MorphBufferSRV = FSkeletalMeshDeformerHelpers::GetMorphTargetBufferForReading(SkeletalMeshObject, LodIndex, InvocationIndex, FrameNumber, bPreviousFrame);
+		constexpr bool bPreviousFrame = false;
+		FRHIShaderResourceView* MorphBufferSRV = FSkeletalMeshDeformerHelpers::GetMorphTargetBufferForReading(SkeletalMeshObject, LodIndex, bPreviousFrame);
 
 		FParameters& Parameters = ParameterArray[InvocationIndex];
 		Parameters.NumVertices = LodRenderData->GetNumVertices();
diff --git a/Engine/Plugins/Animation/DeformerGraph/Source/OptimusCore/Private/DataInterfaces/OptimusDataInterfaceMorphTarget.h b/Engine/Plugins/Animation/DeformerGraph/Source/OptimusCore/Private/DataInterfaces/OptimusDataInterfaceMorphTarget.h
--- a/Engine/Plugins/Animation/DeformerGraph/Source/OptimusCore/Private/DataInterfaces/OptimusDataInterfaceMorphTarget.h
+++ b/Engine/Plugins/Animation/DeformerGraph/Source/OptimusCore/Private/DataInterfaces/OptimusDataInterfaceMorphTarget.h
@@ -4,12 +4,14 @@
 
 #include "OptimusComputeDataInterface.h"
 #include "ComputeFramework/ComputeDataProvider.h"
+#include "PrimitiveComponentId.h"
 
 #include "OptimusDataInterfaceMorphTarget.generated.h"
 
 #define UE_API OPTIMUSCORE_API
 
 class FMorphTargetDataInterfaceParameters;
+class FSceneInterface;
 class FSkeletalMeshObject;
 class USkinnedMeshComponent;
 
@@ -28,7 +30,7 @@ public:
 	
 	//~ Begin UComputeDataInterface Interface
 	TCHAR const* GetClassName() const override { return TEXT("MorphTarget"); }
-	bool CanSupportUnifiedDispatch() const override {return false;}
+	bool CanSupportUnifiedDispatch() const override {return true;}
 	UE_API void GetSupportedInputs(TArray<FShaderFunctionDefinition>& OutFunctions) const override;
 	UE_API void GetShaderParameters(TCHAR const* UID, FShaderParametersMetadataBuilder& InOutBuilder, FShaderParametersMetadataAllocations& InOutAllocations) const override;
 	UE_API void GetPermutations(FComputeKernelPermutationVector& OutPermutationVector) const override;
@@ -70,7 +72,8 @@ public:
 private:
 	using FParameters = FMorphTargetDataInterfaceParameters;
 
-	FSkeletalMeshObject* SkeletalMeshObject = nullptr;
+	FSceneInterface* SceneInterface = nullptr;
+	FPrimitiveComponentId ComponentId;
 	uint32 FrameNumber = 0;
 };
 
diff --git a/Engine/Plugins/Animation/DeformerGraph/Source/OptimusCore/Private/DataInterfaces/OptimusDataInterfaceSkeleton.cpp b/Engine/Plugins/Animation/DeformerGraph/Source/OptimusCore/Private/DataInterfaces/OptimusDataInterfaceSkeleton.cpp
--- a/Engine/Plugins/Animation/DeformerGraph/Source/OptimusCore/Private/DataInterfaces/OptimusDataInterfaceSkeleton.cpp
+++ b/Engine/Plugins/Animation/DeformerGraph/Source/OptimusCore/Private/DataInterfaces/OptimusDataInterfaceSkeleton.cpp
@@ -2,6 +2,7 @@
 
 #include "OptimusDataInterfaceSkeleton.h"
 
+#include "Components/SkinnedMeshComponent.h"
 #include "ComputeFramework/ComputeKernelPermutationVector.h"
 #include "ComputeFramework/ShaderParamTypeDefinition.h"
 #include "OptimusDataDomain.h"
@@ -130,7 +131,11 @@ FComputeDataProviderRenderProxy* UOptimusSkeletonDataProvider::GetRenderProxy()
 
 FOptimusSkeletonDataProviderProxy::FOptimusSkeletonDataProviderProxy(USkinnedMeshComponent* SkinnedMeshComponent)
 {
-	SkeletalMeshObject = SkinnedMeshComponent != nullptr ? SkinnedMeshComponent->MeshObject : nullptr;
+	if (SkinnedMeshComponent != nullptr && SkinnedMeshComponent->IsRegistered())
+	{
+		SceneInterface = SkinnedMeshComponent->GetScene();
+		ComponentId = SkinnedMeshComponent->GetPrimitiveSceneId();
+	}
 	BoneRevisionNumber = SkinnedMeshComponent != nullptr ? SkinnedMeshComponent->GetBoneTransformRevisionNumber() : 0;
 }
 
@@ -140,7 +145,9 @@ bool FOptimusSkeletonDataProviderProxy::IsValid(FValidationData const& InValidat
 	{
 		return false;
 	}
-	if (SkeletalMeshObject == nullptr)
+
+	FSkeletalMeshObject* SkeletalMeshObject = FSkeletalMeshDeformerHelpers::GetSkeletalMeshObject(SceneInterface, ComponentId);
+	if (!SkeletalMeshObject)
 	{
 		return false;
 	}
@@ -186,6 +193,8 @@ struct FSkeletonDataInterfacePermutationIds
 
 void FOptimusSkeletonDataProviderProxy::GatherPermutations(FPermutationData& InOutPermutationData) const
 {
+	FSkeletalMeshObject* SkeletalMeshObject = FSkeletalMeshDeformerHelpers::GetSkeletalMeshObject(SceneInterface, ComponentId);
+
 	const int32 LodIndex = SkeletalMeshObject->GetLOD();
 	FSkeletalMeshRenderData const& SkeletalMeshRenderData = SkeletalMeshObject->GetSkeletalMeshRenderData();
 	FSkeletalMeshLODRenderData const* LodRenderData = &SkeletalMeshRenderData.LODRenderData[LodIndex];
@@ -216,6 +225,8 @@ void FOptimusSkeletonDataProviderProxy::GatherPermutations(FPermutationData& InO
 
 void FOptimusSkeletonDataProviderProxy::GatherDispatchData(FDispatchData const& InDispatchData)
 {
+	FSkeletalMeshObject* SkeletalMeshObject = FSkeletalMeshDeformerHelpers::GetSkeletalMeshObject(SceneInterface, ComponentId);
+
 	const int32 LodIndex = SkeletalMeshObject->GetLOD();
 	FSkeletalMeshRenderData const& SkeletalMeshRenderData = SkeletalMeshObject->GetSkeletalMeshRenderData();
 	FSkeletalMeshLODRenderData const* LodRenderData = &SkeletalMeshRenderData.LODRenderData[LodIndex];
diff --git a/Engine/Plugins/Animation/DeformerGraph/Source/OptimusCore/Private/DataInterfaces/OptimusDataInterfaceSkeleton.h b/Engine/Plugins/Animation/DeformerGraph/Source/OptimusCore/Private/DataInterfaces/OptimusDataInterfaceSkeleton.h
--- a/Engine/Plugins/Animation/DeformerGraph/Source/OptimusCore/Private/DataInterfaces/OptimusDataInterfaceSkeleton.h
+++ b/Engine/Plugins/Animation/DeformerGraph/Source/OptimusCore/Private/DataInterfaces/OptimusDataInterfaceSkeleton.h
@@ -4,11 +4,13 @@
 
 #include "OptimusComputeDataInterface.h"
 #include "ComputeFramework/ComputeDataProvider.h"
+#include "PrimitiveComponentId.h"
 
 #include "OptimusDataInterfaceSkeleton.generated.h"
 
 #define UE_API OPTIMUSCORE_API
 
+class FSceneInterface;
 class FSkeletalMeshObject;
 class FSkeletonDataInterfaceParameters;
 class USkinnedMeshComponent;
@@ -71,7 +73,8 @@ public:
 private:
 	using FParameters = FSkeletonDataInterfaceParameters;
 
-	FSkeletalMeshObject* SkeletalMeshObject;
+	FSceneInterface* SceneInterface = nullptr;
+	FPrimitiveComponentId ComponentId;
 	uint32 BoneRevisionNumber = 0;
 };
 
diff --git a/Engine/Plugins/Animation/DeformerGraph/Source/OptimusCore/Private/DataInterfaces/OptimusDataInterfaceSkinWeightsAsVertexMask.cpp b/Engine/Plugins/Animation/DeformerGraph/Source/OptimusCore/Private/DataInterfaces/OptimusDataInterfaceSkinWeightsAsVertexMask.cpp
--- a/Engine/Plugins/Animation/DeformerGraph/Source/OptimusCore/Private/DataInterfaces/OptimusDataInterfaceSkinWeightsAsVertexMask.cpp
+++ b/Engine/Plugins/Animation/DeformerGraph/Source/OptimusCore/Private/DataInterfaces/OptimusDataInterfaceSkinWeightsAsVertexMask.cpp
@@ -426,7 +426,11 @@ FComputeDataProviderRenderProxy* UOptimusSkinWeightsAsVertexMaskDataProvider::Ge
 			}
 			
 			
-			Proxy->SkeletalMeshObject = SkeletalMeshObject;
+			if (SkeletalMesh.IsValid() && SkeletalMesh->IsRegistered())
+			{
+				Proxy->SceneInterface = SkeletalMesh->GetScene();
+				Proxy->ComponentId = SkeletalMesh->GetPrimitiveSceneId();
+			}
 			Proxy->SkinWeightProfile = SkinWeightProfile;
 			Proxy->BoneIsSelectedPerSectionPerLod = CachedBoneIsSelectedPerSectionPerLod;
 		}
@@ -439,11 +443,12 @@ FComputeDataProviderRenderProxy* UOptimusSkinWeightsAsVertexMaskDataProvider::Ge
 
 bool FOptimusSkinWeightsAsVertexMaskDataProviderProxy::IsValid(FValidationData const& InValidationData) const
 {
-	if (InValidationData.ParameterStructSize != sizeof(FParameters))
+	FSkeletalMeshObject* SkeletalMeshObject = FSkeletalMeshDeformerHelpers::GetSkeletalMeshObject(SceneInterface, ComponentId);
+	if (!SkeletalMeshObject)
 	{
 		return false;
 	}
-	if (SkeletalMeshObject == nullptr)
+	if (InValidationData.ParameterStructSize != sizeof(FParameters))
 	{
 		return false;
 	}
@@ -451,8 +456,8 @@ bool FOptimusSkinWeightsAsVertexMaskDataProviderProxy::IsValid(FValidationData c
 	FSkeletalMeshRenderData const& SkeletalMeshRenderData = SkeletalMeshObject->GetSkeletalMeshRenderData();
 	FSkeletalMeshLODRenderData const* LodRenderData = &SkeletalMeshRenderData.LODRenderData[LodIndex];
 
-	// this mismatch can happen during map load for one frame sometimes, when predicted lod != actual lod
-	// TODO: avoid the mismatch in the first place, use actual lod to query NumInvocations
+	// This data interface does not support unified dispatch (because bone maps can be per-section) 
+	// so the invocation count has to match the render sections
 	if (LodRenderData->RenderSections.Num() != InValidationData.NumInvocations)
 	{
 		return false;
@@ -480,6 +485,7 @@ bool FOptimusSkinWeightsAsVertexMaskDataProviderProxy::IsValid(FValidationData c
 
 void FOptimusSkinWeightsAsVertexMaskDataProviderProxy::AllocateResources(FRDGBuilder& GraphBuilder, FAllocationData const& InAllocationData)
 {
+	FSkeletalMeshObject* SkeletalMeshObject = FSkeletalMeshDeformerHelpers::GetSkeletalMeshObject(SceneInterface, ComponentId);
 	const int32 LodIndex = SkeletalMeshObject->GetLOD();
 	FSkeletalMeshRenderData const& SkeletalMeshRenderData = SkeletalMeshObject->GetSkeletalMeshRenderData();
 	FSkeletalMeshLODRenderData const* LodRenderData = &SkeletalMeshRenderData.LODRenderData[LodIndex];
@@ -544,6 +550,7 @@ struct FSkinWeightsAsVertexMaskDataInterfacePermutationIds
 
 void FOptimusSkinWeightsAsVertexMaskDataProviderProxy::GatherPermutations(FPermutationData& InOutPermutationData) const
 {
+	FSkeletalMeshObject* SkeletalMeshObject = FSkeletalMeshDeformerHelpers::GetSkeletalMeshObject(SceneInterface, ComponentId);
 	const int32 LodIndex = SkeletalMeshObject->GetLOD();
 	FSkeletalMeshRenderData const& SkeletalMeshRenderData = SkeletalMeshObject->GetSkeletalMeshRenderData();
 	FSkeletalMeshLODRenderData const* LodRenderData = &SkeletalMeshRenderData.LODRenderData[LodIndex];
@@ -578,6 +585,7 @@ void FOptimusSkinWeightsAsVertexMaskDataProviderProxy::GatherPermutations(FPermu
 
 void FOptimusSkinWeightsAsVertexMaskDataProviderProxy::GatherDispatchData(FDispatchData const& InDispatchData)
 {
+	FSkeletalMeshObject* SkeletalMeshObject = FSkeletalMeshDeformerHelpers::GetSkeletalMeshObject(SceneInterface, ComponentId);
 	const int32 LodIndex = SkeletalMeshObject->GetLOD();
 	FSkeletalMeshRenderData const& SkeletalMeshRenderData = SkeletalMeshObject->GetSkeletalMeshRenderData();
 	FSkeletalMeshLODRenderData const* LodRenderData = &SkeletalMeshRenderData.LODRenderData[LodIndex];
diff --git a/Engine/Plugins/Animation/DeformerGraph/Source/OptimusCore/Private/DataInterfaces/OptimusDataInterfaceSkinWeightsAsVertexMask.h b/Engine/Plugins/Animation/DeformerGraph/Source/OptimusCore/Private/DataInterfaces/OptimusDataInterfaceSkinWeightsAsVertexMask.h
--- a/Engine/Plugins/Animation/DeformerGraph/Source/OptimusCore/Private/DataInterfaces/OptimusDataInterfaceSkinWeightsAsVertexMask.h
+++ b/Engine/Plugins/Animation/DeformerGraph/Source/OptimusCore/Private/DataInterfaces/OptimusDataInterfaceSkinWeightsAsVertexMask.h
@@ -7,10 +7,12 @@
 #include "OptimusDataType.h"
 #include "ComputeFramework/ComputeDataProvider.h"
 
+#include "PrimitiveComponentId.h"
 #include "OptimusDataInterfaceSkinWeightsAsVertexMask.generated.h"
 
 #define UE_API OPTIMUSCORE_API
 
+class FSceneInterface;
 class FSkinWeightsAsVertexMaskDataInterfaceParameters;
 class FSkeletalMeshObject;
 class USkeletalMeshComponent;
@@ -137,8 +139,9 @@ public:
 
 	using FParameters = FSkinWeightsAsVertexMaskDataInterfaceParameters;
 
-	FSkeletalMeshObject* SkeletalMeshObject = nullptr;
-	
+	FSceneInterface* SceneInterface = nullptr;
+	FPrimitiveComponentId ComponentId;
+
 	FName SkinWeightProfile = NAME_None;
 
 	TArray<TArray<TArray<uint32>>> BoneIsSelectedPerSectionPerLod;
diff --git a/Engine/Plugins/Animation/DeformerGraph/Source/OptimusCore/Private/DataInterfaces/OptimusDataInterfaceSkinnedMesh.cpp b/Engine/Plugins/Animation/DeformerGraph/Source/OptimusCore/Private/DataInterfaces/OptimusDataInterfaceSkinnedMesh.cpp
--- a/Engine/Plugins/Animation/DeformerGraph/Source/OptimusCore/Private/DataInterfaces/OptimusDataInterfaceSkinnedMesh.cpp
+++ b/Engine/Plugins/Animation/DeformerGraph/Source/OptimusCore/Private/DataInterfaces/OptimusDataInterfaceSkinnedMesh.cpp
@@ -2,6 +2,8 @@
 
 #include "OptimusDataInterfaceSkinnedMesh.h"
 
+#include "Components/SkinnedMeshComponent.h"
+#include "SkeletalMeshDeformerHelpers.h"
 #include "OptimusDataDomain.h"
 #include "RenderGraphBuilder.h"
 #include "ComponentSources/OptimusSkinnedMeshComponentSource.h"
@@ -149,16 +151,21 @@ FComputeDataProviderRenderProxy* UOptimusSkinnedMeshDataProvider::GetRenderProxy
 
 FOptimusSkinnedMeshDataProviderProxy::FOptimusSkinnedMeshDataProviderProxy(USkinnedMeshComponent* SkinnedMeshComponent)
 {
-	SkeletalMeshObject = SkinnedMeshComponent != nullptr ? SkinnedMeshComponent->MeshObject : nullptr;
+	if (SkinnedMeshComponent && SkinnedMeshComponent->IsRegistered())
+	{
+		SceneInterface = SkinnedMeshComponent->GetScene();
+		ComponentId = SkinnedMeshComponent->GetPrimitiveSceneId();
+	}
 }
 
 bool FOptimusSkinnedMeshDataProviderProxy::IsValid(FValidationData const& InValidationData) const
 {
-	if (InValidationData.ParameterStructSize != sizeof(FParameters))
+	FSkeletalMeshObject* SkeletalMeshObject = FSkeletalMeshDeformerHelpers::GetSkeletalMeshObject(SceneInterface, ComponentId);
+	if (!SkeletalMeshObject)
 	{
 		return false;
 	}
-	if (SkeletalMeshObject == nullptr)
+	if (InValidationData.ParameterStructSize != sizeof(FParameters))
 	{
 		return false;
 	}
@@ -168,6 +175,8 @@ bool FOptimusSkinnedMeshDataProviderProxy::IsValid(FValidationData const& InVali
 
 void FOptimusSkinnedMeshDataProviderProxy::AllocateResources(FRDGBuilder& GraphBuilder, FAllocationData const& InAllocationData)
 {
+	FSkeletalMeshObject* SkeletalMeshObject = FSkeletalMeshDeformerHelpers::GetSkeletalMeshObject(SceneInterface, ComponentId);
+
 	const int32 LodIndex = SkeletalMeshObject->GetLOD();
 	FSkeletalMeshRenderData& SkeletalMeshRenderData = SkeletalMeshObject->GetSkeletalMeshRenderData();
 	FSkeletalMeshLODRenderData* LodRenderData = &SkeletalMeshRenderData.LODRenderData[LodIndex];
@@ -184,6 +193,8 @@ void FOptimusSkinnedMeshDataProviderProxy::AllocateResources(FRDGBuilder& GraphB
 
 void FOptimusSkinnedMeshDataProviderProxy::GatherDispatchData(FDispatchData const& InDispatchData)
 {
+	FSkeletalMeshObject* SkeletalMeshObject = FSkeletalMeshDeformerHelpers::GetSkeletalMeshObject(SceneInterface, ComponentId);
+
 	const int32 LodIndex = SkeletalMeshObject->GetLOD();
 	FSkeletalMeshRenderData const& SkeletalMeshRenderData = SkeletalMeshObject->GetSkeletalMeshRenderData();
 	FSkeletalMeshLODRenderData const* LodRenderData = &SkeletalMeshRenderData.LODRenderData[LodIndex];
diff --git a/Engine/Plugins/Animation/DeformerGraph/Source/OptimusCore/Private/DataInterfaces/OptimusDataInterfaceSkinnedMesh.h b/Engine/Plugins/Animation/DeformerGraph/Source/OptimusCore/Private/DataInterfaces/OptimusDataInterfaceSkinnedMesh.h
--- a/Engine/Plugins/Animation/DeformerGraph/Source/OptimusCore/Private/DataInterfaces/OptimusDataInterfaceSkinnedMesh.h
+++ b/Engine/Plugins/Animation/DeformerGraph/Source/OptimusCore/Private/DataInterfaces/OptimusDataInterfaceSkinnedMesh.h
@@ -6,10 +6,12 @@
 #include "ComputeFramework/ComputeDataProvider.h"
 #include "RHIFwd.h"
 
+#include "PrimitiveComponentId.h"
 #include "OptimusDataInterfaceSkinnedMesh.generated.h"
 
 #define UE_API OPTIMUSCORE_API
 
+class FSceneInterface;
 class FSkeletalMeshObject;
 class FSkinnedMeshDataInterfaceParameters;
 class USkinnedMeshComponent;
@@ -70,8 +72,9 @@ public:
 private:
 	using FParameters = FSkinnedMeshDataInterfaceParameters;
 
-	FSkeletalMeshObject* SkeletalMeshObject;
-	
+	FSceneInterface* SceneInterface = nullptr;
+	FPrimitiveComponentId ComponentId;
+
 	FShaderResourceViewRHIRef IndexBufferSRV;
 	FShaderResourceViewRHIRef TangentsSRV;
 	FShaderResourceViewRHIRef TexCoordsSRV;
diff --git a/Engine/Plugins/Animation/DeformerGraph/Source/OptimusCore/Private/DataInterfaces/OptimusDataInterfaceSkinnedMeshRead.cpp b/Engine/Plugins/Animation/DeformerGraph/Source/OptimusCore/Private/DataInterfaces/OptimusDataInterfaceSkinnedMeshRead.cpp
--- a/Engine/Plugins/Animation/DeformerGraph/Source/OptimusCore/Private/DataInterfaces/OptimusDataInterfaceSkinnedMeshRead.cpp
+++ b/Engine/Plugins/Animation/DeformerGraph/Source/OptimusCore/Private/DataInterfaces/OptimusDataInterfaceSkinnedMeshRead.cpp
@@ -14,6 +14,7 @@
 #include "SkeletalRenderPublic.h"
 #include "DataDrivenShaderPlatformInfo.h"
 #include "OptimusDeformerInstance.h"
+#include "OptimusSkinnedMeshTrackingExtension.h"
 #include "Animation/MeshDeformerInstance.h"
 #include "ComputeFramework/ComputeKernelPermutationSet.h"
 #include "ComputeFramework/ComputeKernelPermutationVector.h"
@@ -165,10 +166,17 @@ UComputeDataProvider* UOptimusSkinnedMeshReadDataInterface::CreateDataProvider(T
 
 FComputeDataProviderRenderProxy* UOptimusSkinnedMeshReadDataProvider::GetRenderProxy()
 {
-	return new FOptimusSkinnedMeshReadDataProviderProxy(
-		SkinnedMesh.Get(),
-		InputMask,
-		WeakDeformerInstance.IsValid()? WeakDeformerInstance->OutputBuffersFromPreviousInstances : EMeshDeformerOutputBuffer::None, &LastLodIndexCachedByRenderProxy);
+	bool bIsPrimary = true;
+	FName ExecGroup;
+
+	if (WeakDeformerInstance.IsValid())
+	{
+		UOptimusDeformerInstance* Instance = WeakDeformerInstance.Get();
+		bIsPrimary = (SkinnedMesh.Get() == Instance->GetMeshComponent());
+		ExecGroup = Instance->CurrentExecutionGroupName;
+	}
+
+	return new FOptimusSkinnedMeshReadDataProviderProxy(SkinnedMesh.Get(), InputMask, bIsPrimary, ExecGroup);
 }
 
 void UOptimusSkinnedMeshReadDataProvider::SetDeformerInstance(UOptimusDeformerInstance* InInstance)
@@ -176,14 +184,16 @@ void UOptimusSkinnedMeshReadDataProvider::SetDeformerInstance(UOptimusDeformerIn
 	WeakDeformerInstance = InInstance;
 }
 
-FOptimusSkinnedMeshReadDataProviderProxy::FOptimusSkinnedMeshReadDataProviderProxy(USkinnedMeshComponent* InSkinnedMeshComponent, uint64 InInputMask, EMeshDeformerOutputBuffer InOutputBuffersFromPreviousInstances, int32* InLastLodIndexPtr) 
+FOptimusSkinnedMeshReadDataProviderProxy::FOptimusSkinnedMeshReadDataProviderProxy(USkinnedMeshComponent* InSkinnedMeshComponent, uint64 InInputMask, bool bInIsPrimaryComponent, FName InReaderExecutionGroup)
 {
-	SkeletalMeshObject = InSkinnedMeshComponent != nullptr ? InSkinnedMeshComponent->MeshObject : nullptr;
+	if (InSkinnedMeshComponent && InSkinnedMeshComponent->IsRegistered())
+	{
+		SceneInterface = InSkinnedMeshComponent->GetScene();
+		ComponentId = InSkinnedMeshComponent->GetPrimitiveSceneId();
+	}
 	InputMask = InInputMask;
-	OutputBuffersFromPreviousInstances = InOutputBuffersFromPreviousInstances;
-	LastLodIndexPtr = InLastLodIndexPtr;
-
-	check(LastLodIndexPtr != nullptr);
+	bIsPrimaryComponent = bInIsPrimaryComponent;
+	ReaderExecutionGroup = InReaderExecutionGroup;
 }
 
 bool FOptimusSkinnedMeshReadDataProviderProxy::IsValid(FValidationData const& InValidationData) const
@@ -192,11 +202,9 @@ bool FOptimusSkinnedMeshReadDataProviderProxy::IsValid(FValidationData const& In
 	{
 		return false;
 	}
-	if (SkeletalMeshObject == nullptr)
-	{
-		return false;
-	}
-	if (SkeletalMeshObject->IsCPUSkinned())
+
+	FSkeletalMeshObject* SkeletalMeshObject = FSkeletalMeshDeformerHelpers::GetSkeletalMeshObject(SceneInterface, ComponentId);
+	if (!SkeletalMeshObject)
 	{
 		return false;
 	}
@@ -206,34 +214,14 @@ bool FOptimusSkinnedMeshReadDataProviderProxy::IsValid(FValidationData const& In
 		return false;
 	}
 
-	const int32 LodIndex = SkeletalMeshObject->GetLOD();
-	FSkeletalMeshRenderData const& SkeletalMeshRenderData = SkeletalMeshObject->GetSkeletalMeshRenderData();
-	FSkeletalMeshLODRenderData const* LodRenderData = &SkeletalMeshRenderData.LODRenderData[LodIndex];
-
-	// this mismatch can happen during map load for one frame sometimes, when predicted lod != actual lod
-	// TODO: avoid the mismatch in the first place, use actual lod to query NumInvocations
-	if (LodRenderData->RenderSections.Num() != InValidationData.NumInvocations)
-	{
-		return false;
-	}
-	
 	return true;
 }
 
 void FOptimusSkinnedMeshReadDataProviderProxy::AllocateResources(FRDGBuilder& GraphBuilder, FAllocationData const& InAllocationData)
 {
-	// Allocate required buffers
-	const int32 LodIndex = SkeletalMeshObject->GetLOD();
-	
-	// Avoid using previous position buffer from when this LOD was last active to compute motion vectors.
-	// The position delta between that previous position (could be from any time ago) and the current position can be any crazy value that is not meaningful.
-	bool bInvalidatePreviousPosition = false;
-	if (LodIndex != *LastLodIndexPtr)
-	{
-		bInvalidatePreviousPosition = true;
-		*LastLodIndexPtr = LodIndex;
-	}
+	FSkeletalMeshObject* SkeletalMeshObject = FSkeletalMeshDeformerHelpers::GetSkeletalMeshObject(SceneInterface, ComponentId);
 
+	const int32 LodIndex = SkeletalMeshObject->GetLOD();
 	FSkeletalMeshRenderData const& SkeletalMeshRenderData = SkeletalMeshObject->GetSkeletalMeshRenderData();
 	FSkeletalMeshLODRenderData const* LodRenderData = &SkeletalMeshRenderData.LODRenderData[LodIndex];
 	
@@ -242,9 +230,14 @@ void FOptimusSkinnedMeshReadDataProviderProxy::AllocateResources(FRDGBuilder& Gr
 	// Color SRV is either created by default or not available at all
 	ColorsSRV = LodRenderData->StaticVertexBuffers.ColorVertexBuffer.GetColorComponentsSRV();
 
-	if (InputMask & static_cast<uint64>(ESkinnedMeshReadDataInterfaceInputSelectorMask::Position) )
+	// Query which deformer output buffers are readable for this component
+	const EMeshDeformerOutputBuffer ReadableBuffers = FOptimusSkinnedMeshTrackingExtension::GetReadableDeformedBuffers(
+		SceneInterface, ComponentId, ReaderExecutionGroup, bIsPrimaryComponent);
+
+	if ((InputMask & static_cast<uint64>(ESkinnedMeshReadDataInterfaceInputSelectorMask::Position)) &&
+		EnumHasAnyFlags(ReadableBuffers, EMeshDeformerOutputBuffer::SkinnedMeshPosition))
 	{
-		FRDGBuffer* PositionBuffer = FSkeletalMeshDeformerHelpers::AllocateVertexFactoryPositionBuffer(GraphBuilder, InAllocationData.ExternalAccessQueue, SkeletalMeshObject, LodIndex, TEXT("OptimusSkinnedMeshPosition"));
+		FRDGBuffer* PositionBuffer = FSkeletalMeshDeformerHelpers::GetAllocatedPositionBuffer(GraphBuilder, SkeletalMeshObject, LodIndex);
 		PositionBufferSRV = GraphBuilder.CreateSRV(PositionBuffer, PF_R32_FLOAT); 
 	}
 	else
@@ -256,10 +249,11 @@ void FOptimusSkinnedMeshReadDataProviderProxy::AllocateResources(FRDGBuilder& Gr
 	// OpenGL ES does not support writing to RGBA16_SNORM images, instead pack data into SINT in the shader
 	const EPixelFormat TangentsFormat = IsOpenGLPlatform(GMaxRHIShaderPlatform) ? PF_R16G16B16A16_SINT : PF_R16G16B16A16_SNORM;
 
-	if ((InputMask & static_cast<uint64>(ESkinnedMeshReadDataInterfaceInputSelectorMask::TangentX)) ||
-		(InputMask & static_cast<uint64>(ESkinnedMeshReadDataInterfaceInputSelectorMask::TangentZ)))
+	if (((InputMask & static_cast<uint64>(ESkinnedMeshReadDataInterfaceInputSelectorMask::TangentX)) ||
+		(InputMask & static_cast<uint64>(ESkinnedMeshReadDataInterfaceInputSelectorMask::TangentZ))) &&
+		EnumHasAnyFlags(ReadableBuffers, EMeshDeformerOutputBuffer::SkinnedMeshTangents))
 	{
-		FRDGBuffer* TangentBuffer = FSkeletalMeshDeformerHelpers::AllocateVertexFactoryTangentBuffer(GraphBuilder, InAllocationData.ExternalAccessQueue, SkeletalMeshObject, LodIndex, TEXT("OptimusSkinnedMeshTangent"));
+		FRDGBuffer* TangentBuffer = FSkeletalMeshDeformerHelpers::GetAllocatedTangentBuffer(GraphBuilder, SkeletalMeshObject, LodIndex);
 		TangentBufferSRV = GraphBuilder.CreateSRV(TangentBuffer, TangentsFormat); 
 	}
 	else
@@ -268,10 +262,10 @@ void FOptimusSkinnedMeshReadDataProviderProxy::AllocateResources(FRDGBuilder& Gr
 		TangentBufferSRV = GraphBuilder.CreateSRV(White, TangentsFormat); 
 	}
 
-	if (InputMask & static_cast<uint64>(ESkinnedMeshReadDataInterfaceInputSelectorMask::Color))
+	if ((InputMask & static_cast<uint64>(ESkinnedMeshReadDataInterfaceInputSelectorMask::Color)) &&
+		EnumHasAnyFlags(ReadableBuffers, EMeshDeformerOutputBuffer::SkinnedMeshVertexColor))
 	{
-		FRDGBuffer* ColorBuffer = FSkeletalMeshDeformerHelpers::AllocateVertexFactoryColorBuffer(GraphBuilder, InAllocationData.ExternalAccessQueue, SkeletalMeshObject, LodIndex, TEXT("OptimusSkinnedMeshColor"));
-		// using RGBA here and do a manual fetch swizzle in shader instead of BGRA directly because some Mac does not support it. See GMetalBufferFormats[PF_B8G8R8A8] 
+		FRDGBuffer* ColorBuffer = FSkeletalMeshDeformerHelpers::GetAllocatedColorBuffer(GraphBuilder, SkeletalMeshObject, LodIndex);
 		ColorBufferSRV = GraphBuilder.CreateSRV(ColorBuffer, PF_R8G8B8A8); 
 	}
 	else
@@ -279,22 +273,24 @@ void FOptimusSkinnedMeshReadDataProviderProxy::AllocateResources(FRDGBuilder& Gr
 		FRDGBuffer* White = GraphBuilder.RegisterExternalBuffer(GWhiteVertexBufferWithRDG->Buffer);		
 		ColorBufferSRV = GraphBuilder.CreateSRV(White, TangentsFormat); 
 	}
-
-	FSkeletalMeshDeformerHelpers::UpdateVertexFactoryBufferOverrides(GraphBuilder, SkeletalMeshObject, LodIndex, bInvalidatePreviousPosition);
 }
 
 void FOptimusSkinnedMeshReadDataProviderProxy::GatherPermutations(FPermutationData& InOutPermutationData) const
 {
+	// Query the tracking extension for readable deformed buffers on this component
+	const EMeshDeformerOutputBuffer ReadableBuffers = FOptimusSkinnedMeshTrackingExtension::GetReadableDeformedBuffers(
+		SceneInterface, ComponentId, ReaderExecutionGroup, bIsPrimaryComponent);
+
 	EOptimusSkinnedMeshReadReadableOutputBuffer ReadableOutputBuffers = EOptimusSkinnedMeshReadReadableOutputBuffer::None; 
-	if (EnumHasAnyFlags(OutputBuffersFromPreviousInstances, EMeshDeformerOutputBuffer::SkinnedMeshPosition))
+	if (EnumHasAnyFlags(ReadableBuffers, EMeshDeformerOutputBuffer::SkinnedMeshPosition))
 	{
 		ReadableOutputBuffers |= EOptimusSkinnedMeshReadReadableOutputBuffer::Position;
 	}
-	if (EnumHasAnyFlags(OutputBuffersFromPreviousInstances, EMeshDeformerOutputBuffer::SkinnedMeshTangents))
+	if (EnumHasAnyFlags(ReadableBuffers, EMeshDeformerOutputBuffer::SkinnedMeshTangents))
 	{
 		ReadableOutputBuffers |= EOptimusSkinnedMeshReadReadableOutputBuffer::Tangents;
 	}
-	if (EnumHasAnyFlags(OutputBuffersFromPreviousInstances, EMeshDeformerOutputBuffer::SkinnedMeshVertexColor))
+	if (EnumHasAnyFlags(ReadableBuffers, EMeshDeformerOutputBuffer::SkinnedMeshVertexColor))
 	{
 		ReadableOutputBuffers |= EOptimusSkinnedMeshReadReadableOutputBuffer::VertexColor;
 	}
@@ -310,6 +306,8 @@ void FOptimusSkinnedMeshReadDataProviderProxy::GatherPermutations(FPermutationDa
 
 void FOptimusSkinnedMeshReadDataProviderProxy::GatherDispatchData(FDispatchData const& InDispatchData)
 {
+	FSkeletalMeshObject* SkeletalMeshObject = FSkeletalMeshDeformerHelpers::GetSkeletalMeshObject(SceneInterface, ComponentId);
+
 	const int32 LodIndex = SkeletalMeshObject->GetLOD();
 	FSkeletalMeshRenderData const& SkeletalMeshRenderData = SkeletalMeshObject->GetSkeletalMeshRenderData();
 	FSkeletalMeshLODRenderData const* LodRenderData = &SkeletalMeshRenderData.LODRenderData[LodIndex];
@@ -319,8 +317,6 @@ void FOptimusSkinnedMeshReadDataProviderProxy::GatherDispatchData(FDispatchData
 	const TStridedView<FParameters> ParameterArray = MakeStridedParameterView<FParameters>(InDispatchData);
 	for (int32 InvocationIndex = 0; InvocationIndex < ParameterArray.Num(); ++InvocationIndex)
 	{
-		FSkelMeshRenderSection const& RenderSection = LodRenderData->RenderSections[InvocationIndex];
-
 		FParameters& Parameters = ParameterArray[InvocationIndex];
 		Parameters.NumVertices = LodRenderData->GetNumVertices();
 
diff --git a/Engine/Plugins/Animation/DeformerGraph/Source/OptimusCore/Private/DataInterfaces/OptimusDataInterfaceSkinnedMeshRead.h b/Engine/Plugins/Animation/DeformerGraph/Source/OptimusCore/Private/DataInterfaces/OptimusDataInterfaceSkinnedMeshRead.h
--- a/Engine/Plugins/Animation/DeformerGraph/Source/OptimusCore/Private/DataInterfaces/OptimusDataInterfaceSkinnedMeshRead.h
+++ b/Engine/Plugins/Animation/DeformerGraph/Source/OptimusCore/Private/DataInterfaces/OptimusDataInterfaceSkinnedMeshRead.h
@@ -5,6 +5,7 @@
 #include "IOptimusDeformerInstanceAccessor.h"
 #include "OptimusComputeDataInterface.h"
 #include "ComputeFramework/ComputeDataProvider.h"
+#include "PrimitiveComponentId.h"
 #include "RHIFwd.h"
 
 #include "OptimusDataInterfaceSkinnedMeshRead.generated.h"
@@ -13,6 +14,7 @@
 
 class FRDGBuffer;
 class FRDGBufferSRV;
+class FSceneInterface;
 class FSkeletalMeshObject;
 class FSkinnedMeshReadDataInterfaceParameters;
 class USkinnedMeshComponent;
@@ -64,9 +66,6 @@ public:
 
 	uint64 InputMask = 0;
 
-	// Served as persistent storage for the provider proxy, should not be used by the data provider itself
-	int32 LastLodIndexCachedByRenderProxy = 0;
-
 	//~ Begin UComputeDataProvider Interface
 	FComputeDataProviderRenderProxy* GetRenderProxy() override;
 	//~ End UComputeDataProvider Interface
@@ -82,7 +81,7 @@ private:
 class FOptimusSkinnedMeshReadDataProviderProxy : public FComputeDataProviderRenderProxy
 {
 public:
-	FOptimusSkinnedMeshReadDataProviderProxy(USkinnedMeshComponent* InSkinnedMeshComponent, uint64 InInputMask, EMeshDeformerOutputBuffer InOutputBuffersWithValidData, int32* InLastLodIndexPtr);
+	FOptimusSkinnedMeshReadDataProviderProxy(USkinnedMeshComponent* InSkinnedMeshComponent, uint64 InInputMask, bool bInIsPrimaryComponent, FName InReaderExecutionGroup);
 
 	//~ Begin FComputeDataProviderRenderProxy Interface
 	bool IsValid(FValidationData const& InValidationData) const override;
@@ -94,10 +93,11 @@ public:
 private:
 	using FParameters = FSkinnedMeshReadDataInterfaceParameters;
 
-	FSkeletalMeshObject* SkeletalMeshObject = nullptr;
+	FSceneInterface* SceneInterface = nullptr;
+	FPrimitiveComponentId ComponentId;
 	uint64 InputMask = 0;
-	int32* LastLodIndexPtr = nullptr;
-	EMeshDeformerOutputBuffer OutputBuffersFromPreviousInstances;
+	bool bIsPrimaryComponent = true;
+	FName ReaderExecutionGroup;
 
 	FShaderResourceViewRHIRef TangentsSRV;
 	FShaderResourceViewRHIRef ColorsSRV;
diff --git a/Engine/Plugins/Animation/DeformerGraph/Source/OptimusCore/Private/DataInterfaces/OptimusDataInterfaceSkinnedMeshVertexAttribute.cpp b/Engine/Plugins/Animation/DeformerGraph/Source/OptimusCore/Private/DataInterfaces/OptimusDataInterfaceSkinnedMeshVertexAttribute.cpp
--- a/Engine/Plugins/Animation/DeformerGraph/Source/OptimusCore/Private/DataInterfaces/OptimusDataInterfaceSkinnedMeshVertexAttribute.cpp
+++ b/Engine/Plugins/Animation/DeformerGraph/Source/OptimusCore/Private/DataInterfaces/OptimusDataInterfaceSkinnedMeshVertexAttribute.cpp
@@ -2,6 +2,8 @@
 
 #include "OptimusDataInterfaceSkinnedMeshVertexAttribute.h"
 
+#include "Components/SkinnedMeshComponent.h"
+#include "SkeletalMeshDeformerHelpers.h"
 #include "OptimusDataTypeRegistry.h"
 #include "OptimusDeformerInstance.h"
 #include "OptimusValueContainerStruct.h"
@@ -171,7 +173,11 @@ FOptimusSkinnedMeshVertexAttributeDataProviderProxy::FOptimusSkinnedMeshVertexAt
 	float InDefaultValue
 	)
 {
-	SkeletalMeshObject = InSkinnedMeshComponent ? InSkinnedMeshComponent->MeshObject : nullptr;
+	if (InSkinnedMeshComponent && InSkinnedMeshComponent->IsRegistered())
+	{
+		SceneInterface = InSkinnedMeshComponent->GetScene();
+		ComponentId = InSkinnedMeshComponent->GetPrimitiveSceneId();
+	}
 	AttributeName = InAttributeName;
 	DefaultValue = InDefaultValue;
 }
@@ -179,11 +185,12 @@ FOptimusSkinnedMeshVertexAttributeDataProviderProxy::FOptimusSkinnedMeshVertexAt
 
 bool FOptimusSkinnedMeshVertexAttributeDataProviderProxy::IsValid(FValidationData const& InValidationData) const
 {
-	if (InValidationData.ParameterStructSize != sizeof(FParameters))
+	FSkeletalMeshObject* SkeletalMeshObject = FSkeletalMeshDeformerHelpers::GetSkeletalMeshObject(SceneInterface, ComponentId);
+	if (!SkeletalMeshObject)
 	{
 		return false;
 	}
-	if (SkeletalMeshObject == nullptr)
+	if (InValidationData.ParameterStructSize != sizeof(FParameters))
 	{
 		return false;
 	}
@@ -196,6 +203,7 @@ void FOptimusSkinnedMeshVertexAttributeDataProviderProxy::GatherDispatchData(
 	FDispatchData const& InDispatchData
 	)
 {
+	FSkeletalMeshObject* SkeletalMeshObject = FSkeletalMeshDeformerHelpers::GetSkeletalMeshObject(SceneInterface, ComponentId);
 	const int32 LodIndex = SkeletalMeshObject->GetLOD();
 	FSkeletalMeshRenderData const& SkeletalMeshRenderData = SkeletalMeshObject->GetSkeletalMeshRenderData();
 	FSkeletalMeshLODRenderData const* LodRenderData = &SkeletalMeshRenderData.LODRenderData[LodIndex];
diff --git a/Engine/Plugins/Animation/DeformerGraph/Source/OptimusCore/Private/DataInterfaces/OptimusDataInterfaceSkinnedMeshVertexAttribute.h b/Engine/Plugins/Animation/DeformerGraph/Source/OptimusCore/Private/DataInterfaces/OptimusDataInterfaceSkinnedMeshVertexAttribute.h
--- a/Engine/Plugins/Animation/DeformerGraph/Source/OptimusCore/Private/DataInterfaces/OptimusDataInterfaceSkinnedMeshVertexAttribute.h
+++ b/Engine/Plugins/Animation/DeformerGraph/Source/OptimusCore/Private/DataInterfaces/OptimusDataInterfaceSkinnedMeshVertexAttribute.h
@@ -7,9 +7,11 @@
 
 #include "ComputeFramework/ComputeDataProvider.h"
 
+#include "PrimitiveComponentId.h"
 #include "OptimusDataInterfaceSkinnedMeshVertexAttribute.generated.h"
 
 
+class FSceneInterface;
 class FSkeletalMeshObject;
 class FSkinnedMeshVertexAttributeDataInterfaceParameters;
 
@@ -95,7 +97,8 @@ public:
 private:
 	using FParameters = FSkinnedMeshVertexAttributeDataInterfaceParameters;
 	
-	FSkeletalMeshObject* SkeletalMeshObject;
+	FSceneInterface* SceneInterface = nullptr;
+	FPrimitiveComponentId ComponentId;
 	FName AttributeName;
 	float DefaultValue;
 };
diff --git a/Engine/Plugins/Animation/DeformerGraph/Source/OptimusCore/Private/DataInterfaces/OptimusDataInterfaceSkinnedMeshWrite.cpp b/Engine/Plugins/Animation/DeformerGraph/Source/OptimusCore/Private/DataInterfaces/OptimusDataInterfaceSkinnedMeshWrite.cpp
--- a/Engine/Plugins/Animation/DeformerGraph/Source/OptimusCore/Private/DataInterfaces/OptimusDataInterfaceSkinnedMeshWrite.cpp
+++ b/Engine/Plugins/Animation/DeformerGraph/Source/OptimusCore/Private/DataInterfaces/OptimusDataInterfaceSkinnedMeshWrite.cpp
@@ -4,6 +4,7 @@
 
 #include "Components/SkinnedMeshComponent.h"
 #include "ComputeFramework/ShaderParamTypeDefinition.h"
+#include "OptimusSkinnedMeshTrackingExtension.h"
 #include "OptimusDataDomain.h"
 #include "RenderGraphBuilder.h"
 #include "Rendering/SkeletalMeshLODRenderData.h"
@@ -14,9 +15,9 @@
 #include "SkeletalRenderPublic.h"
 #include "DataDrivenShaderPlatformInfo.h"
 #include "Animation/MeshDeformerInstance.h"
+#include "OptimusDeformerInstance.h"
 
 #if WITH_EDITORONLY_DATA
-#include "OptimusDeformerInstance.h"
 #include "Components/SkeletalMeshComponent.h"
 #include "Rendering/SkeletalMeshModel.h"
 #include "MeshAttributes.h"
@@ -169,10 +170,18 @@ UOptimusSkinnedMeshWriteDataProvider::UOptimusSkinnedMeshWriteDataProvider(FVTab
 FComputeDataProviderRenderProxy* UOptimusSkinnedMeshWriteDataProvider::GetRenderProxy()
 {
 	FOptimusSkinnedMeshWriteDataProviderProxy* Proxy = new FOptimusSkinnedMeshWriteDataProviderProxy();
-	
-	Proxy->SkeletalMeshObject = SkinnedMesh.IsValid() ? SkinnedMesh->MeshObject : nullptr;
+
+	if (SkinnedMesh.IsValid() && SkinnedMesh->IsRegistered())
+	{
+		Proxy->SceneInterface = SkinnedMesh->GetScene();
+		Proxy->ComponentId = SkinnedMesh->GetPrimitiveSceneId();
+	}
 	Proxy->OutputMask = OutputMask;
-	Proxy->LastLodIndexPtr = &LastLodIndexCachedByRenderProxy;
+
+	if (WeakDeformerInstance.IsValid())
+	{
+		Proxy->ExecutionGroup = WeakDeformerInstance->CurrentExecutionGroupName;
+	}
 
 #if WITH_EDITORONLY_DATA
 	Proxy->FrameNumber = GFrameNumber;
@@ -196,17 +205,20 @@ bool UOptimusSkinnedMeshWriteDataProvider::RequestReadbackDeformerGeometry(TUniq
 }
 #endif // WITH_EDITORONLY_DATA
 
+void UOptimusSkinnedMeshWriteDataProvider::SetDeformerInstance(UOptimusDeformerInstance* InInstance)
+{
+	WeakDeformerInstance = InInstance;
+}
+
 bool FOptimusSkinnedMeshWriteDataProviderProxy::IsValid(FValidationData const& InValidationData) const
 {
 	if (InValidationData.ParameterStructSize != sizeof(FParameters))
 	{
 		return false;
 	}
-	if (SkeletalMeshObject == nullptr)
-	{
-		return false;
-	}
-	if (SkeletalMeshObject->IsCPUSkinned())
+
+	FSkeletalMeshObject* SkeletalMeshObject = FSkeletalMeshDeformerHelpers::GetSkeletalMeshObject(SceneInterface, ComponentId);
+	if (!SkeletalMeshObject)
 	{
 		return false;
 	}
@@ -232,22 +244,17 @@ bool FOptimusSkinnedMeshWriteDataProviderProxy::IsValid(FValidationData const& I
 
 void FOptimusSkinnedMeshWriteDataProviderProxy::AllocateResources(FRDGBuilder& GraphBuilder, FAllocationData const& InAllocationData)
 {
+	FSkeletalMeshObject* SkeletalMeshObject = FSkeletalMeshDeformerHelpers::GetSkeletalMeshObject(SceneInterface, ComponentId);
+
 	// Allocate required buffers
 	const int32 LodIndex = SkeletalMeshObject->GetLOD();
-
-	// Avoid using previous position buffer from when this LOD was last active to compute motion vectors.
-	// The position delta between that previous position (could be from any time ago) and the current position can be any crazy value that is not meaningful.
-	bool bInvalidatePreviousPosition = false;
-	if (LodIndex != *LastLodIndexPtr)
-	{
-		bInvalidatePreviousPosition = true;
-		*LastLodIndexPtr = LodIndex;
-	}
+	EMeshDeformerOutputBuffer DeformerOutputBuffers = EMeshDeformerOutputBuffer::None;
 
 	if (OutputMask & static_cast<uint64>(ESkinnedMeshWriteDataInterfaceOutputSelectorMask::Position))
 	{
 		PositionBuffer = FSkeletalMeshDeformerHelpers::AllocateVertexFactoryPositionBuffer(GraphBuilder, InAllocationData.ExternalAccessQueue, SkeletalMeshObject, LodIndex, TEXT("OptimusSkinnedMeshPosition"));
 		PositionBufferUAV = GraphBuilder.CreateUAV(PositionBuffer, PF_R32_FLOAT, ERDGUnorderedAccessViewFlags::SkipBarrier);
+		DeformerOutputBuffers |= EMeshDeformerOutputBuffer::SkinnedMeshPosition;
 	}
 	else
 	{
@@ -270,6 +277,7 @@ void FOptimusSkinnedMeshWriteDataProviderProxy::AllocateResources(FRDGBuilder& G
 	{
 		TangentBuffer = FSkeletalMeshDeformerHelpers::AllocateVertexFactoryTangentBuffer(GraphBuilder, InAllocationData.ExternalAccessQueue, SkeletalMeshObject, LodIndex, TEXT("OptimusSkinnedMeshTangent"));
 		TangentBufferUAV = GraphBuilder.CreateUAV(TangentBuffer, TangentsFormat, ERDGUnorderedAccessViewFlags::SkipBarrier);
+		DeformerOutputBuffers |= EMeshDeformerOutputBuffer::SkinnedMeshTangents;
 	}
 	else
 	{
@@ -289,6 +297,7 @@ void FOptimusSkinnedMeshWriteDataProviderProxy::AllocateResources(FRDGBuilder& G
 		ColorBuffer = FSkeletalMeshDeformerHelpers::AllocateVertexFactoryColorBuffer(GraphBuilder, InAllocationData.ExternalAccessQueue, SkeletalMeshObject, LodIndex, TEXT("OptimusSkinnedMeshColor"));
 		// using RGBA here and do a manual fetch swizzle in shader instead of BGRA directly because some Mac does not support it. See GMetalBufferFormats[PF_B8G8R8A8] 
 		ColorBufferUAV = GraphBuilder.CreateUAV(ColorBuffer, PF_R8G8B8A8, ERDGUnorderedAccessViewFlags::SkipBarrier);
+		DeformerOutputBuffers |= EMeshDeformerOutputBuffer::SkinnedMeshVertexColor;
 	}
 	else
 	{
@@ -303,11 +312,13 @@ void FOptimusSkinnedMeshWriteDataProviderProxy::AllocateResources(FRDGBuilder& G
 		ColorBufferUAV = GraphBuilder.CreateUAV(GraphBuilder.RegisterExternalBuffer(GWhiteVertexBufferWithRDG->Buffer), PF_A32B32G32R32F);
 	}
 
-	FSkeletalMeshDeformerHelpers::UpdateVertexFactoryBufferOverrides(GraphBuilder, SkeletalMeshObject, LodIndex, bInvalidatePreviousPosition);
+	FOptimusSkinnedMeshTrackingExtension::MarkDeformed(GraphBuilder, SceneInterface, ComponentId, SkeletalMeshObject, DeformerOutputBuffers, ExecutionGroup);
 }
 
 void FOptimusSkinnedMeshWriteDataProviderProxy::GatherDispatchData(FDispatchData const& InDispatchData)
 {
+	FSkeletalMeshObject* SkeletalMeshObject = FSkeletalMeshDeformerHelpers::GetSkeletalMeshObject(SceneInterface, ComponentId);
+
 	const int32 LodIndex = SkeletalMeshObject->GetLOD();
 	FSkeletalMeshRenderData const& SkeletalMeshRenderData = SkeletalMeshObject->GetSkeletalMeshRenderData();
 	FSkeletalMeshLODRenderData const* LodRenderData = &SkeletalMeshRenderData.LODRenderData[LodIndex];
@@ -315,8 +326,6 @@ void FOptimusSkinnedMeshWriteDataProviderProxy::GatherDispatchData(FDispatchData
 	const TStridedView<FParameters> ParameterArray = MakeStridedParameterView<FParameters>(InDispatchData);
 	for (int32 InvocationIndex = 0; InvocationIndex < ParameterArray.Num(); ++InvocationIndex)
 	{
-		FSkelMeshRenderSection const& RenderSection = LodRenderData->RenderSections[InvocationIndex];
-
 		FParameters& Parameters = ParameterArray[InvocationIndex];
 		Parameters.NumVertices =LodRenderData->GetNumVertices();
 		Parameters.PositionBufferUAV = PositionBufferUAV;
@@ -343,7 +352,8 @@ void FOptimusSkinnedMeshWriteDataProviderProxy::GetReadbackData(TArray<FReadback
 			GeometryReadback->GeometryReadbackRequests.Add(MoveTemp(Request));
 		}
 		GeometryReadback->SkeletalMesh = SkeletalMeshAsset;
-		GeometryReadback->LodIndex = SkeletalMeshObject->GetLOD();
+		FSkeletalMeshObject* ReadbackMeshObject = FSkeletalMeshDeformerHelpers::GetSkeletalMeshObject(SceneInterface, ComponentId);
+		GeometryReadback->LodIndex = ReadbackMeshObject ? ReadbackMeshObject->GetLOD() : 0;
 
 		auto SetupBufferReadback = [&OutReadbackData, this](FRDGBuffer* GPUBuffer, FBufferReadback& BufferReadback)
 		{
diff --git a/Engine/Plugins/Animation/DeformerGraph/Source/OptimusCore/Private/DataInterfaces/OptimusDataInterfaceSkinnedMeshWrite.h b/Engine/Plugins/Animation/DeformerGraph/Source/OptimusCore/Private/DataInterfaces/OptimusDataInterfaceSkinnedMeshWrite.h
--- a/Engine/Plugins/Animation/DeformerGraph/Source/OptimusCore/Private/DataInterfaces/OptimusDataInterfaceSkinnedMeshWrite.h
+++ b/Engine/Plugins/Animation/DeformerGraph/Source/OptimusCore/Private/DataInterfaces/OptimusDataInterfaceSkinnedMeshWrite.h
@@ -2,10 +2,12 @@
 
 #pragma once
 
+#include "IOptimusDeformerInstanceAccessor.h"
 #include "IOptimusDeformerGeometryReadbackProvider.h"
 #include "IOptimusOutputBufferWriter.h"
 #include "OptimusComputeDataInterface.h"
 #include "ComputeFramework/ComputeDataProvider.h"
+#include "PrimitiveComponentId.h"
 
 
 #include "OptimusDataInterfaceSkinnedMeshWrite.generated.h"
@@ -14,9 +16,11 @@
 
 class FRDGBuffer;
 class FRDGBufferUAV;
+class FSceneInterface;
 class FSkeletalMeshObject;
 class FSkinedMeshWriteDataInterfaceParameters;
 class USkinnedMeshComponent;
+enum class EMeshDeformerOutputBuffer : uint8;
 
 #if WITH_EDITORONLY_DATA
 class USkeletalMesh;
@@ -64,6 +68,7 @@ private:
 UCLASS(BlueprintType, editinlinenew, Category = ComputeFramework)
 class UOptimusSkinnedMeshWriteDataProvider :
 	public UComputeDataProvider,
+	public IOptimusDeformerInstanceAccessor,
 	public IOptimusDeformerGeometryReadbackProvider
 {
 	GENERATED_BODY()
@@ -80,6 +85,10 @@ public:
 	//~ Begin UComputeDataProvider Interface
 	FComputeDataProviderRenderProxy* GetRenderProxy() override;
 	//~ End UComputeDataProvider Interface
+
+	//~ Begin IOptimusDeformerInstanceAccessor Interface
+	void SetDeformerInstance(UOptimusDeformerInstance* InInstance) override;
+	//~ End IOptimusDeformerInstanceAccessor Interface
 	
 #if WITH_EDITORONLY_DATA	
 	//~ Begin IOptimusDeformerGeometryReadbackProvider Interface
@@ -87,8 +96,7 @@ public:
 	//~ End IOptimusDeformerGeometryReadbackProvider Interface
 #endif // WITH_EDITORONLY_DATA	
 private:
-	// Served as persistent storage for the provider proxy, should not be used by the data provider itself
-	int32 LastLodIndexCachedByRenderProxy = 0;
+	TWeakObjectPtr<UOptimusDeformerInstance> WeakDeformerInstance = nullptr;
 
 #if WITH_EDITORONLY_DATA
 	/** Readback requests for the current frame */
@@ -106,9 +114,10 @@ public:
 	void GetReadbackData(TArray<FReadbackData>& OutReadbackData) const override;
 	//~ End FComputeDataProviderRenderProxy Interface
 
-	FSkeletalMeshObject* SkeletalMeshObject = nullptr;
+	FSceneInterface* SceneInterface = nullptr;
+	FPrimitiveComponentId ComponentId;
 	uint64 OutputMask = 0;
-	int32* LastLodIndexPtr = nullptr;
+	FName ExecutionGroup;
 
 #if WITH_EDITORONLY_DATA
 
diff --git a/Engine/Plugins/Animation/DeformerGraph/Source/OptimusCore/Private/OptimusDeformerDynamicInstanceManager.cpp b/Engine/Plugins/Animation/DeformerGraph/Source/OptimusCore/Private/OptimusDeformerDynamicInstanceManager.cpp
--- a/Engine/Plugins/Animation/DeformerGraph/Source/OptimusCore/Private/OptimusDeformerDynamicInstanceManager.cpp
+++ b/Engine/Plugins/Animation/DeformerGraph/Source/OptimusCore/Private/OptimusDeformerDynamicInstanceManager.cpp
@@ -113,14 +113,10 @@ void UOptimusDeformerDynamicInstanceManager::EnqueueWork(FEnqueueWorkDesc const&
 	
 	// Making sure instances in the queue are dispatched sequentially
 	uint8 NumComputeGraphsPossiblyEnqueued = 0;
-	// Used to inform later instances whether specific buffers have valid data in them
-	EMeshDeformerOutputBuffer OutputBuffers = EMeshDeformerOutputBuffer::None;
 	
 	for (int32 Index = 0; Index < SortedInstances.Num(); Index++)
 	{
 		UOptimusDeformerInstance* Instance = SortedInstances[Index];
-		Instance->OutputBuffersFromPreviousInstances = OutputBuffers;
-		OutputBuffers |= Instance->GetOutputBuffers();
 		
 		Instance->GraphSortPriorityOffset = NumComputeGraphsPossiblyEnqueued;
 		NumComputeGraphsPossiblyEnqueued += Instance->ComputeGraphExecInfos.Num();
diff --git a/Engine/Plugins/Animation/DeformerGraph/Source/OptimusCore/Private/OptimusDeformerInstance.cpp b/Engine/Plugins/Animation/DeformerGraph/Source/OptimusCore/Private/OptimusDeformerInstance.cpp
--- a/Engine/Plugins/Animation/DeformerGraph/Source/OptimusCore/Private/OptimusDeformerInstance.cpp
+++ b/Engine/Plugins/Animation/DeformerGraph/Source/OptimusCore/Private/OptimusDeformerInstance.cpp
@@ -537,6 +537,9 @@ void UOptimusDeformerInstance::EnqueueWork(UMeshDeformerInstance::FEnqueueWorkDe
 
 
 	
+	// Store the execution group name for data providers to use
+	CurrentExecutionGroupName = ExecutionGroupName;
+
 	// Enqueue work.
 	bool bIsWorkEnqueued = false;
 	if (bCanBeActive)
diff --git a/Engine/Plugins/Animation/DeformerGraph/Source/OptimusCore/Private/OptimusSkinnedMeshTrackingExtension.cpp b/Engine/Plugins/Animation/DeformerGraph/Source/OptimusCore/Private/OptimusSkinnedMeshTrackingExtension.cpp
new file mode 100644
--- /dev/null
+++ b/Engine/Plugins/Animation/DeformerGraph/Source/OptimusCore/Private/OptimusSkinnedMeshTrackingExtension.cpp
@@ -0,0 +1,154 @@
+// Copyright Epic Games, Inc. All Rights Reserved.
+
+#include "OptimusSkinnedMeshTrackingExtension.h"
+
+#include "ComputeWorkerInterface.h"
+#include "ScenePrivate.h"
+#include "ScenePrimitiveUpdates.h"
+#include "SkeletalMeshDeformerHelpers.h"
+#include "SkeletalRenderPublic.h"
+
+IMPLEMENT_SCENE_EXTENSION(FOptimusSkinnedMeshTrackingExtension);
+
+// ---------------------------------------------------------------------------
+// Updater
+// ---------------------------------------------------------------------------
+
+class FOptimusSkinnedMeshTrackingExtension::FUpdater : public ISceneExtensionUpdater
+{
+	DECLARE_SCENE_EXTENSION_UPDATER(FUpdater, FOptimusSkinnedMeshTrackingExtension);
+
+public:
+	explicit FUpdater(FOptimusSkinnedMeshTrackingExtension& InExtension)
+		: Extension(InExtension)
+	{
+	}
+
+	virtual void PreSceneUpdate(
+		FRDGBuilder& GraphBuilder,
+		const FScenePreUpdateChangeSet& ChangeSet) override
+	{
+		for (FPrimitiveSceneInfo* PrimitiveSceneInfo : ChangeSet.RemovedPrimitiveSceneInfos)
+		{
+			Extension.ComponentDataMap.Remove(PrimitiveSceneInfo->PrimitiveComponentId);
+		}
+	}
+
+private:
+	FOptimusSkinnedMeshTrackingExtension& Extension;
+};
+
+ISceneExtensionUpdater* FOptimusSkinnedMeshTrackingExtension::CreateUpdater()
+{
+	return new FUpdater(*this);
+}
+
+// ---------------------------------------------------------------------------
+// Accessors
+// ---------------------------------------------------------------------------
+
+void FOptimusSkinnedMeshTrackingExtension::MarkDeformed(
+	FRDGBuilder& GraphBuilder, FSceneInterface* SceneInterface, FPrimitiveComponentId ComponentId, FSkeletalMeshObject* MeshObject,
+	EMeshDeformerOutputBuffer OutputBuffers, FName ExecutionGroup)
+{
+	if (!SceneInterface || !MeshObject || !ComponentId.IsValid())
+	{
+		return;
+	}
+
+	FOptimusSkinnedMeshTrackingExtension* Extension = Get(SceneInterface);
+	if (!Extension)
+	{
+		return;
+	}
+
+	FComponentData& Data = Extension->ComponentDataMap.FindOrAdd(ComponentId);
+
+	// Start fresh when it is a new frame
+	if (Data.DeformedFrame != GFrameNumberRenderThread)
+	{
+		Data.DeformedFrame = GFrameNumberRenderThread;
+		
+		const int32 LodIndex = MeshObject->GetLOD();
+
+		// Avoid using previous position from the new lod during LOD transition, 
+		// because it likely contains position data too outdated to be used to 
+		// compute motion vectors for the current frame
+		bool bInvalidatePreviousPosition = false;
+		if (LodIndex != Data.LastLodIndex)
+		{
+			Data.LastLodIndex = LodIndex;
+			bInvalidatePreviousPosition = true;
+		}
+
+		// Add a pass to override vertex factory buffers with buffers allocated by Optimus
+		FSkeletalMeshDeformerHelpers::UpdateVertexFactoryBufferOverrides(
+			GraphBuilder, MeshObject, LodIndex, bInvalidatePreviousPosition);
+		
+		// Restart output buffer tracking
+		Data.EndOfFrameUpdateOutputBuffers = EMeshDeformerOutputBuffer::None;
+		Data.BeginInitViewsOutputBuffers = EMeshDeformerOutputBuffer::None;
+	}
+
+	if (ExecutionGroup == ComputeTaskExecutionGroup::EndOfFrameUpdate)
+	{
+		Data.EndOfFrameUpdateOutputBuffers |= OutputBuffers;
+	}
+	else if (ExecutionGroup == ComputeTaskExecutionGroup::BeginInitViews)
+	{
+		Data.BeginInitViewsOutputBuffers |= OutputBuffers;
+	}
+}
+
+EMeshDeformerOutputBuffer FOptimusSkinnedMeshTrackingExtension::GetReadableDeformedBuffers(
+	FSceneInterface* SceneInterface, FPrimitiveComponentId ComponentId,
+	FName ReaderExecutionGroup, bool bIsPrimaryComponent)
+{
+	FOptimusSkinnedMeshTrackingExtension* Extension = Get(SceneInterface);
+	if (!Extension)
+	{
+		return EMeshDeformerOutputBuffer::None;
+	}
+
+	const FComponentData* Data = Extension->ComponentDataMap.Find(ComponentId);
+	if (!Data || Data->DeformedFrame != GFrameNumberRenderThread)
+	{
+		// Component was not deformed by Optimus this frame
+		return EMeshDeformerOutputBuffer::None;
+	}
+
+	EMeshDeformerOutputBuffer Result = EMeshDeformerOutputBuffer::None;
+
+	if (bIsPrimaryComponent)
+	{
+		// Same-component reads are serialized via sort priority — all buffers are safe
+		Result = Data->EndOfFrameUpdateOutputBuffers | Data->BeginInitViewsOutputBuffers;
+	}
+	else
+	{
+		// Secondary read from BeginInitViews can read EndOfFrameUpdate outputs (earlier group)
+		if (ReaderExecutionGroup == ComputeTaskExecutionGroup::BeginInitViews)
+		{
+			Result = Data->EndOfFrameUpdateOutputBuffers;
+		}
+	}
+	
+	return Result;
+}
+
+FOptimusSkinnedMeshTrackingExtension* FOptimusSkinnedMeshTrackingExtension::Get(FSceneInterface* SceneInterface)
+{
+	if (!SceneInterface)
+	{
+		return nullptr;
+	}
+
+	FScene* RenderScene = SceneInterface->GetRenderScene();
+	if (!RenderScene)
+	{
+		return nullptr;
+	}
+
+	return RenderScene->GetExtensionPtr<FOptimusSkinnedMeshTrackingExtension>();
+}
+
diff --git a/Engine/Plugins/Animation/DeformerGraph/Source/OptimusCore/Private/OptimusSkinnedMeshTrackingExtension.h b/Engine/Plugins/Animation/DeformerGraph/Source/OptimusCore/Private/OptimusSkinnedMeshTrackingExtension.h
new file mode 100644
--- /dev/null
+++ b/Engine/Plugins/Animation/DeformerGraph/Source/OptimusCore/Private/OptimusSkinnedMeshTrackingExtension.h
@@ -0,0 +1,73 @@
+// Copyright Epic Games, Inc. All Rights Reserved.
+
+#pragma once
+
+#include "Animation/MeshDeformerInstance.h"
+#include "PrimitiveComponentId.h"
+#include "SceneExtensions.h"
+
+class FRDGBuilder;
+class FScene;
+class FSceneInterface;
+class FSkeletalMeshObject;
+
+/**
+ * Scene extension that tracks per-component deformer state on the render thread.
+ * Entries are created on demand when MarkDeformed is called and removed via
+ * PreSceneUpdate when primitives leave the scene.
+ */
+class FOptimusSkinnedMeshTrackingExtension : public ISceneExtension
+{
+	DECLARE_SCENE_EXTENSION(OPTIMUSCORE_API, FOptimusSkinnedMeshTrackingExtension);
+
+public:
+	/** Per-component tracking data. Packed to 2 DWords (8 bytes). */
+	struct FComponentData
+	{
+		int8 LastLodIndex = INDEX_NONE;
+		EMeshDeformerOutputBuffer EndOfFrameUpdateOutputBuffers = EMeshDeformerOutputBuffer::None;
+		EMeshDeformerOutputBuffer BeginInitViewsOutputBuffers = EMeshDeformerOutputBuffer::None;
+		uint32 DeformedFrame = 0;
+	};
+
+	using ISceneExtension::ISceneExtension;
+
+	//~ ISceneExtension interface
+	virtual ISceneExtensionUpdater* CreateUpdater() override;
+
+	/**
+	 * RT: Mark a component as deformed by Optimus. 
+	 * Adds a RDG pass to make sure the passthrough vertex factory bind to buffers allocated by Optimus.
+	 * Also keeps track of the deformer geometry buffers that Optimus writes to during the dispatch of each execution group
+	 */
+	static void MarkDeformed(
+		FRDGBuilder& GraphBuilder, 
+		FSceneInterface* SceneInterface, 
+		FPrimitiveComponentId ComponentId, 
+		FSkeletalMeshObject* MeshObject,
+		EMeshDeformerOutputBuffer OutputBuffers, 
+		FName ExecutionGroup
+		);
+
+	/**
+	 * RT: Returns which deformer output buffers are safe to read for a given component.
+	 * For primary component reads (bIsPrimaryComponent=true), all recorded output buffers are returned
+	 * since deformer instances on the same component are dispatched sequentially.
+	 * For secondary component reads, only output buffers from execution groups that dispatched earlier
+	 * than the reader's execution group are returned — same-group cross-component reads are excluded
+	 * because their relative dispatch order is not guaranteed.
+	 */
+	static EMeshDeformerOutputBuffer GetReadableDeformedBuffers(
+		FSceneInterface* SceneInterface, 
+		FPrimitiveComponentId ComponentId,
+		FName ReaderExecutionGroup, 
+		bool bIsPrimaryComponent);
+
+private:
+	class FUpdater;
+
+	static FOptimusSkinnedMeshTrackingExtension* Get(FSceneInterface* SceneInterface);
+
+	/** Map of component ID → tracking data. Accessed only on RT. */
+	TMap<FPrimitiveComponentId, FComponentData> ComponentDataMap;
+};
diff --git a/Engine/Plugins/Animation/DeformerGraph/Source/OptimusCore/Public/OptimusDeformerInstance.h b/Engine/Plugins/Animation/DeformerGraph/Source/OptimusCore/Public/OptimusDeformerInstance.h
--- a/Engine/Plugins/Animation/DeformerGraph/Source/OptimusCore/Public/OptimusDeformerInstance.h
+++ b/Engine/Plugins/Animation/DeformerGraph/Source/OptimusCore/Public/OptimusDeformerInstance.h
@@ -310,9 +310,13 @@ public:
 	
 	/** Making sure compute graphs belong to this instance does not run before instances before it */
 	int32 GraphSortPriorityOffset = 0;
-	
-	/** Used to see which buffers have valid data produced by dispatched instances and are safe to access for the current instance */	
-	EMeshDeformerOutputBuffer OutputBuffersFromPreviousInstances = EMeshDeformerOutputBuffer::None;
+
+	/** The execution group name for the current EnqueueWork call (e.g. EndOfFrameUpdate, BeginInitViews).
+	 *  Set on game thread before GetRenderProxy() calls, used by read/write proxies to inform the tracking extension. */
+	FName CurrentExecutionGroupName = NAME_None;
+
+	/** Get the primary mesh component that owns this deformer instance. */
+	UMeshComponent* GetMeshComponent() const { return MeshComponent.Get(); }
 protected:
 	friend class UOptimusDeformerDynamicInstanceManager;
 
diff --git a/Engine/Plugins/Animation/MLDeformer/MLDeformerFramework/Source/MLDeformerFramework/Private/DeformerGraphDataInterfaceSkeletonWithQuats.cpp b/Engine/Plugins/Animation/MLDeformer/MLDeformerFramework/Source/MLDeformerFramework/Private/DeformerGraphDataInterfaceSkeletonWithQuats.cpp
--- a/Engine/Plugins/Animation/MLDeformer/MLDeformerFramework/Source/MLDeformerFramework/Private/DeformerGraphDataInterfaceSkeletonWithQuats.cpp
+++ b/Engine/Plugins/Animation/MLDeformer/MLDeformerFramework/Source/MLDeformerFramework/Private/DeformerGraphDataInterfaceSkeletonWithQuats.cpp
@@ -133,7 +133,11 @@ namespace UE::MLDeformer
 {
 	FDeformerGraphSkeletonWithWeightedQuatsDataProviderProxy::FDeformerGraphSkeletonWithWeightedQuatsDataProviderProxy(USkinnedMeshComponent* SkinnedMeshComponent)
 	{
-		SkeletalMeshObject = SkinnedMeshComponent != nullptr ? SkinnedMeshComponent->MeshObject : nullptr;
+		if (SkinnedMeshComponent && SkinnedMeshComponent->IsRegistered())
+		{
+			SceneInterface = SkinnedMeshComponent->GetScene();
+			ComponentId = SkinnedMeshComponent->GetPrimitiveSceneId();
+		}
 		BoneRevisionNumber = SkinnedMeshComponent != nullptr ? SkinnedMeshComponent->GetBoneTransformRevisionNumber() : 0;
 	}
 
@@ -143,7 +147,8 @@ namespace UE::MLDeformer
 		{
 			return false;
 		}
-		if (SkeletalMeshObject == nullptr)
+		FSkeletalMeshObject* SkeletalMeshObject = FSkeletalMeshDeformerHelpers::GetSkeletalMeshObject(SceneInterface, ComponentId);
+		if (!SkeletalMeshObject)
 		{
 			return false;
 		}
@@ -189,6 +194,7 @@ namespace UE::MLDeformer
 
 	void FDeformerGraphSkeletonWithWeightedQuatsDataProviderProxy::GatherPermutations(FPermutationData& InOutPermutationData) const
 	{
+		FSkeletalMeshObject* SkeletalMeshObject = FSkeletalMeshDeformerHelpers::GetSkeletalMeshObject(SceneInterface, ComponentId);
 		const int32 LodIndex = SkeletalMeshObject->GetLOD();
 		FSkeletalMeshRenderData const& SkeletalMeshRenderData = SkeletalMeshObject->GetSkeletalMeshRenderData();
 		FSkeletalMeshLODRenderData const* LodRenderData = &SkeletalMeshRenderData.LODRenderData[LodIndex];
@@ -227,6 +233,7 @@ namespace UE::MLDeformer
 			return;
 		}
 
+		FSkeletalMeshObject* SkeletalMeshObject = FSkeletalMeshDeformerHelpers::GetSkeletalMeshObject(SceneInterface, ComponentId);
 		const int32 LodIndex = SkeletalMeshObject->GetLOD();
 		FSkeletalMeshRenderData const& SkeletalMeshRenderData = SkeletalMeshObject->GetSkeletalMeshRenderData();
 		FSkeletalMeshLODRenderData const* LodRenderData = &SkeletalMeshRenderData.LODRenderData[LodIndex];
@@ -268,6 +275,7 @@ namespace UE::MLDeformer
 
 	void FDeformerGraphSkeletonWithWeightedQuatsDataProviderProxy::GatherDispatchData(FDispatchData const& InDispatchData)
 	{
+		FSkeletalMeshObject* SkeletalMeshObject = FSkeletalMeshDeformerHelpers::GetSkeletalMeshObject(SceneInterface, ComponentId);
 		const int32 LodIndex = SkeletalMeshObject->GetLOD();
 		FSkeletalMeshRenderData const& SkeletalMeshRenderData = SkeletalMeshObject->GetSkeletalMeshRenderData();
 		FSkeletalMeshLODRenderData const* LodRenderData = &SkeletalMeshRenderData.LODRenderData[LodIndex];
@@ -304,6 +312,7 @@ namespace UE::MLDeformer
 	void FDeformerGraphSkeletonWithWeightedQuatsDataProviderProxy::CacheRefToLocalQuats(TArray<FQuat4f>& OutRefToLocalQuats) const
 	{
 		OutRefToLocalQuats.Reset();
+		FSkeletalMeshObject* SkeletalMeshObject = FSkeletalMeshDeformerHelpers::GetSkeletalMeshObject(SceneInterface, ComponentId);
 		if (!SkeletalMeshObject || !SkeletalMeshObject->HaveValidDynamicData())
 		{
 			return;
diff --git a/Engine/Plugins/Animation/MLDeformer/MLDeformerFramework/Source/MLDeformerFramework/Private/DeformerGraphDataInterfaceSkeletonWithQuats.h b/Engine/Plugins/Animation/MLDeformer/MLDeformerFramework/Source/MLDeformerFramework/Private/DeformerGraphDataInterfaceSkeletonWithQuats.h
--- a/Engine/Plugins/Animation/MLDeformer/MLDeformerFramework/Source/MLDeformerFramework/Private/DeformerGraphDataInterfaceSkeletonWithQuats.h
+++ b/Engine/Plugins/Animation/MLDeformer/MLDeformerFramework/Source/MLDeformerFramework/Private/DeformerGraphDataInterfaceSkeletonWithQuats.h
@@ -4,11 +4,13 @@
 
 #include "OptimusComputeDataInterface.h"
 #include "ComputeFramework/ComputeDataProvider.h"
+#include "PrimitiveComponentId.h"
 #include "RenderGraphFwd.h"
 #include "DeformerGraphDataInterfaceSkeletonWithQuats.generated.h"
 
 #define UE_API MLDEFORMERFRAMEWORK_API
 
+class FSceneInterface;
 class FSkeletalMeshObject;
 class FSkeletonWithQuatsDataInterfaceParameters;
 class USkinnedMeshComponent;
@@ -77,7 +79,8 @@ namespace UE::MLDeformer
 		using FParameters = FSkeletonWithQuatsDataInterfaceParameters;
 		void CacheRefToLocalQuats(TArray<FQuat4f>& OutRefToLocalQuats) const;
 
-		FSkeletalMeshObject* SkeletalMeshObject = nullptr;
+		FSceneInterface* SceneInterface = nullptr;
+		FPrimitiveComponentId ComponentId;
 		uint32 BoneRevisionNumber = 0;
 		TArray<TArray<FQuat4f> > PerSectionRefToLocalQuats;
 		TArray<FRDGBufferSRVRef> PerSectionRefToLocalQuatsSRVs;
diff --git a/Engine/Plugins/Animation/MLDeformer/MLDeformerFramework/Source/MLDeformerFramework/Private/MLDeformerGraphDebugDataInterface.cpp b/Engine/Plugins/Animation/MLDeformer/MLDeformerFramework/Source/MLDeformerFramework/Private/MLDeformerGraphDebugDataInterface.cpp
--- a/Engine/Plugins/Animation/MLDeformer/MLDeformerFramework/Source/MLDeformerFramework/Private/MLDeformerGraphDebugDataInterface.cpp
+++ b/Engine/Plugins/Animation/MLDeformer/MLDeformerFramework/Source/MLDeformerFramework/Private/MLDeformerGraphDebugDataInterface.cpp
@@ -2,6 +2,7 @@
 
 #include "MLDeformerGraphDebugDataInterface.h"
 #include "MLDeformerModelInstance.h"
+#include "SkeletalMeshDeformerHelpers.h"
 #include "MLDeformerAsset.h"
 #include "MLDeformerComponent.h"
 #include "MLDeformerVizSettings.h"
@@ -158,7 +159,13 @@ namespace UE::MLDeformer
 			const UMLDeformerVizSettings* VizSettings = Model->GetVizSettings();
 			const UMLDeformerModelInstance* ModelInstance = DeformerComponent->GetModelInstance();
 
-			SkeletalMeshObject = (ModelInstance && ModelInstance->GetSkeletalMeshComponent()) ? ModelInstance->GetSkeletalMeshComponent()->MeshObject : nullptr;
+			USkinnedMeshComponent* SkinnedMeshComp = (ModelInstance) ? ModelInstance->GetSkeletalMeshComponent() : nullptr;
+			if (SkinnedMeshComp && SkinnedMeshComp->IsRegistered())
+			{
+				SceneInterface = SkinnedMeshComp->GetScene();
+				ComponentId = SkinnedMeshComp->GetPrimitiveSceneId();
+			}
+
 			VertexMapBufferSRV = Model->GetVertexMapBuffer().ShaderResourceViewRHI;
 			HeatMapMode = static_cast<int32>(VizSettings->GetHeatMapMode());
 			HeatMapMax = 1.0f / FMath::Max(VizSettings->GetHeatMapMax(), 0.00001f);
@@ -191,7 +198,8 @@ namespace UE::MLDeformer
 			return false;
 		}
 
-		if (SkeletalMeshObject == nullptr || VertexMapBufferSRV == nullptr)
+		FSkeletalMeshObject* SkeletalMeshObject = FSkeletalMeshDeformerHelpers::GetSkeletalMeshObject(SceneInterface, ComponentId);
+		if (!SkeletalMeshObject || VertexMapBufferSRV == nullptr)
 		{
 			return false;
 		}
@@ -215,6 +223,7 @@ namespace UE::MLDeformer
 
 	void FMLDeformerGraphDebugDataProviderProxy::GatherDispatchData(FDispatchData const& InDispatchData)
 	{
+		FSkeletalMeshObject* SkeletalMeshObject = FSkeletalMeshDeformerHelpers::GetSkeletalMeshObject(SceneInterface, ComponentId);
 		const FSkeletalMeshRenderData& SkeletalMeshRenderData = SkeletalMeshObject->GetSkeletalMeshRenderData();
 		const FSkeletalMeshLODRenderData* LodRenderData = SkeletalMeshRenderData.GetPendingFirstLOD(0);
 		const TStridedView<FMLDeformerGraphDebugDataInterfaceParameters> ParameterArray = MakeStridedParameterView<FMLDeformerGraphDebugDataInterfaceParameters>(InDispatchData);
diff --git a/Engine/Plugins/Animation/MLDeformer/MLDeformerFramework/Source/MLDeformerFramework/Public/MLDeformerGraphDebugDataInterface.h b/Engine/Plugins/Animation/MLDeformer/MLDeformerFramework/Source/MLDeformerFramework/Public/MLDeformerGraphDebugDataInterface.h
--- a/Engine/Plugins/Animation/MLDeformer/MLDeformerFramework/Source/MLDeformerFramework/Public/MLDeformerGraphDebugDataInterface.h
+++ b/Engine/Plugins/Animation/MLDeformer/MLDeformerFramework/Source/MLDeformerFramework/Public/MLDeformerGraphDebugDataInterface.h
@@ -6,6 +6,7 @@
 #include "MLDeformerModel.h"
 #include "OptimusComputeDataInterface.h"
 #include "ComputeFramework/ComputeDataProvider.h"
+#include "PrimitiveComponentId.h"
 #include "MLDeformerGraphDebugDataInterface.generated.h"
 
 #define UE_API MLDEFORMERFRAMEWORK_API
@@ -13,6 +14,7 @@
 class FRDGBuffer;
 class FRDGBufferSRV;
 class FRHIShaderResourceView;
+class FSceneInterface;
 class FSkeletalMeshObject;
 class UMLDeformerAsset;
 class USkeletalMeshComponent;
@@ -93,7 +95,8 @@ namespace UE::MLDeformer
 
 	protected:
 		TObjectPtr<UMLDeformerGraphDebugDataProvider> Provider = nullptr;
-		FSkeletalMeshObject* SkeletalMeshObject = nullptr;
+		FSceneInterface* SceneInterface = nullptr;
+		FPrimitiveComponentId ComponentId;
 		TArray<FVector3f> GroundTruthPositions;
 		FRHIShaderResourceView* VertexMapBufferSRV = nullptr;
 		FRDGBuffer* GroundTruthBuffer = nullptr;
diff --git a/Engine/Plugins/Animation/MLDeformer/NearestNeighborModel/Source/NearestNeighborModel/Private/OptimusDataInterfaceSkeletonWithQuats.cpp b/Engine/Plugins/Animation/MLDeformer/NearestNeighborModel/Source/NearestNeighborModel/Private/OptimusDataInterfaceSkeletonWithQuats.cpp
--- a/Engine/Plugins/Animation/MLDeformer/NearestNeighborModel/Source/NearestNeighborModel/Private/OptimusDataInterfaceSkeletonWithQuats.cpp
+++ b/Engine/Plugins/Animation/MLDeformer/NearestNeighborModel/Source/NearestNeighborModel/Private/OptimusDataInterfaceSkeletonWithQuats.cpp
@@ -137,7 +137,11 @@ FComputeDataProviderRenderProxy* UOptimusSkeletonWithQuatsDataProvider::GetRende
 
 FOptimusSkeletonWithQuatsDataProviderProxy::FOptimusSkeletonWithQuatsDataProviderProxy(USkinnedMeshComponent* SkinnedMeshComponent)
 {
-	SkeletalMeshObject = SkinnedMeshComponent != nullptr ? SkinnedMeshComponent->MeshObject : nullptr;
+	if (SkinnedMeshComponent && SkinnedMeshComponent->IsRegistered())
+	{
+		SceneInterface = SkinnedMeshComponent->GetScene();
+		ComponentId = SkinnedMeshComponent->GetPrimitiveSceneId();
+	}
 	BoneRevisionNumber = SkinnedMeshComponent != nullptr ? SkinnedMeshComponent->GetBoneTransformRevisionNumber() : 0;
 }
 
@@ -147,7 +151,9 @@ bool FOptimusSkeletonWithQuatsDataProviderProxy::IsValid(FValidationData const&
 	{
 		return false;
 	}
-	if (SkeletalMeshObject == nullptr)
+
+	FSkeletalMeshObject* SkeletalMeshObject = FSkeletalMeshDeformerHelpers::GetSkeletalMeshObject(SceneInterface, ComponentId);
+	if (!SkeletalMeshObject)
 	{
 		return false;
 	}
@@ -193,6 +199,7 @@ struct FSkeletonWithQuatsDataInterfacePermutationIds
 
 void FOptimusSkeletonWithQuatsDataProviderProxy::GatherPermutations(FPermutationData& InOutPermutationData) const
 {
+	FSkeletalMeshObject* SkeletalMeshObject = FSkeletalMeshDeformerHelpers::GetSkeletalMeshObject(SceneInterface, ComponentId);
 	const int32 LodIndex = SkeletalMeshObject->GetLOD();
 	FSkeletalMeshRenderData const& SkeletalMeshRenderData = SkeletalMeshObject->GetSkeletalMeshRenderData();
 	FSkeletalMeshLODRenderData const* LodRenderData = &SkeletalMeshRenderData.LODRenderData[LodIndex];
@@ -223,6 +230,7 @@ void FOptimusSkeletonWithQuatsDataProviderProxy::GatherPermutations(FPermutation
 
 void FOptimusSkeletonWithQuatsDataProviderProxy::AllocateResources(FRDGBuilder& GraphBuilder, FAllocationData const& InAllocationData)
 {
+	FSkeletalMeshObject* SkeletalMeshObject = FSkeletalMeshDeformerHelpers::GetSkeletalMeshObject(SceneInterface, ComponentId);
 	TArray<FQuat4f> RefToLocalQuats;
 	CacheRefToLocalQuats(RefToLocalQuats);
 
@@ -271,6 +279,7 @@ void FOptimusSkeletonWithQuatsDataProviderProxy::AllocateResources(FRDGBuilder&
 
 void FOptimusSkeletonWithQuatsDataProviderProxy::GatherDispatchData(FDispatchData const& InDispatchData)
 {
+	FSkeletalMeshObject* SkeletalMeshObject = FSkeletalMeshDeformerHelpers::GetSkeletalMeshObject(SceneInterface, ComponentId);
 	const int32 LodIndex = SkeletalMeshObject->GetLOD();
 	FSkeletalMeshRenderData const& SkeletalMeshRenderData = SkeletalMeshObject->GetSkeletalMeshRenderData();
 	FSkeletalMeshLODRenderData const* LodRenderData = &SkeletalMeshRenderData.LODRenderData[LodIndex];
@@ -307,6 +316,8 @@ void FOptimusSkeletonWithQuatsDataProviderProxy::GatherDispatchData(FDispatchDat
 void FOptimusSkeletonWithQuatsDataProviderProxy::CacheRefToLocalQuats(TArray<FQuat4f>& OutRefToLocalQuats) const
 {
 	OutRefToLocalQuats.Reset();
+
+	FSkeletalMeshObject* SkeletalMeshObject = FSkeletalMeshDeformerHelpers::GetSkeletalMeshObject(SceneInterface, ComponentId);
 	if (!SkeletalMeshObject || !SkeletalMeshObject->HaveValidDynamicData())
 	{
 		return;
diff --git a/Engine/Plugins/Animation/MLDeformer/NearestNeighborModel/Source/NearestNeighborModel/Private/OptimusDataInterfaceSkeletonWithQuats.h b/Engine/Plugins/Animation/MLDeformer/NearestNeighborModel/Source/NearestNeighborModel/Private/OptimusDataInterfaceSkeletonWithQuats.h
--- a/Engine/Plugins/Animation/MLDeformer/NearestNeighborModel/Source/NearestNeighborModel/Private/OptimusDataInterfaceSkeletonWithQuats.h
+++ b/Engine/Plugins/Animation/MLDeformer/NearestNeighborModel/Source/NearestNeighborModel/Private/OptimusDataInterfaceSkeletonWithQuats.h
@@ -4,12 +4,14 @@
 
 #include "OptimusComputeDataInterface.h"
 #include "ComputeFramework/ComputeDataProvider.h"
+#include "PrimitiveComponentId.h"
 #include "RenderGraphFwd.h"
 
 #include "OptimusDataInterfaceSkeletonWithQuats.generated.h"
 
 #define UE_API NEARESTNEIGHBORMODEL_API
 
+class FSceneInterface;
 class FSkeletalMeshObject;
 class FSkeletonWithQuatsDataInterfaceParameters;
 class USkinnedMeshComponent;
@@ -73,7 +75,8 @@ private:
 	using FParameters = FSkeletonWithQuatsDataInterfaceParameters;
 	void CacheRefToLocalQuats(TArray<FQuat4f>& OutRefToLocalQuats) const;
 
-	FSkeletalMeshObject* SkeletalMeshObject;
+	FSceneInterface* SceneInterface = nullptr;
+	FPrimitiveComponentId ComponentId;
 	uint32 BoneRevisionNumber = 0;
 	TArray<TArray<FQuat4f> > PerSectionRefToLocalQuats;
 	TArray<FRDGBufferSRVRef> PerSectionRefToLocalQuatsSRVs;
diff --git a/Engine/Plugins/Animation/MLDeformer/VertexDeltaModel/Source/VertexDeltaModel/Private/VertexDeltaGraphDataInterface.cpp b/Engine/Plugins/Animation/MLDeformer/VertexDeltaModel/Source/VertexDeltaModel/Private/VertexDeltaGraphDataInterface.cpp
--- a/Engine/Plugins/Animation/MLDeformer/VertexDeltaModel/Source/VertexDeltaModel/Private/VertexDeltaGraphDataInterface.cpp
+++ b/Engine/Plugins/Animation/MLDeformer/VertexDeltaModel/Source/VertexDeltaModel/Private/VertexDeltaGraphDataInterface.cpp
@@ -8,6 +8,7 @@
 #include "MLDeformerModelInstance.h"
 #include "MLDeformerAsset.h"
 #include "MLDeformerComponent.h"
+#include "SkeletalMeshDeformerHelpers.h"
 
 #include "VertexDeltaModelInstance.h"
 #include "OptimusDataDomain.h"
@@ -114,7 +115,12 @@ namespace UE::VertexDeltaModel
 			VertexDeltaModelInstance = Cast<UVertexDeltaModelInstance>(ModelInstance);
 			if (VertexDeltaModelInstance)
 			{
-				SkeletalMeshObject = ModelInstance->GetSkeletalMeshComponent()->MeshObject;
+				USkinnedMeshComponent* SkinnedMeshComp = ModelInstance->GetSkeletalMeshComponent();
+				if (SkinnedMeshComp && SkinnedMeshComp->IsRegistered())
+				{
+					SceneInterface = SkinnedMeshComp->GetScene();
+					ComponentId = SkinnedMeshComp->GetPrimitiveSceneId();
+				}
 				NNEModelInstanceRDG = VertexDeltaModelInstance->GetNNEModelInstanceRDG();
 
 				if (NNEModelInstanceRDG)
@@ -133,7 +139,8 @@ namespace UE::VertexDeltaModel
 			return false;
 		}
 
-		if (SkeletalMeshObject == nullptr || NNEModelInstanceRDG == nullptr || VertexMapBufferSRV == nullptr)
+		FSkeletalMeshObject* SkeletalMeshObject = FSkeletalMeshDeformerHelpers::GetSkeletalMeshObject(SceneInterface, ComponentId);
+		if (!SkeletalMeshObject || NNEModelInstanceRDG == nullptr || VertexMapBufferSRV == nullptr)
 		{
 			return false;
 		}
@@ -153,6 +160,7 @@ namespace UE::VertexDeltaModel
 
 	void FVertexDeltaGraphDataProviderProxy::GatherDispatchData(FDispatchData const& InDispatchData)
 	{
+		FSkeletalMeshObject* SkeletalMeshObject = FSkeletalMeshDeformerHelpers::GetSkeletalMeshObject(SceneInterface, ComponentId);
 		const FSkeletalMeshRenderData& SkeletalMeshRenderData = SkeletalMeshObject->GetSkeletalMeshRenderData();
 		const FSkeletalMeshLODRenderData* LodRenderData = SkeletalMeshRenderData.GetPendingFirstLOD(0);
 		const TStridedView<FVertexDeltaGraphDataInterfaceParameters> ParameterArray = MakeStridedParameterView<FVertexDeltaGraphDataInterfaceParameters>(InDispatchData);
diff --git a/Engine/Plugins/Animation/MLDeformer/VertexDeltaModel/Source/VertexDeltaModel/Private/VertexDeltaGraphDebugDataInterface.cpp b/Engine/Plugins/Animation/MLDeformer/VertexDeltaModel/Source/VertexDeltaModel/Private/VertexDeltaGraphDebugDataInterface.cpp
--- a/Engine/Plugins/Animation/MLDeformer/VertexDeltaModel/Source/VertexDeltaModel/Private/VertexDeltaGraphDebugDataInterface.cpp
+++ b/Engine/Plugins/Animation/MLDeformer/VertexDeltaModel/Source/VertexDeltaModel/Private/VertexDeltaGraphDebugDataInterface.cpp
@@ -5,6 +5,7 @@
 #include "MLDeformerAsset.h"
 #include "MLDeformerComponent.h"
 #include "MLDeformerVizSettings.h"
+#include "SkeletalMeshDeformerHelpers.h"
 #include "Components/SkeletalMeshComponent.h"
 #include "ComputeFramework/ComputeKernelPermutationSet.h"
 #include "ComputeFramework/ShaderParamTypeDefinition.h"
@@ -158,7 +159,12 @@ namespace UE::VertexDeltaModel
 			const UMLDeformerVizSettings* VizSettings = Model->GetVizSettings();
 			const UMLDeformerModelInstance* ModelInstance = DeformerComponent->GetModelInstance();
 
-			SkeletalMeshObject = ModelInstance->GetSkeletalMeshComponent()->MeshObject;
+			USkinnedMeshComponent* SkinnedMeshComp = ModelInstance->GetSkeletalMeshComponent();
+			if (SkinnedMeshComp && SkinnedMeshComp->IsRegistered())
+			{
+				SceneInterface = SkinnedMeshComp->GetScene();
+				ComponentId = SkinnedMeshComp->GetPrimitiveSceneId();
+			}
 			VertexMapBufferSRV = Model->GetVertexMapBuffer().ShaderResourceViewRHI;
 			HeatMapMode = static_cast<int32>(VizSettings->GetHeatMapMode());
 			HeatMapMax = 1.0f / FMath::Max(VizSettings->GetHeatMapMax(), 0.00001f);
@@ -190,7 +196,9 @@ namespace UE::VertexDeltaModel
 		{
 			return false;
 		}
-		if (SkeletalMeshObject == nullptr || VertexMapBufferSRV == nullptr)
+
+		FSkeletalMeshObject* SkeletalMeshObject = FSkeletalMeshDeformerHelpers::GetSkeletalMeshObject(SceneInterface, ComponentId);
+		if (!SkeletalMeshObject || VertexMapBufferSRV == nullptr)
 		{
 			return false;
 		}
@@ -207,6 +215,7 @@ namespace UE::VertexDeltaModel
 
 	void FVertexDeltaGraphDebugDataProviderProxy::GatherDispatchData(FDispatchData const& InDispatchData)
 	{
+		FSkeletalMeshObject* SkeletalMeshObject = FSkeletalMeshDeformerHelpers::GetSkeletalMeshObject(SceneInterface, ComponentId);
 		const FSkeletalMeshRenderData& SkeletalMeshRenderData = SkeletalMeshObject->GetSkeletalMeshRenderData();
 		const FSkeletalMeshLODRenderData* LodRenderData = SkeletalMeshRenderData.GetPendingFirstLOD(0);
 		const TStridedView<FVertexDeltaGraphDebugDataInterfaceParameters> ParameterArray = MakeStridedParameterView<FVertexDeltaGraphDebugDataInterfaceParameters>(InDispatchData);
diff --git a/Engine/Plugins/Animation/MLDeformer/VertexDeltaModel/Source/VertexDeltaModel/Public/VertexDeltaGraphDataInterface.h b/Engine/Plugins/Animation/MLDeformer/VertexDeltaModel/Source/VertexDeltaModel/Public/VertexDeltaGraphDataInterface.h
--- a/Engine/Plugins/Animation/MLDeformer/VertexDeltaModel/Source/VertexDeltaModel/Public/VertexDeltaGraphDataInterface.h
+++ b/Engine/Plugins/Animation/MLDeformer/VertexDeltaModel/Source/VertexDeltaModel/Public/VertexDeltaGraphDataInterface.h
@@ -5,6 +5,7 @@
 #include "CoreTypes.h"
 #include "OptimusComputeDataInterface.h"
 #include "ComputeFramework/ComputeDataProvider.h"
+#include "PrimitiveComponentId.h"
 #include "Rendering/SkeletalMeshRenderData.h"
 #include "Rendering/SkeletalMeshLODRenderData.h"
 #include "SkeletalRenderPublic.h"
@@ -16,6 +17,7 @@
 class FRDGBuffer;
 class FRDGBufferSRV;
 class FRHIShaderResourceView;
+class FSceneInterface;
 class FSkeletalMeshObject;
 class UMLDeformerComponent;
 class USkeletalMeshComponent;
@@ -84,7 +86,8 @@ namespace UE::VertexDeltaModel
 		// ~END FComputeDataProviderRenderProxy overrides.
 
 	protected:
-		FSkeletalMeshObject* SkeletalMeshObject = nullptr;
+		FSceneInterface* SceneInterface = nullptr;
+		FPrimitiveComponentId ComponentId;
 		UE::NNE::IModelInstanceRDG* NNEModelInstanceRDG = nullptr;
 		UVertexDeltaModelInstance* VertexDeltaModelInstance = nullptr;
 		FRHIShaderResourceView* VertexMapBufferSRV = nullptr;
diff --git a/Engine/Plugins/Animation/MLDeformer/VertexDeltaModel/Source/VertexDeltaModel/Public/VertexDeltaGraphDebugDataInterface.h b/Engine/Plugins/Animation/MLDeformer/VertexDeltaModel/Source/VertexDeltaModel/Public/VertexDeltaGraphDebugDataInterface.h
--- a/Engine/Plugins/Animation/MLDeformer/VertexDeltaModel/Source/VertexDeltaModel/Public/VertexDeltaGraphDebugDataInterface.h
+++ b/Engine/Plugins/Animation/MLDeformer/VertexDeltaModel/Source/VertexDeltaModel/Public/VertexDeltaGraphDebugDataInterface.h
@@ -4,6 +4,7 @@
 
 #include "CoreTypes.h"
 #include "MLDeformerModel.h"
+#include "PrimitiveComponentId.h"
 #include "VertexDeltaGraphDataInterface.h"
 #include "OptimusComputeDataInterface.h"
 #include "ComputeFramework/ComputeDataProvider.h"
@@ -12,6 +13,7 @@
 class FRDGBuffer;
 class FRDGBufferSRV;
 class FRHIShaderResourceView;
+class FSceneInterface;
 class FSkeletalMeshObject;
 class UMLDeformerAsset;
 class USkeletalMeshComponent;
@@ -92,7 +94,8 @@ namespace UE::VertexDeltaModel
 
 	protected:
 		TObjectPtr<UDEPRECATED_VertexDeltaGraphDebugDataProvider> Provider = nullptr;
-		FSkeletalMeshObject* SkeletalMeshObject = nullptr;
+		FSceneInterface* SceneInterface = nullptr;
+		FPrimitiveComponentId ComponentId;
 		TArray<FVector3f> GroundTruthPositions;
 		FRHIShaderResourceView* VertexMapBufferSRV = nullptr;
 		FRDGBuffer* GroundTruthBuffer = nullptr;
diff --git a/Engine/Plugins/Experimental/ChaosFlesh/Source/ChaosFleshEngine/Private/DataInterfaces/DIFleshDeformer.cpp b/Engine/Plugins/Experimental/ChaosFlesh/Source/ChaosFleshEngine/Private/DataInterfaces/DIFleshDeformer.cpp
--- a/Engine/Plugins/Experimental/ChaosFlesh/Source/ChaosFleshEngine/Private/DataInterfaces/DIFleshDeformer.cpp
+++ b/Engine/Plugins/Experimental/ChaosFlesh/Source/ChaosFleshEngine/Private/DataInterfaces/DIFleshDeformer.cpp
@@ -265,11 +265,15 @@ FDIFleshDeformerProviderProxy::FDIFleshDeformerProviderProxy(
 	UDeformableTetrahedralComponent* FleshComponentIn,
 	FFleshDeformerParameters const& FleshDeformerParametersIn)
 	: SkinnedMeshComponent(SkinnedMeshComponentIn)
-	, SkeletalMeshObject(SkinnedMeshComponentIn ? SkinnedMeshComponentIn->MeshObject : nullptr)
 	, FleshComponent(FleshComponentIn)
 	, FleshDeformerParameters(FleshDeformerParametersIn)
 {
-	LodIndex = SkeletalMeshObject ? SkeletalMeshObject->GetLOD() : 0;
+	if (SkinnedMeshComponentIn && SkinnedMeshComponentIn->IsRegistered())
+	{
+		SceneInterface = SkinnedMeshComponentIn->GetScene();
+		ComponentId = SkinnedMeshComponentIn->GetPrimitiveSceneId();
+	}
+
 	if (SkinnedMeshComponent)
 	{
 		if (USkinnedAsset* SkinnedAsset = SkinnedMeshComponent->GetSkinnedAsset())
@@ -287,7 +291,9 @@ FDIFleshDeformerProviderProxy::IsValid(FValidationData const& InValidationData)
 	{
 		return false;
 	}
-	if (SkinnedMeshComponent == nullptr || SkeletalMeshObject == nullptr || FleshComponent == nullptr)
+
+	FSkeletalMeshObject* SkeletalMeshObject = FSkeletalMeshDeformerHelpers::GetSkeletalMeshObject(SceneInterface, ComponentId);
+	if (SkinnedMeshComponent == nullptr || !SkeletalMeshObject || FleshComponent == nullptr)
 	{
 		return false;
 	}
@@ -307,7 +313,7 @@ FDIFleshDeformerProviderProxy::IsValid(FValidationData const& InValidationData)
 	// This function is called every draw call.  It may be better to do this initialization
 	// higher in the stack, but we do it here to support on the fly LOD switching.  I think...
 	FName MeshName(*MeshId, MeshId.Len());
-	if (!FleshComponent->GetGPUBufferManager().InitGPUBindingsBuffer(SkinnedMeshComponent, MeshName, LodIndex))
+	if (!FleshComponent->GetGPUBufferManager().InitGPUBindingsBuffer(SkinnedMeshComponent, MeshName, SkeletalMeshObject->GetLOD()))
 	{
 		return false;
 	}
@@ -347,13 +353,14 @@ FDIFleshDeformerProviderProxy::AllocateResources(FRDGBuilder& GraphBuilder, FAll
 void
 FDIFleshDeformerProviderProxy::GatherDispatchData(FDispatchData const& InDispatchData)
 {
+	FSkeletalMeshObject* SkeletalMeshObject = FSkeletalMeshDeformerHelpers::GetSkeletalMeshObject(SceneInterface, ComponentId);
 	FSkeletalMeshRenderData const& SkeletalMeshRenderData = SkeletalMeshObject->GetSkeletalMeshRenderData();
-	FSkeletalMeshLODRenderData const* LodRenderData = &SkeletalMeshRenderData.LODRenderData[LodIndex];
+	FSkeletalMeshLODRenderData const* LodRenderData = &SkeletalMeshRenderData.LODRenderData[SkeletalMeshObject->GetLOD()];
 	const TStridedView<FParameters> ParameterArray = MakeStridedParameterView<FParameters>(InDispatchData);
 
 	FName MeshName(*MeshId, MeshId.Len());
 	Chaos::Softs::FChaosFleshDeformableGPUManager& GPUManager = FleshComponent->GetGPUBufferManager();
-	Chaos::Softs::FChaosFleshDeformableGPUManager::FBindingsBuffer* Bindings = GPUManager.GetBindingsBuffer(SkinnedMeshComponent, MeshName, LodIndex);
+	Chaos::Softs::FChaosFleshDeformableGPUManager::FBindingsBuffer* Bindings = GPUManager.GetBindingsBuffer(SkinnedMeshComponent, MeshName, SkeletalMeshObject->GetLOD());
 	if (!Bindings)
 	{
 		return;
diff --git a/Engine/Plugins/Experimental/ChaosFlesh/Source/ChaosFleshEngine/Public/DataInterfaces/DIFleshDeformer.h b/Engine/Plugins/Experimental/ChaosFlesh/Source/ChaosFleshEngine/Public/DataInterfaces/DIFleshDeformer.h
--- a/Engine/Plugins/Experimental/ChaosFlesh/Source/ChaosFleshEngine/Public/DataInterfaces/DIFleshDeformer.h
+++ b/Engine/Plugins/Experimental/ChaosFlesh/Source/ChaosFleshEngine/Public/DataInterfaces/DIFleshDeformer.h
@@ -4,12 +4,14 @@
 #include "OptimusComputeDataInterface.h"
 #include "ComputeFramework/ComputeDataProvider.h"
 #include "ChaosFlesh/ChaosDeformableTetrahedralComponent.h"
+#include "PrimitiveComponentId.h"
 
 #include "DIFleshDeformer.generated.h"
 
 class FFleshDeformerDataInterfaceParameters;
 class FRDGBuffer;
 class FRDGBufferSRV;
+class FSceneInterface;
 class FSkeletalMeshObject;
 class USkinnedMeshComponent;
 
@@ -96,13 +98,13 @@ private:
 	using FParameters = FFleshDeformerDataInterfaceParameters;
 
 	USkinnedMeshComponent* SkinnedMeshComponent = nullptr;
-	FSkeletalMeshObject* SkeletalMeshObject = nullptr;
+	FSceneInterface* SceneInterface = nullptr;
+	FPrimitiveComponentId ComponentId;
 	UDeformableTetrahedralComponent* FleshComponent = nullptr;
 
 	FFleshDeformerParameters FleshDeformerParameters;
 
 	int32 TetIndex = INDEX_NONE;
-	int32 LodIndex = INDEX_NONE;
 	FString MeshId;
 
 	uint32 NumVertices = 0;
diff --git a/Engine/Plugins/Runtime/ComputeFramework/Source/ComputeFramework/Private/ComputeGraphWorker.cpp b/Engine/Plugins/Runtime/ComputeFramework/Source/ComputeFramework/Private/ComputeGraphWorker.cpp
--- a/Engine/Plugins/Runtime/ComputeFramework/Source/ComputeFramework/Private/ComputeGraphWorker.cpp
+++ b/Engine/Plugins/Runtime/ComputeFramework/Source/ComputeFramework/Private/ComputeGraphWorker.cpp
@@ -181,10 +181,12 @@ void FComputeGraphTaskWorker::SubmitWork(FComputeContext& Context)
 						FComputeDataProviderRenderProxy::FValidationData ValidationData{ NumSubInvocations, (int32)Member.GetStructMetadata()->GetSize() };
 						bIsValid &= DataProvider->IsValid(ValidationData);
 
-						if (bIsValid)
+						if (!bIsValid)
 						{
-							DataProvider->GatherPermutations(PermutationData);
+							continue;
 						}
+						
+						DataProvider->GatherPermutations(PermutationData);
 					}
 				}
 			}
@@ -301,11 +303,6 @@ void FComputeGraphTaskWorker::SubmitWork(FComputeContext& Context)
 					// 1. Data interfaces sharing the same binding (primary) as the kernel should present its data in a way that
 					// matches the kernel dispatch method, which can be either unified(full buffer) or non-unified (per invocation window into the full buffer)
 					// 2. Data interfaces not sharing the same binding (secondary) should always provide a full view to its data (unified)
-					// Note: In case of non-unified kernel, extra work maybe needed to read from secondary buffers.
-					// When kernel is non-unified, index = 0...section.max for each invocation/section, 
-					// so user may want to consider using a dummy buffer that maps section index to the indices of secondary buffers
-					// for example, given a non-unified kernel, primary and secondary components sharing the same vertex count, we might want to create a buffer
-					// in the primary group that is simply [0,1,2...,NumVerts-1], which we can then index into to map section vert index to the global vert index
 					DispatchData.bUnifiedDispatch = KernelInvocation.BoundProviderIsPrimary[MemberIndex]? bIsUnifiedDispatch : true;
 					DispatchData.ParameterStructSize = Member.GetStructMetadata()->GetSize();
 					DispatchData.ParameterBufferOffset = Member.GetOffset();
diff --git a/Engine/Source/Runtime/Engine/Private/SkeletalMeshDeformerHelpers.cpp b/Engine/Source/Runtime/Engine/Private/SkeletalMeshDeformerHelpers.cpp
--- a/Engine/Source/Runtime/Engine/Private/SkeletalMeshDeformerHelpers.cpp
+++ b/Engine/Source/Runtime/Engine/Private/SkeletalMeshDeformerHelpers.cpp
@@ -2,10 +2,13 @@
 #include "SkeletalMeshDeformerHelpers.h"
 
 #include "DataDrivenShaderPlatformInfo.h"
+#include "PrimitiveSceneInfo.h"
 #include "RenderGraphBuilder.h"
 #include "RenderGraphResources.h"
 #include "RenderGraphUtils.h"
 #include "RenderingThread.h"
+#include "SceneInterface.h"
+#include "SkeletalMeshSceneProxy.h"
 #include "SkeletalRenderGPUSkin.h"
 
 FRHIShaderResourceView* FSkeletalMeshDeformerHelpers::GetBoneBufferForReading(
@@ -31,10 +34,8 @@ FRHIShaderResourceView* FSkeletalMeshDeformerHelpers::GetBoneBufferForReading(
 }
 
 FRHIShaderResourceView* FSkeletalMeshDeformerHelpers::GetMorphTargetBufferForReading(
-	FSkeletalMeshObject const* InMeshObject,
+	const FSkeletalMeshObject* InMeshObject,
 	int32 InLodIndex,
-	int32 InSectionIndex,
-	uint32 InFrameNumber,
 	bool bInPreviousFrame)
 {
 	if (InMeshObject->IsCPUSkinned())
@@ -42,11 +43,13 @@ FRHIShaderResourceView* FSkeletalMeshDeformerHelpers::GetMorphTargetBufferForRea
 		return nullptr;
 	}
 
-	FSkeletalMeshObjectGPUSkin const* MeshObjectGPU = static_cast<FSkeletalMeshObjectGPUSkin const*>(InMeshObject);
-	FGPUBaseSkinVertexFactory const* BaseVertexFactory = MeshObjectGPU->GetBaseSkinVertexFactory(InLodIndex, InSectionIndex);
-	FMorphVertexBuffer const* MorphVertexBuffer = BaseVertexFactory != nullptr ? BaseVertexFactory->GetMorphVertexBuffer(bInPreviousFrame) : nullptr;
-
-	return MorphVertexBuffer != nullptr ? MorphVertexBuffer->GetSRV() : nullptr;
+	const FSkeletalMeshObjectGPUSkin* MeshObjectGPU = static_cast<const FSkeletalMeshObjectGPUSkin *>(InMeshObject);
+	if (!MeshObjectGPU->LODs[InLodIndex].MorphVertexBufferPool->IsInitialized())
+	{
+		return nullptr;
+	}
+	const FMorphVertexBuffer& MorphVertexBuffer = MeshObjectGPU->LODs[InLodIndex].MorphVertexBufferPool->GetMorphVertexBufferForReading(bInPreviousFrame);
+	return MorphVertexBuffer.GetSRV();
 }
 
 FSkeletalMeshDeformerHelpers::FClothBuffers FSkeletalMeshDeformerHelpers::GetClothBuffersForReading(
@@ -155,6 +158,28 @@ int32 FSkeletalMeshDeformerHelpers::GetIndexOfFirstAvailableSection(FSkeletalMes
 	return Result;
 }
 
+FSkeletalMeshObject* FSkeletalMeshDeformerHelpers::GetSkeletalMeshObject(FSceneInterface* InSceneInterface, FPrimitiveComponentId InComponentId)
+{
+	if (!InSceneInterface || !InComponentId.IsValid())
+	{
+		return nullptr;
+	}
+
+	FPrimitiveSceneInfo* PrimitiveSceneInfo = InSceneInterface->GetPrimitiveSceneInfo(InComponentId);
+	if (!PrimitiveSceneInfo || !PrimitiveSceneInfo->Proxy || !PrimitiveSceneInfo->Proxy->IsSkinnedMesh())
+	{
+		return nullptr;
+	}
+
+	FSkeletalMeshObject* MeshObject = static_cast<FSkeletalMeshSceneProxy*>(PrimitiveSceneInfo->Proxy)->GetMeshObject();
+	if (!MeshObject || MeshObject->IsCPUSkinned())
+	{
+		return nullptr;
+	}
+
+	return MeshObject;
+}
+
 FRDGBuffer* FSkeletalMeshDeformerHelpers::AllocateVertexFactoryPositionBuffer(FRDGBuilder& GraphBuilder, FRDGExternalAccessQueue& ExternalAccesQueue, FSkeletalMeshObject* InMeshObject, int32 InLodIndex, TCHAR const* InBufferName)
 {
 	if (InMeshObject->IsCPUSkinned())
@@ -417,6 +442,14 @@ PRAGMA_DISABLE_DEPRECATION_WARNINGS
 PRAGMA_ENABLE_DEPRECATION_WARNINGS
 }
 
+FRHIShaderResourceView* FSkeletalMeshDeformerHelpers::GetMorphTargetBufferForReading(FSkeletalMeshObject const* InMeshObject, int32 InLodIndex,
+	int32 InSectionIndex, uint32 InFrameNumber, bool bInPreviousFrame)
+{
+PRAGMA_DISABLE_DEPRECATION_WARNINGS
+	return GetMorphTargetBufferForReading(InMeshObject, InLodIndex, bInPreviousFrame);
+PRAGMA_ENABLE_DEPRECATION_WARNINGS
+}
+
 void FSkeletalMeshDeformerHelpers::ResetVertexFactoryBufferOverrides(FSkeletalMeshObject* InMeshObject, int32 LODIndex)
 {
 	if (InMeshObject->IsCPUSkinned())
diff --git a/Engine/Source/Runtime/Engine/Public/SkeletalMeshDeformerHelpers.h b/Engine/Source/Runtime/Engine/Public/SkeletalMeshDeformerHelpers.h
--- a/Engine/Source/Runtime/Engine/Public/SkeletalMeshDeformerHelpers.h
+++ b/Engine/Source/Runtime/Engine/Public/SkeletalMeshDeformerHelpers.h
@@ -3,11 +3,13 @@
 #pragma once
 
 #include "CoreMinimal.h"
+#include "PrimitiveComponentId.h"
 #include "RenderGraphResources.h"
 
 class FRDGBuffer;
 class FRDGExternalAccessQueue;
 class FRHIShaderResourceView;
+class FSceneInterface;
 class FSkeletalMeshObject;
 
 /** Functions that expose some internal functionality of FSkeletalMeshObject required by MeshDeformer systems. */
@@ -31,10 +33,8 @@ public:
 
 	/** Get direct access to morph target buffer SRV. */
 	ENGINE_API static FRHIShaderResourceView* GetMorphTargetBufferForReading(
-		FSkeletalMeshObject const* InMeshObject,
+		const FSkeletalMeshObject* InMeshObject,
 		int32 InLodIndex,
-		int32 InSectionIndex,
-		uint32 InFrameNumber,
 		bool bInPreviousFrame);
 
 	/** Buffer SRVs from the cloth system. */
@@ -66,6 +66,12 @@ public:
 	/** Returns the index of the first section that is not disabled */
 	ENGINE_API static int32 GetIndexOfFirstAvailableSection(FSkeletalMeshObject* InMeshObject, int32 InLodIndex);
 	
+	/** 
+	 * Resolves and validates a FSkeletalMeshObject from a scene + component ID.
+	 * Returns nullptr if the primitive is not found, not a skinned mesh, not yet updated, or CPU-skinned. 
+	 */
+	ENGINE_API static FSkeletalMeshObject* GetSkeletalMeshObject(FSceneInterface* InSceneInterface, FPrimitiveComponentId InComponentId);
+	
 #pragma endregion GetInternals
 
 #pragma region SetInternals
@@ -123,5 +129,14 @@ public:
 	UE_DEPRECATED(5.4, "UpdateVertexFactoryBufferOverrides requires a command list.")
 	ENGINE_API static void UpdateVertexFactoryBufferOverrides(FSkeletalMeshObject* InMeshObject, int32 InLodIndex);
 
+	UE_DEPRECATED(5.8, "GetMorphTargetBufferForReading no longer requires section index and frame number.")
+	/** Get direct access to morph target buffer SRV. */
+	ENGINE_API static FRHIShaderResourceView* GetMorphTargetBufferForReading(
+		FSkeletalMeshObject const* InMeshObject,
+		int32 InLodIndex,
+		int32 InSectionIndex,
+		uint32 InFrameNumber,
+		bool bInPreviousFrame);
+
 #pragma endregion Deprecated
 };
