Fix CIS error after bad merge
#rnx
#Synced-CL 52470409

diff --git a/Engine/Plugins/Animation/DeformerGraph/Source/OptimusCore/Public/OptimusKernelSource.h b/Engine/Plugins/Animation/DeformerGraph/Source/OptimusCore/Public/OptimusKernelSource.h
--- a/Engine/Plugins/Animation/DeformerGraph/Source/OptimusCore/Public/OptimusKernelSource.h
+++ b/Engine/Plugins/Animation/DeformerGraph/Source/OptimusCore/Public/OptimusKernelSource.h
@@ -18,7 +18,6 @@ public:
 		Source = InSource;
 	}
 	
-#if WITH_EDITOR
 	//~ Begin UComputeKernelSource Interface.
 	FString GetSource() const override
 	{
