[Deformer Graph] Move Mark Deformed to PostSubmit and GetReadableOutputBuffer to Gather dispatch data such that Writes and Reads respects sorted kernel  dispatch order properly. Previously these functions were called in independent passes when processing graph invocations following the enqueued order so Readable Buffer Tracking did not take into account the sorted kernel dispatch order.

#Synced-CL 52601091

diff --git a/Engine/Plugins/Animation/DeformerGraph/Shaders/Private/DataInterfaceSkinnedMeshRead.ush b/Engine/Plugins/Animation/DeformerGraph/Shaders/Private/DataInterfaceSkinnedMeshRead.ush
--- a/Engine/Plugins/Animation/DeformerGraph/Shaders/Private/DataInterfaceSkinnedMeshRead.ush
+++ b/Engine/Plugins/Animation/DeformerGraph/Shaders/Private/DataInterfaceSkinnedMeshRead.ush
@@ -20,6 +20,7 @@ Buffer<float> {DataInterfaceName}_PositionStaticBuffer;
 Buffer<SNORM float4> {DataInterfaceName}_TangentStaticBuffer;
 Buffer<float4> {DataInterfaceName}_ColorStaticBuffer;
 uint {DataInterfaceName}_ColorIndexMask;
+uint {DataInterfaceName}_ReadableOutputBuffers;
 
 uint ReadNumVertices_{DataInterfaceName}()
 {
@@ -30,49 +31,63 @@ float3 ReadPosition_{DataInterfaceName}(uint VertexIndex)
 {
 	uint BufferIndex = (VertexIndex) * 3;
 
-#if READABLE_OUTPUT_BUFFERS & OPTIMUS_SKINNED_MESH_READ_POSITION
-	return float3(
-		{DataInterfaceName}_PositionBufferSRV[BufferIndex + 0], 
-		{DataInterfaceName}_PositionBufferSRV[BufferIndex + 1],
-		{DataInterfaceName}_PositionBufferSRV[BufferIndex + 2]);
-#else
-	return float3(
-		{DataInterfaceName}_PositionStaticBuffer[BufferIndex + 0], 
-		{DataInterfaceName}_PositionStaticBuffer[BufferIndex + 1],
-		{DataInterfaceName}_PositionStaticBuffer[BufferIndex + 2]);
-#endif 
+	if ({DataInterfaceName}_ReadableOutputBuffers & OPTIMUS_SKINNED_MESH_READ_POSITION)
+	{
+		return float3(
+			{DataInterfaceName}_PositionBufferSRV[BufferIndex + 0], 
+			{DataInterfaceName}_PositionBufferSRV[BufferIndex + 1],
+			{DataInterfaceName}_PositionBufferSRV[BufferIndex + 2]);
+	}
+	else
+	{
+		return float3(
+			{DataInterfaceName}_PositionStaticBuffer[BufferIndex + 0], 
+			{DataInterfaceName}_PositionStaticBuffer[BufferIndex + 1],
+			{DataInterfaceName}_PositionStaticBuffer[BufferIndex + 2]);
+	}
 }
 
 float4 ReadTangentX_{DataInterfaceName}(uint VertexIndex)
 {
 	uint BufferIndex = (VertexIndex) * 2;
 	
-#if READABLE_OUTPUT_BUFFERS & OPTIMUS_SKINNED_MESH_READ_TANGENTS
-	return TangentRead_SkinCache({DataInterfaceName}_TangentBufferSRV[BufferIndex]) ;
-#else	
-	return TangentBias_SkinCache({DataInterfaceName}_TangentStaticBuffer[BufferIndex]) ;
-#endif
+	if ({DataInterfaceName}_ReadableOutputBuffers & OPTIMUS_SKINNED_MESH_READ_TANGENTS)
+	{
+		return TangentRead_SkinCache({DataInterfaceName}_TangentBufferSRV[BufferIndex]) ;
+	}
+	else
+	{
+		return TangentBias_SkinCache({DataInterfaceName}_TangentStaticBuffer[BufferIndex]) ;
+	}
 }
 
 float4 ReadTangentZ_{DataInterfaceName}(uint VertexIndex)
 {
 	uint BufferIndex = (VertexIndex) * 2;
-#if READABLE_OUTPUT_BUFFERS & OPTIMUS_SKINNED_MESH_READ_TANGENTS
-	return TangentRead_SkinCache({DataInterfaceName}_TangentBufferSRV[BufferIndex + 1]) ;
-#else	
-	return TangentBias_SkinCache({DataInterfaceName}_TangentStaticBuffer[BufferIndex + 1]) ;
-#endif
+
+	if ({DataInterfaceName}_ReadableOutputBuffers & OPTIMUS_SKINNED_MESH_READ_TANGENTS)
+	{
+		return TangentRead_SkinCache({DataInterfaceName}_TangentBufferSRV[BufferIndex + 1]) ;
+	}
+	else
+	{
+		return TangentBias_SkinCache({DataInterfaceName}_TangentStaticBuffer[BufferIndex + 1]) ;
+	}
 }
 
 float4 ReadColor_{DataInterfaceName}(uint VertexIndex)
 {
 	// Index Mask is used here in case ColorStaticBuffer is bound to the GWhiteVertexBufferWithSRV that is only a few bytes
 	uint BufferIndex = (VertexIndex) & ({DataInterfaceName}_ColorIndexMask);
-#if READABLE_OUTPUT_BUFFERS & OPTIMUS_SKINNED_MESH_READ_COLOR
-	return {DataInterfaceName}_ColorBufferSRV[BufferIndex] FMANUALFETCH_COLOR_COMPONENT_SWIZZLE;
-#else	
-	return {DataInterfaceName}_ColorStaticBuffer[BufferIndex] FMANUALFETCH_COLOR_COMPONENT_SWIZZLE;
-#endif	
+
+	if ({DataInterfaceName}_ReadableOutputBuffers & OPTIMUS_SKINNED_MESH_READ_COLOR)
+	{
+		return {DataInterfaceName}_ColorBufferSRV[BufferIndex] FMANUALFETCH_COLOR_COMPONENT_SWIZZLE;
+	}
+	else
+	{
+		return {DataInterfaceName}_ColorStaticBuffer[BufferIndex] FMANUALFETCH_COLOR_COMPONENT_SWIZZLE;
+	}
 }
 
 
diff --git a/Engine/Plugins/Animation/DeformerGraph/Source/OptimusCore/Private/DataInterfaces/OptimusDataInterfaceSkinnedMeshRead.cpp b/Engine/Plugins/Animation/DeformerGraph/Source/OptimusCore/Private/DataInterfaces/OptimusDataInterfaceSkinnedMeshRead.cpp
--- a/Engine/Plugins/Animation/DeformerGraph/Source/OptimusCore/Private/DataInterfaces/OptimusDataInterfaceSkinnedMeshRead.cpp
+++ b/Engine/Plugins/Animation/DeformerGraph/Source/OptimusCore/Private/DataInterfaces/OptimusDataInterfaceSkinnedMeshRead.cpp
@@ -17,11 +17,9 @@
 #include "OptimusSkinnedMeshTrackingExtension.h"
 #include "Animation/MeshDeformerInstance.h"
 #include "ComputeFramework/ComputeKernelPermutationSet.h"
-#include "ComputeFramework/ComputeKernelPermutationVector.h"
 
-#include UE_INLINE_GENERATED_CPP_BY_NAME(OptimusDataInterfaceSkinnedMeshRead)
 
-const TCHAR* UOptimusSkinnedMeshReadDataInterface::ReadableOutputBufferPermutationName = TEXT("READABLE_OUTPUT_BUFFERS");
+#include UE_INLINE_GENERATED_CPP_BY_NAME(OptimusDataInterfaceSkinnedMeshRead)
 
 FString UOptimusSkinnedMeshReadDataInterface::GetDisplayName() const
 {
@@ -96,6 +94,7 @@ BEGIN_SHADER_PARAMETER_STRUCT(FSkinnedMeshReadDataInterfaceParameters, )
 	SHADER_PARAMETER_SRV(Buffer<SNORM float4>, TangentStaticBuffer)
 	SHADER_PARAMETER_SRV(Buffer<float4>, ColorStaticBuffer)
 	SHADER_PARAMETER(uint32, ColorIndexMask)
+	SHADER_PARAMETER(uint32, ReadableOutputBuffers)
 END_SHADER_PARAMETER_STRUCT()
 
 void UOptimusSkinnedMeshReadDataInterface::GetShaderParameters(TCHAR const* UID, FShaderParametersMetadataBuilder& InOutBuilder, FShaderParametersMetadataAllocations& InOutAllocations) const
@@ -127,35 +126,19 @@ void UOptimusSkinnedMeshReadDataInterface::GetHLSL(FString& OutHLSL, FString con
 	OutHLSL += FString::Format(*TemplateFile, TemplateArgs);
 }
 
-enum class EOptimusSkinnedMeshReadReadableOutputBuffer: uint32
-{
-	None = 0,
-	Position = 1<<0,
-	Tangents = 1<<1,
-	VertexColor = 1<<2,
-	NumPermutations = 1 << 3
-};
-ENUM_CLASS_FLAGS(EOptimusSkinnedMeshReadReadableOutputBuffer);
-
 void UOptimusSkinnedMeshReadDataInterface::GetDefines(FComputeKernelDefinitionSet& OutDefinitionSet) const
 {
 	OutDefinitionSet.Defines.Add(FComputeKernelDefinition(
 		TEXT("OPTIMUS_SKINNED_MESH_READ_POSITION"),
-		FString::FromInt(static_cast<uint32>(EOptimusSkinnedMeshReadReadableOutputBuffer::Position))));
+		FString::FromInt(static_cast<uint32>(EMeshDeformerOutputBuffer::SkinnedMeshPosition))));
 	OutDefinitionSet.Defines.Add(FComputeKernelDefinition(
 		TEXT("OPTIMUS_SKINNED_MESH_READ_TANGENTS"),
-		FString::FromInt(static_cast<uint32>(EOptimusSkinnedMeshReadReadableOutputBuffer::Tangents))));
+		FString::FromInt(static_cast<uint32>(EMeshDeformerOutputBuffer::SkinnedMeshTangents))));
 	OutDefinitionSet.Defines.Add(FComputeKernelDefinition(
 		TEXT("OPTIMUS_SKINNED_MESH_READ_COLOR"),
-		FString::FromInt(static_cast<uint32>(EOptimusSkinnedMeshReadReadableOutputBuffer::VertexColor))));
+		FString::FromInt(static_cast<uint32>(EMeshDeformerOutputBuffer::SkinnedMeshVertexColor))));
 }
 
-void UOptimusSkinnedMeshReadDataInterface::GetPermutations(FComputeKernelPermutationVector& OutPermutationVector) const
-{
-	OutPermutationVector.AddPermutation(ReadableOutputBufferPermutationName, static_cast<uint32>(EOptimusSkinnedMeshReadReadableOutputBuffer::NumPermutations));
-}
-
-
 UComputeDataProvider* UOptimusSkinnedMeshReadDataInterface::CreateDataProvider(TObjectPtr<UObject> InBinding, uint64 InInputMask, uint64 InOutputMask) const
 {
 	UOptimusSkinnedMeshReadDataProvider* Provider = NewObject<UOptimusSkinnedMeshReadDataProvider>();
@@ -230,78 +213,29 @@ void FOptimusSkinnedMeshReadDataProviderProxy::AllocateResources(FRDGBuilder& Gr
 	// Color SRV is either created by default or not available at all
 	ColorsSRV = LodRenderData->StaticVertexBuffers.ColorVertexBuffer.GetColorComponentsSRV();
 
-	// Query which deformer output buffers are readable for this component
-	const EMeshDeformerOutputBuffer ReadableBuffers = FOptimusSkinnedMeshTrackingExtension::GetReadableDeformedBuffers(
-		SceneInterface, ComponentId, ReaderExecutionGroup, bIsPrimaryComponent);
+	// Try to get deformed buffers. If they exist (allocated by a write proxy), create SRV from them.
+	// If not, fall back to a white buffer. The shader will use ReadableOutputBuffers
+	// (set in GatherDispatchData) to decide which buffer to actually read from at runtime.
+	FRDGBuffer* White = GraphBuilder.RegisterExternalBuffer(GWhiteVertexBufferWithRDG->Buffer);
 
-	if ((InputMask & static_cast<uint64>(ESkinnedMeshReadDataInterfaceInputSelectorMask::Position)) &&
-		EnumHasAnyFlags(ReadableBuffers, EMeshDeformerOutputBuffer::SkinnedMeshPosition))
-	{
-		FRDGBuffer* PositionBuffer = FSkeletalMeshDeformerHelpers::GetAllocatedPositionBuffer(GraphBuilder, SkeletalMeshObject, LodIndex);
-		PositionBufferSRV = GraphBuilder.CreateSRV(PositionBuffer, PF_R32_FLOAT); 
-	}
-	else
-	{
-		FRDGBuffer* White = GraphBuilder.RegisterExternalBuffer(GWhiteVertexBufferWithRDG->Buffer);	
-		PositionBufferSRV = GraphBuilder.CreateSRV(White, PF_R32_FLOAT); 
-	}	
+	FRDGBuffer* DeformedPositionBuffer = (InputMask & static_cast<uint64>(ESkinnedMeshReadDataInterfaceInputSelectorMask::Position))
+		? FSkeletalMeshDeformerHelpers::GetAllocatedPositionBuffer(GraphBuilder, SkeletalMeshObject, LodIndex)
+		: nullptr;
+	PositionBufferSRV = GraphBuilder.CreateSRV(DeformedPositionBuffer ? DeformedPositionBuffer : White, PF_R32_FLOAT);
 
 	// OpenGL ES does not support writing to RGBA16_SNORM images, instead pack data into SINT in the shader
 	const EPixelFormat TangentsFormat = IsOpenGLPlatform(GMaxRHIShaderPlatform) ? PF_R16G16B16A16_SINT : PF_R16G16B16A16_SNORM;
 
-	if (((InputMask & static_cast<uint64>(ESkinnedMeshReadDataInterfaceInputSelectorMask::TangentX)) ||
-		(InputMask & static_cast<uint64>(ESkinnedMeshReadDataInterfaceInputSelectorMask::TangentZ))) &&
-		EnumHasAnyFlags(ReadableBuffers, EMeshDeformerOutputBuffer::SkinnedMeshTangents))
-	{
-		FRDGBuffer* TangentBuffer = FSkeletalMeshDeformerHelpers::GetAllocatedTangentBuffer(GraphBuilder, SkeletalMeshObject, LodIndex);
-		TangentBufferSRV = GraphBuilder.CreateSRV(TangentBuffer, TangentsFormat); 
-	}
-	else
-	{
-		FRDGBuffer* White = GraphBuilder.RegisterExternalBuffer(GWhiteVertexBufferWithRDG->Buffer);	
-		TangentBufferSRV = GraphBuilder.CreateSRV(White, TangentsFormat); 
-	}
+	FRDGBuffer* DeformedTangentBuffer = ((InputMask & static_cast<uint64>(ESkinnedMeshReadDataInterfaceInputSelectorMask::TangentX)) ||
+		(InputMask & static_cast<uint64>(ESkinnedMeshReadDataInterfaceInputSelectorMask::TangentZ)))
+		? FSkeletalMeshDeformerHelpers::GetAllocatedTangentBuffer(GraphBuilder, SkeletalMeshObject, LodIndex)
+		: nullptr;
+	TangentBufferSRV = GraphBuilder.CreateSRV(DeformedTangentBuffer ? DeformedTangentBuffer : White, TangentsFormat);
 
-	if ((InputMask & static_cast<uint64>(ESkinnedMeshReadDataInterfaceInputSelectorMask::Color)) &&
-		EnumHasAnyFlags(ReadableBuffers, EMeshDeformerOutputBuffer::SkinnedMeshVertexColor))
-	{
-		FRDGBuffer* ColorBuffer = FSkeletalMeshDeformerHelpers::GetAllocatedColorBuffer(GraphBuilder, SkeletalMeshObject, LodIndex);
-		ColorBufferSRV = GraphBuilder.CreateSRV(ColorBuffer, PF_R8G8B8A8); 
-	}
-	else
-	{
-		FRDGBuffer* White = GraphBuilder.RegisterExternalBuffer(GWhiteVertexBufferWithRDG->Buffer);		
-		ColorBufferSRV = GraphBuilder.CreateSRV(White, TangentsFormat); 
-	}
-}
-
-void FOptimusSkinnedMeshReadDataProviderProxy::GatherPermutations(FPermutationData& InOutPermutationData) const
-{
-	// Query the tracking extension for readable deformed buffers on this component
-	const EMeshDeformerOutputBuffer ReadableBuffers = FOptimusSkinnedMeshTrackingExtension::GetReadableDeformedBuffers(
-		SceneInterface, ComponentId, ReaderExecutionGroup, bIsPrimaryComponent);
-
-	EOptimusSkinnedMeshReadReadableOutputBuffer ReadableOutputBuffers = EOptimusSkinnedMeshReadReadableOutputBuffer::None; 
-	if (EnumHasAnyFlags(ReadableBuffers, EMeshDeformerOutputBuffer::SkinnedMeshPosition))
-	{
-		ReadableOutputBuffers |= EOptimusSkinnedMeshReadReadableOutputBuffer::Position;
-	}
-	if (EnumHasAnyFlags(ReadableBuffers, EMeshDeformerOutputBuffer::SkinnedMeshTangents))
-	{
-		ReadableOutputBuffers |= EOptimusSkinnedMeshReadReadableOutputBuffer::Tangents;
-	}
-	if (EnumHasAnyFlags(ReadableBuffers, EMeshDeformerOutputBuffer::SkinnedMeshVertexColor))
-	{
-		ReadableOutputBuffers |= EOptimusSkinnedMeshReadReadableOutputBuffer::VertexColor;
-	}
-	static FString Name(UOptimusSkinnedMeshReadDataInterface::ReadableOutputBufferPermutationName);
-	static uint32 Hash = GetTypeHash(Name);
-	const uint32 ReadableOutputBuffersBits = InOutPermutationData.PermutationVector.GetPermutationBits(Name, Hash, static_cast<uint32>(ReadableOutputBuffers));
-	
-	for (int32 InvocationIndex = 0; InvocationIndex < InOutPermutationData.NumInvocations; ++InvocationIndex)
-	{
-		InOutPermutationData.PermutationIds[InvocationIndex] |= ReadableOutputBuffersBits;
-	}
+	FRDGBuffer* DeformedColorBuffer = (InputMask & static_cast<uint64>(ESkinnedMeshReadDataInterfaceInputSelectorMask::Color))
+		? FSkeletalMeshDeformerHelpers::GetAllocatedColorBuffer(GraphBuilder, SkeletalMeshObject, LodIndex)
+		: nullptr;
+	ColorBufferSRV = GraphBuilder.CreateSRV(DeformedColorBuffer ? DeformedColorBuffer : White, PF_R8G8B8A8);
 }
 
 void FOptimusSkinnedMeshReadDataProviderProxy::GatherDispatchData(FDispatchData const& InDispatchData)
@@ -312,6 +246,12 @@ void FOptimusSkinnedMeshReadDataProviderProxy::GatherDispatchData(FDispatchData
 	FSkeletalMeshRenderData const& SkeletalMeshRenderData = SkeletalMeshObject->GetSkeletalMeshRenderData();
 	FSkeletalMeshLODRenderData const* LodRenderData = &SkeletalMeshRenderData.LODRenderData[LodIndex];
 
+	// Query tracking extension for which deformed buffers are readable at this point in dispatch order.
+	// SkinnedMeshWrite's PostSubmit calls MarkDeformed after each kernel dispatch in sorted order,
+	// so this sees only outputs from kernels that have already dispatched.
+	const EMeshDeformerOutputBuffer ReadableBuffers = FOptimusSkinnedMeshTrackingExtension::GetReadableDeformedBuffers(
+		SceneInterface, ComponentId, ReaderExecutionGroup, bIsPrimaryComponent);
+
 	FRHIShaderResourceView* NullSRVBinding = GWhiteVertexBufferWithSRV->ShaderResourceViewRHI.GetReference();
 
 	const TStridedView<FParameters> ParameterArray = MakeStridedParameterView<FParameters>(InDispatchData);
@@ -339,5 +279,7 @@ void FOptimusSkinnedMeshReadDataProviderProxy::GatherDispatchData(FDispatchData
 		// Ideally we should be getting this value from the GPUBaseSkinVertexFactory but the need for
 		// section index make it tricky when we are doing unified dispatch
 		Parameters.ColorIndexMask = MeshColorBufferSRV != nullptr ? ~0u : 0;
+
+		Parameters.ReadableOutputBuffers = static_cast<uint32>(ReadableBuffers);;
 	}
 }
diff --git a/Engine/Plugins/Animation/DeformerGraph/Source/OptimusCore/Private/DataInterfaces/OptimusDataInterfaceSkinnedMeshRead.h b/Engine/Plugins/Animation/DeformerGraph/Source/OptimusCore/Private/DataInterfaces/OptimusDataInterfaceSkinnedMeshRead.h
--- a/Engine/Plugins/Animation/DeformerGraph/Source/OptimusCore/Private/DataInterfaces/OptimusDataInterfaceSkinnedMeshRead.h
+++ b/Engine/Plugins/Animation/DeformerGraph/Source/OptimusCore/Private/DataInterfaces/OptimusDataInterfaceSkinnedMeshRead.h
@@ -28,7 +28,6 @@ class UOptimusSkinnedMeshReadDataInterface :
 	GENERATED_BODY()
 
 public:
-	static UE_API const TCHAR* ReadableOutputBufferPermutationName;
 	//~ Begin UOptimusComputeDataInterface Interface
 	UE_API FString GetDisplayName() const override;
 	UE_API FName GetCategory() const override;
@@ -45,7 +44,6 @@ public:
 	UE_API void GetShaderHash(FString& InOutKey) const override;
 	UE_API void GetHLSL(FString& OutHLSL, FString const& InDataInterfaceName) const override;
 	UE_API void GetDefines(FComputeKernelDefinitionSet& OutDefinitionSet) const override;
-	UE_API void GetPermutations(FComputeKernelPermutationVector& OutPermutationVector) const override;
 	UE_API UComputeDataProvider* CreateDataProvider(TObjectPtr<UObject> InBinding, uint64 InInputMask, uint64 InOutputMask) const override;
 	//~ End UComputeDataInterface Interface
 
@@ -86,7 +84,6 @@ public:
 	//~ Begin FComputeDataProviderRenderProxy Interface
 	bool IsValid(FValidationData const& InValidationData) const override;
 	void AllocateResources(FRDGBuilder& GraphBuilder, FAllocationData const& InAllocationData) override;
-	void GatherPermutations(FPermutationData& InOutPermutationData) const override;
 	void GatherDispatchData(FDispatchData const& InDispatchData) override;
 	//~ End FComputeDataProviderRenderProxy Interface
 
diff --git a/Engine/Plugins/Animation/DeformerGraph/Source/OptimusCore/Private/DataInterfaces/OptimusDataInterfaceSkinnedMeshWrite.cpp b/Engine/Plugins/Animation/DeformerGraph/Source/OptimusCore/Private/DataInterfaces/OptimusDataInterfaceSkinnedMeshWrite.cpp
--- a/Engine/Plugins/Animation/DeformerGraph/Source/OptimusCore/Private/DataInterfaces/OptimusDataInterfaceSkinnedMeshWrite.cpp
+++ b/Engine/Plugins/Animation/DeformerGraph/Source/OptimusCore/Private/DataInterfaces/OptimusDataInterfaceSkinnedMeshWrite.cpp
@@ -4,6 +4,7 @@
 
 #include "Components/SkinnedMeshComponent.h"
 #include "ComputeFramework/ShaderParamTypeDefinition.h"
+#include "ComputeWorkerInterface.h"
 #include "OptimusSkinnedMeshTrackingExtension.h"
 #include "OptimusDataDomain.h"
 #include "RenderGraphBuilder.h"
@@ -248,13 +249,11 @@ void FOptimusSkinnedMeshWriteDataProviderProxy::AllocateResources(FRDGBuilder& G
 
 	// Allocate required buffers
 	const int32 LodIndex = SkeletalMeshObject->GetLOD();
-	EMeshDeformerOutputBuffer DeformerOutputBuffers = EMeshDeformerOutputBuffer::None;
 
 	if (OutputMask & static_cast<uint64>(ESkinnedMeshWriteDataInterfaceOutputSelectorMask::Position))
 	{
 		PositionBuffer = FSkeletalMeshDeformerHelpers::AllocateVertexFactoryPositionBuffer(GraphBuilder, InAllocationData.ExternalAccessQueue, SkeletalMeshObject, LodIndex, TEXT("OptimusSkinnedMeshPosition"));
 		PositionBufferUAV = GraphBuilder.CreateUAV(PositionBuffer, PF_R32_FLOAT, ERDGUnorderedAccessViewFlags::SkipBarrier);
-		DeformerOutputBuffers |= EMeshDeformerOutputBuffer::SkinnedMeshPosition;
 	}
 	else
 	{
@@ -277,7 +276,6 @@ void FOptimusSkinnedMeshWriteDataProviderProxy::AllocateResources(FRDGBuilder& G
 	{
 		TangentBuffer = FSkeletalMeshDeformerHelpers::AllocateVertexFactoryTangentBuffer(GraphBuilder, InAllocationData.ExternalAccessQueue, SkeletalMeshObject, LodIndex, TEXT("OptimusSkinnedMeshTangent"));
 		TangentBufferUAV = GraphBuilder.CreateUAV(TangentBuffer, TangentsFormat, ERDGUnorderedAccessViewFlags::SkipBarrier);
-		DeformerOutputBuffers |= EMeshDeformerOutputBuffer::SkinnedMeshTangents;
 	}
 	else
 	{
@@ -297,7 +295,6 @@ void FOptimusSkinnedMeshWriteDataProviderProxy::AllocateResources(FRDGBuilder& G
 		ColorBuffer = FSkeletalMeshDeformerHelpers::AllocateVertexFactoryColorBuffer(GraphBuilder, InAllocationData.ExternalAccessQueue, SkeletalMeshObject, LodIndex, TEXT("OptimusSkinnedMeshColor"));
 		// using RGBA here and do a manual fetch swizzle in shader instead of BGRA directly because some Mac does not support it. See GMetalBufferFormats[PF_B8G8R8A8] 
 		ColorBufferUAV = GraphBuilder.CreateUAV(ColorBuffer, PF_R8G8B8A8, ERDGUnorderedAccessViewFlags::SkipBarrier);
-		DeformerOutputBuffers |= EMeshDeformerOutputBuffer::SkinnedMeshVertexColor;
 	}
 	else
 	{
@@ -312,7 +309,33 @@ void FOptimusSkinnedMeshWriteDataProviderProxy::AllocateResources(FRDGBuilder& G
 		ColorBufferUAV = GraphBuilder.CreateUAV(GraphBuilder.RegisterExternalBuffer(GWhiteVertexBufferWithRDG->Buffer), PF_A32B32G32R32F);
 	}
 
-	FOptimusSkinnedMeshTrackingExtension::MarkDeformed(GraphBuilder, SceneInterface, ComponentId, SkeletalMeshObject, DeformerOutputBuffers, ExecutionGroup);
+}
+
+void FOptimusSkinnedMeshWriteDataProviderProxy::PostSubmit(FComputeContext& Context) const
+{
+	FSkeletalMeshObject* SkeletalMeshObject = FSkeletalMeshDeformerHelpers::GetSkeletalMeshObject(SceneInterface, ComponentId);
+	if (!SkeletalMeshObject)
+	{
+		return;
+	}
+
+	EMeshDeformerOutputBuffer DeformerOutputBuffers = EMeshDeformerOutputBuffer::None;
+
+	if (OutputMask & static_cast<uint64>(ESkinnedMeshWriteDataInterfaceOutputSelectorMask::Position))
+	{
+		DeformerOutputBuffers |= EMeshDeformerOutputBuffer::SkinnedMeshPosition;
+	}
+	if ((OutputMask & static_cast<uint64>(ESkinnedMeshWriteDataInterfaceOutputSelectorMask::TangentX)) ||
+		(OutputMask & static_cast<uint64>(ESkinnedMeshWriteDataInterfaceOutputSelectorMask::TangentZ)))
+	{
+		DeformerOutputBuffers |= EMeshDeformerOutputBuffer::SkinnedMeshTangents;
+	}
+	if (OutputMask & static_cast<uint64>(ESkinnedMeshWriteDataInterfaceOutputSelectorMask::Color))
+	{
+		DeformerOutputBuffers |= EMeshDeformerOutputBuffer::SkinnedMeshVertexColor;
+	}
+
+	FOptimusSkinnedMeshTrackingExtension::MarkDeformed(Context.GraphBuilder, SceneInterface, ComponentId, SkeletalMeshObject, DeformerOutputBuffers, ExecutionGroup);
 }
 
 void FOptimusSkinnedMeshWriteDataProviderProxy::GatherDispatchData(FDispatchData const& InDispatchData)
diff --git a/Engine/Plugins/Animation/DeformerGraph/Source/OptimusCore/Private/DataInterfaces/OptimusDataInterfaceSkinnedMeshWrite.h b/Engine/Plugins/Animation/DeformerGraph/Source/OptimusCore/Private/DataInterfaces/OptimusDataInterfaceSkinnedMeshWrite.h
--- a/Engine/Plugins/Animation/DeformerGraph/Source/OptimusCore/Private/DataInterfaces/OptimusDataInterfaceSkinnedMeshWrite.h
+++ b/Engine/Plugins/Animation/DeformerGraph/Source/OptimusCore/Private/DataInterfaces/OptimusDataInterfaceSkinnedMeshWrite.h
@@ -53,6 +53,7 @@ public:
 	UE_API void GetHLSL(FString& OutHLSL, FString const& InDataInterfaceName) const override;
 	UE_API UComputeDataProvider* CreateDataProvider(TObjectPtr<UObject> InBinding, uint64 InInputMask, uint64 InOutputMask) const override;
 	bool GetRequiresReadback() const override { return true; }
+	bool GetRequiresPostSubmitCall() const override { return true; }
 	//~ End UComputeDataInterface Interface
 
 	//~ Begin IOptimusOutputBufferWriter Interface
@@ -111,6 +112,7 @@ public:
 	bool IsValid(FValidationData const& InValidationData) const override;
 	void AllocateResources(FRDGBuilder& GraphBuilder, FAllocationData const& InAllocationData) override;
 	void GatherDispatchData(FDispatchData const& InDispatchData) override;
+	void PostSubmit(FComputeContext& Context) const override;
 	void GetReadbackData(TArray<FReadbackData>& OutReadbackData) const override;
 	//~ End FComputeDataProviderRenderProxy Interface
 
