Deformer Graph: ability to look for secondary bindings in parent actors if not found in the components's owner

#rnx

#Synced-CL 52777971

diff --git a/Engine/Plugins/Animation/DeformerGraph/Source/OptimusCore/Private/OptimusDeformerInstance.cpp b/Engine/Plugins/Animation/DeformerGraph/Source/OptimusCore/Private/OptimusDeformerInstance.cpp
--- a/Engine/Plugins/Animation/DeformerGraph/Source/OptimusCore/Private/OptimusDeformerInstance.cpp
+++ b/Engine/Plugins/Animation/DeformerGraph/Source/OptimusCore/Private/OptimusDeformerInstance.cpp
@@ -234,7 +234,60 @@ void UOptimusDeformerInstanceSettings::GetComponentBindings(
 		UActorComponent* Component = ActorComponent.Get();
 		ExistingBindings.Add(Binding.ProviderName, Component);
 	}
-	
+
+	auto GetMatchingComponent = [InPrimaryComponent](const AActor* Actor, const UOptimusComponentSourceBinding* Binding) -> UActorComponent*
+	{
+		if (Actor && Binding)
+		{
+			const TSubclassOf<UActorComponent>& BindingSourceClass = Binding->GetComponentSource()->GetComponentClass();
+			if (!Binding->ComponentTags.IsEmpty())
+			{
+				TSet<UActorComponent*> TaggedComponents;
+				for (FName Tag : Binding->ComponentTags)
+				{
+					TArray<UActorComponent*> Components = Actor->GetComponentsByTag(BindingSourceClass, Tag);
+					for (UActorComponent* Component : Components)
+					{
+						TaggedComponents.Add(Component);
+					}
+				}
+
+				if (!TaggedComponents.IsEmpty())
+				{
+					TArray<UActorComponent*> RankedTaggedComponents = TaggedComponents.Array();
+
+					// Rank the components by the number of tags they match.
+					RankedTaggedComponents.Sort(
+						[Tags=TSet<FName>(Binding->ComponentTags)](const UActorComponent& InCompA, const UActorComponent& InCompB)
+					{
+						TSet<FName> TagsA(InCompA.ComponentTags);
+						TSet<FName> TagsB(InCompB.ComponentTags);
+
+						return Tags.Intersect(TagsA).Num() < Tags.Intersect(TagsB).Num();
+					});
+
+					if (!RankedTaggedComponents.IsEmpty())
+					{
+						return RankedTaggedComponents[0];
+					}
+				}
+			}
+
+			// Otherwise just use class matching on components owned by the actor.
+			TArray<UActorComponent*> Components;
+			Actor->GetComponents(BindingSourceClass, Components);
+			const int32 Index = Components.IndexOfByPredicate([InPrimaryComponent](const UActorComponent* Component)
+			{
+				return Component != InPrimaryComponent;
+			});
+			if (Index != INDEX_NONE)
+			{
+				return Components[Index];
+			}
+		}
+		return nullptr;
+	};
+
 	// Iterate component bindings and try to find a match.
 	TSet<UActorComponent*> ComponentsUsed;
 	const TArray<UOptimusComponentSourceBinding*>& ComponentBindings = InDeformer->GetComponentBindings();
@@ -265,44 +318,19 @@ void UOptimusDeformerInstanceSettings::GetComponentBindings(
 			}
 			
 			// If not, try to find a component owned by this actor that matches the tag and class.
-			if (!BoundComponent && Actor != nullptr && !Binding->ComponentTags.IsEmpty())
+			if (!BoundComponent && Actor != nullptr)
 			{
-				TSet<UActorComponent*> TaggedComponents;
-				for (FName Tag: Binding->ComponentTags)
+				BoundComponent = GetMatchingComponent(Actor, Binding);
+				if (!BoundComponent)
 				{
-					TArray<UActorComponent*> Components = Actor->GetComponentsByTag(Binding->GetComponentSource()->GetComponentClass(), Tag);
-
-					for (UActorComponent* Component: Components)
+					// If not, try to find a component owned by one of its parent actor that matches the tag and class
+					AActor* ParentActor = InPrimaryComponent->GetAttachParentActor();
+					while (ParentActor && !BoundComponent)
 					{
-						TaggedComponents.Add(Component);
+						BoundComponent = GetMatchingComponent(ParentActor, Binding);
+						ParentActor = ParentActor->GetAttachParentActor();
 					}
 				}
-				TArray<UActorComponent*> RankedTaggedComponents = TaggedComponents.Array();
-
-				// Rank the components by the number of tags they match.
-				RankedTaggedComponents.Sort([Tags=TSet<FName>(Binding->ComponentTags)](const UActorComponent& InCompA, const UActorComponent& InCompB)
-				{
-					TSet<FName> TagsA(InCompA.ComponentTags);
-					TSet<FName> TagsB(InCompB.ComponentTags);
-					
-					return Tags.Intersect(TagsA).Num() < Tags.Intersect(TagsB).Num();
-				});
-
-				if (!RankedTaggedComponents.IsEmpty())
-				{
-					BoundComponent = RankedTaggedComponents[0];
-				}
-			}
-
-			// Otherwise just use class matching on components owned by the actor.
-			if (!BoundComponent && Actor != nullptr)
-			{
-				TArray<UActorComponent*> Components;
-				Actor->GetComponents(Binding->GetComponentSource()->GetComponentClass(), Components);
-				if (!Components.IsEmpty())
-				{
-					BoundComponent = Components[0];
-				}
 			}
 		}
 
